#ttp 330821 - EDITOR: Camera preview viewport stays visible when replacing the camera actor if the viewport was pinned
#proj UE4
#branch UE4
#summary The actor preview now hooks into the actor deleted callback to remove stale pinned previews

diff --git a/Engine/Source/Editor/LevelEditor/Private/SLevelViewport.cpp b/Engine/Source/Editor/LevelEditor/Private/SLevelViewport.cpp
--- a/Engine/Source/Editor/LevelEditor/Private/SLevelViewport.cpp
+++ b/Engine/Source/Editor/LevelEditor/Private/SLevelViewport.cpp
@@ -141,6 +141,8 @@ void SLevelViewport::Construct(const FArguments& InArgs)
 
 	// Tell the level editor we want to be notified when selection changes
 	LevelEditor.OnMapChanged().AddSP( this, &SLevelViewport::OnMapChanged );
+
+	GEngine->OnLevelActorDeleted().AddSP( this, &SLevelViewport::OnLevelActorsRemoved );
 }
 
 void SLevelViewport::ConstructViewportOverlayContent()
@@ -800,6 +802,20 @@ void SLevelViewport::OnMapChanged( UWorld* World, EMapChangeType::Type MapChange
 	}
 }
 
+void SLevelViewport::OnLevelActorsRemoved(AActor* InActor)
+{
+	// Kill any existing actor previews that have expired
+	for( int32 PreviewIndex = 0; PreviewIndex < ActorPreviews.Num(); ++PreviewIndex )
+	{
+		AActor* ExistingActor = ActorPreviews[PreviewIndex].Actor.Get();
+		if ( !ExistingActor || ExistingActor == InActor )
+		{
+			// decrement index so we don't miss next preview after deleting
+			RemoveActorPreview( PreviewIndex-- );
+		}
+	}
+}
+
 void FLevelViewportDropContextMenuImpl::FillDropAddReplaceActorMenu( bool bReplace, FMenuBuilder& MenuBuilder )
 {
 	// Builds a submenu for the Drag Drop context menu used to replace all actors in the current editor selection with a different asset
diff --git a/Engine/Source/Editor/LevelEditor/Public/SLevelViewport.h b/Engine/Source/Editor/LevelEditor/Public/SLevelViewport.h
--- a/Engine/Source/Editor/LevelEditor/Public/SLevelViewport.h
+++ b/Engine/Source/Editor/LevelEditor/Public/SLevelViewport.h
@@ -540,6 +540,9 @@ private:
 	 */
 	void OnMapChanged( UWorld* World, EMapChangeType::Type MapChangeType );
 
+	/** Called in response to an actor being deleted in the level */
+	void OnLevelActorsRemoved(AActor* InActor);
+
 	/** Gets whether the locked icon should be shown in the viewport because it is locked to an actor */
 	EVisibility GetLockedIconVisibility() const;
 
