#ttp 330988 - EDITOR: 'Paste Here' should maybe be removed from scene outliner edit sub menu since it doesn't function as intended there
#branch UE4
#proj Editor
#summary Removed position based menu actions from the scene outliner context menu
#change FLevelViewportContextMenu has been renamed to FLevelEditorContextMenu as it is not actually specific to viewports (it has been used in the scene outliner for some time now)
#change Added an enum to specify a broad context for the level editor context menu (whether it should show position based menu actions or not)

diff --git a/Engine/Source/Editor/LevelEditor/Private/LevelEditor.cpp b/Engine/Source/Editor/LevelEditor/Private/LevelEditor.cpp
--- a/Engine/Source/Editor/LevelEditor/Private/LevelEditor.cpp
+++ b/Engine/Source/Editor/LevelEditor/Private/LevelEditor.cpp
@@ -7,7 +7,7 @@
 #include "LevelEditorActions.h"
 #include "LevelEditorModesActions.h"
 #include "AssetSelection.h"
-#include "LevelViewportContextMenu.h"
+#include "LevelEditorContextMenu.h"
 #include "SLevelViewport.h"
 #include "Editor/WorkspaceMenuStructure/Public/WorkspaceMenuStructureModule.h"
 #include "Developer/MessageLog/Public/MessageLogModule.h"
diff --git a/Engine/Source/Editor/LevelEditor/Private/LevelEditorContextMenu.cpp b/Engine/Source/Editor/LevelEditor/Private/LevelEditorContextMenu.cpp
new file mode 100644
--- /dev/null
+++ b/Engine/Source/Editor/LevelEditor/Private/LevelEditorContextMenu.cpp
@@ -0,0 +1,988 @@
+// Copyright 1998-2014 Epic Games, Inc. All Rights Reserved.
+
+
+#include "LevelEditor.h"
+#include "BlueprintUtilities.h"
+#include "Editor/Kismet/Public/BlueprintEditorModule.h"
+#include "Editor/UnrealEd/Public/Kismet2/KismetEditorUtilities.h"
+#include "AssetSelection.h"
+#include "LevelEditorContextMenu.h"
+#include "LevelEditorActions.h"
+#include "ScopedTransaction.h"
+#include "Toolkits/AssetEditorManager.h"
+#include "SLevelEditor.h"
+#include "SLevelViewport.h"
+#include "Editor/SceneOutliner/Public/SceneOutlinerModule.h"
+#include "DelegateFilter.h"
+#include "AssetData.h"
+#include "DebuggerCommands.h"
+#include "AssetThumbnail.h"
+#include "ClassIconFinder.h"
+#include "IPlacementModeModule.h"
+#include "AssetRegistryModule.h"
+#include "EngineDecalClasses.h"
+#include "EngineUtils.h"
+#include "EditorViewportCommands.h"
+#include "SoundDefinitions.h"
+#include "GlobalEditorCommonCommands.h"
+#include "LevelEditorCreateActorMenu.h"
+#include "SourceCodeNavigation.h"
+#include "Developer/MeshUtilities/Public/MeshUtilities.h"
+
+#define LOCTEXT_NAMESPACE "LevelViewportContextMenu"
+
+DEFINE_LOG_CATEGORY_STATIC(LogViewportMenu, Log, All);
+
+class FLevelEditorContextMenuImpl
+{
+public:
+	static FSelectedActorInfo SelectionInfo;
+public:
+	/**
+	 * Fills in menu options for the select actor menu
+	 *
+	 * @param MenuBuilder	The menu to add items to
+	 */
+	static void FillSelectActorMenu( class FMenuBuilder& MenuBuilder );
+
+	/**
+	 * Fills in menu options for the actor visibility menu
+	 *
+	 * @param MenuBuilder	The menu to add items to
+	 */
+	static void FillActorVisibilityMenu( class FMenuBuilder& MenuBuilder );
+
+	/**
+	 * Fills in menu options for the actor level menu
+	 *
+	 * @param SharedLevel			The level shared between all selected actors.  If any actors are in a different level, this is NULL
+	 * @param bAllInCurrentLevel	true if all selected actors are in the current level
+	 * @param MenuBuilder			The menu to add items to
+	 */
+	static void FillActorLevelMenu( class FMenuBuilder& MenuBuilder );
+
+	/**
+	 * Fills in menu options for the transform menu
+	 *
+	 * @param MenuBuilder	The menu to add items to
+	 */
+	static void FillTransformMenu( class FMenuBuilder& MenuBuilder );
+
+	/**
+	 * Fills in menu options for the Fill Actor menu
+	 *
+	 * @param MenuBuilder	The menu to add items to
+	 */
+	static void FillActorMenu( class FMenuBuilder& MenuBuilder );
+
+	/**
+	 * Fills in menu options for the snap menu
+	 *
+	 * @param MenuBuilder	The menu to add items to
+	 */
+	static void FillSnapAlignMenu( class FMenuBuilder& MenuBuilder );
+
+	/**
+	 * Fills in menu options for the pivot menu
+	 *
+	 * @param MenuBuilder	The menu to add items to
+	 */
+	static void FillPivotMenu( class FMenuBuilder& MenuBuilder );
+	
+	/**
+	 * Fills in menu options for the group menu
+	 *
+	 * @param MenuBuilder	The menu to add items to
+	 */
+	static void FillGroupMenu( class FMenuBuilder& MenuBuilder );
+
+	/**
+	 * Fills in menu options for the edit menu
+	 *
+	 * @param MenuBuilder	The menu to add items to
+	 * @param ContextType	The context for this editor menu
+	 */
+	static void FillEditMenu( class FMenuBuilder& MenuBuilder, LevelEditorMenuContext ContextType );
+
+	/**
+	 * Fills in menu options for the actors merging
+	 *
+	 * @param MenuBuilder	The menu to add items to
+	 */
+	static void FillMergeActorsMenu( class FMenuBuilder& MenuBuilder );
+
+private:
+	/**
+	 * Fills in menu options for the matinee selection menu
+	 *
+	 * @param MenuBuilder	The menu to add items to
+	 */
+	static void FillMatineeSelectActorMenu( class FMenuBuilder& MenuBuilder );
+};
+
+FSelectedActorInfo FLevelEditorContextMenuImpl::SelectionInfo;
+
+struct FLevelScriptEventMenuHelper
+{
+	/*
+	* Fills in menu options for events that can be associated with that actors's blueprint in the level script blueprint
+	*
+	* @param MenuBuilder	The menu to add items to
+	*/
+	static void FillLevelBlueprintEventsMenu(class FMenuBuilder& MenuBuilder, const TArray<AActor*>& SelectedActors);
+};
+
+// NOTE: We intentionally receive a WEAK pointer here because we want to be callable by a delegate whose
+//       payload contains a weak reference to a level editor instance
+TSharedPtr< SWidget > FLevelEditorContextMenu::BuildMenuWidget( TWeakPtr< SLevelEditor > LevelEditor, LevelEditorMenuContext ContextType, TSharedPtr<FExtender> Extender )
+{
+	// Build up the menu
+	const bool bShouldCloseWindowAfterMenuSelection = true;
+	FMenuBuilder MenuBuilder(bShouldCloseWindowAfterMenuSelection, TSharedPtr<const FUICommandList>());
+
+	FillMenu(MenuBuilder, LevelEditor, ContextType, Extender);
+	
+	return MenuBuilder.MakeWidget();
+}
+
+void FLevelEditorContextMenu::FillMenu( FMenuBuilder& MenuBuilder, TWeakPtr<SLevelEditor> LevelEditor, LevelEditorMenuContext ContextType, TSharedPtr<FExtender> Extender )
+{
+	// Generate information about our selection
+	TArray<AActor*> SelectedActors;
+	GEditor->GetSelectedActors()->GetSelectedObjects<AActor>( SelectedActors );
+
+	FSelectedActorInfo& SelectionInfo = FLevelEditorContextMenuImpl::SelectionInfo;
+	SelectionInfo = AssetSelectionUtils::BuildSelectedActorInfo( SelectedActors );
+
+	// Get all menu extenders for this context menu from the level editor module
+	FLevelEditorModule& LevelEditorModule = FModuleManager::GetModuleChecked<FLevelEditorModule>( TEXT("LevelEditor") );
+	TArray<FLevelEditorModule::FLevelViewportMenuExtender_SelectedActors> MenuExtenderDelegates = LevelEditorModule.GetAllLevelViewportContextMenuExtenders();
+
+	TArray<TSharedPtr<FExtender>> Extenders;
+	if(Extender.IsValid())
+	{
+		Extenders.Add(Extender);
+	}
+
+	auto LevelEditorActions = LevelEditor.Pin()->GetLevelEditorActions().ToSharedRef();
+	for (int32 i = 0; i < MenuExtenderDelegates.Num(); ++i)
+	{
+		if (MenuExtenderDelegates[i].IsBound())
+		{
+			Extenders.Add(MenuExtenderDelegates[i].Execute(LevelEditorActions, SelectedActors));
+		}
+	}
+
+	MenuBuilder.PushCommandList(LevelEditorActions);
+	MenuBuilder.PushExtender(FExtender::Combine(Extenders).ToSharedRef());
+
+	TArray<TWeakObjectPtr<UObject>> LabelObjects;
+	for ( FSelectionIterator SelItor(*GEditor->GetSelectedActors()) ; SelItor ; ++SelItor )
+	{
+		LabelObjects.Add(*SelItor);
+	}
+
+	// Check if current selection has any assets that can be browsed to
+	TArray< UObject* > ReferencedAssets;
+	GEditor->GetReferencedAssetsForEditorSelection( ReferencedAssets );
+
+	const bool bCanSyncToContentBrowser = GEditor->CanSyncToContentBrowser();
+
+	if( bCanSyncToContentBrowser || ReferencedAssets.Num() > 0 )		
+	{
+		MenuBuilder.BeginSection("ActorAsset", LOCTEXT("AssetHeading", "Asset") );
+		{
+			if( bCanSyncToContentBrowser )
+			{
+				MenuBuilder.AddMenuEntry( FGlobalEditorCommonCommands::Get().FindInContentBrowser );
+			}
+
+			if( ReferencedAssets.Num() == 1 )
+			{
+				auto Asset = ReferencedAssets[0];
+
+				MenuBuilder.AddMenuEntry( 
+					FLevelEditorCommands::Get().EditAsset,
+					NAME_None,
+					FText::Format( LOCTEXT("EditAssociatedAsset", "Edit {0}"), FText::FromString( Asset->GetName() ) ),
+					TAttribute<FText>(),
+					FSlateIcon( FEditorStyle::GetStyleSetName(), FClassIconFinder::FindIconNameForClass( Asset->GetClass() ) )
+					);
+			}
+			else if ( ReferencedAssets.Num() > 1 )
+			{
+				MenuBuilder.AddMenuEntry( 
+					FLevelEditorCommands::Get().EditAssetNoConfirmMultiple,
+					NAME_None,
+					LOCTEXT("EditAssociatedAssetsMultiple", "Edit Multiple Assets"),
+					TAttribute<FText>(),
+					FSlateIcon( FEditorStyle::GetStyleSetName(), "ClassIcon.Default" )
+					);
+
+			}
+
+		}
+		MenuBuilder.EndSection();
+	}
+
+
+	MenuBuilder.BeginSection( "ActorControl", LOCTEXT("ActorHeading", "Actor") );
+	{
+		MenuBuilder.AddMenuEntry( FEditorViewportCommands::Get().FocusViewportToSelection );
+		MenuBuilder.AddMenuEntry( FLevelEditorCommands::Get().SnapCameraToActor );
+	}
+	MenuBuilder.EndSection();
+
+	// Go to C++ Code
+	if( SelectionInfo.SelectionClass != NULL && FSourceCodeNavigation::IsCompilerAvailable() )
+	{
+		FString ClassHeaderPath;
+		if( FSourceCodeNavigation::FindClassHeaderPath( SelectionInfo.SelectionClass, ClassHeaderPath ) && IFileManager::Get().FileSize( *ClassHeaderPath ) != INDEX_NONE )
+		{
+			const FString CodeFileName = FPaths::GetCleanFilename( *ClassHeaderPath );
+
+			MenuBuilder.BeginSection( "ActorCode", LOCTEXT("ActorCodeHeading", "C++") );
+			{
+				MenuBuilder.AddMenuEntry( FLevelEditorCommands::Get().GoToCodeForActor,
+					NAME_None, 
+					FText::Format( LOCTEXT("GoToCodeForActor", "Open {0}"), FText::FromString( CodeFileName ) ),
+					FText::Format( LOCTEXT("GoToCodeForActor_ToolTip", "Opens the header file for this actor ({0}) in a code editing program"), FText::FromString( CodeFileName ) ) );
+			}
+			MenuBuilder.EndSection();
+		}
+	}
+
+	MenuBuilder.BeginSection("ActorSelectVisibilityLevels");
+	{
+		// Add a sub-menu for "Select"
+		MenuBuilder.AddSubMenu( 
+			LOCTEXT("SelectSubMenu", "Select"),
+			LOCTEXT("SelectSubMenu_ToolTip", "Opens the actor selection menu"),
+			FNewMenuDelegate::CreateStatic( &FLevelEditorContextMenuImpl::FillSelectActorMenu ) );
+
+		MenuBuilder.AddSubMenu( 
+			LOCTEXT("EditSubMenu", "Edit"),
+			FText::GetEmpty(),
+			FNewMenuDelegate::CreateStatic( &FLevelEditorContextMenuImpl::FillEditMenu, ContextType ) );
+
+		MenuBuilder.AddSubMenu( 
+			LOCTEXT("VisibilitySubMenu", "Visibility"),
+			LOCTEXT("VisibilitySubMenu_ToolTip", "Selected actor visibility options"),
+			FNewMenuDelegate::CreateStatic( &FLevelEditorContextMenuImpl::FillActorVisibilityMenu ) );
+
+		// Build the menu for grouping actors
+		BuildGroupMenu( MenuBuilder, SelectionInfo );
+
+		MenuBuilder.AddSubMenu( 
+			LOCTEXT("LevelSubMenu", "Level"),
+			LOCTEXT("LevelSubMenu_ToolTip", "Options for interacting with this actor's level"),
+			FNewMenuDelegate::CreateStatic( &FLevelEditorContextMenuImpl::FillActorLevelMenu ) );
+	}
+	MenuBuilder.EndSection();
+
+	if (ContextType == LevelEditorMenuContext::Viewport)
+	{
+		LevelEditorCreateActorMenu::FillAddReplaceViewportContextMenuSections( MenuBuilder );
+	}
+
+	if( GEditor->PlayWorld != NULL )
+	{
+		if( SelectionInfo.NumSelected > 0 )
+		{
+			MenuBuilder.BeginSection( "Simulation", NSLOCTEXT( "LevelViewportContextMenu", "SimulationHeading", "Simulation" ) );
+			{
+				MenuBuilder.AddMenuEntry( FLevelEditorCommands::Get().KeepSimulationChanges );
+			}
+			MenuBuilder.EndSection();
+		}
+	}
+
+	MenuBuilder.BeginSection("LevelViewportAttach");
+	{
+		// Only display the attach menu if we have actors selected
+		if ( GEditor->GetSelectedActorCount() )
+		{
+			if(SelectionInfo.bHaveAttachedActor)
+			{
+				MenuBuilder.AddMenuEntry( FLevelEditorCommands::Get().DetachFromParent );
+			}
+
+			MenuBuilder.AddSubMenu( 
+				LOCTEXT( "ActorAttachToSubMenu", "Attach To" ), 
+				LOCTEXT( "ActorAttachToSubMenu_ToolTip", "Attach Actor as child" ),
+				FNewMenuDelegate::CreateStatic( &FLevelEditorContextMenuImpl::FillActorMenu ) );
+		}
+
+		// Add a heading for "Movement" if an actor is selected
+		if ( GEditor->GetSelectedActorIterator() )
+		{
+			// Add a sub-menu for "Transform"
+			MenuBuilder.AddSubMenu( 
+				LOCTEXT("TransformSubMenu", "Transform"), 
+				LOCTEXT("TransformSubMenu_ToolTip", "Actor transform utils"),
+				FNewMenuDelegate::CreateStatic( &FLevelEditorContextMenuImpl::FillTransformMenu ) );
+		}
+
+		// @todo UE4: The current pivot options only work for brushes
+		if( SelectionInfo.bHaveBrush )
+		{
+			// Add a sub-menu for "Pivot"
+			MenuBuilder.AddSubMenu( 
+				LOCTEXT("PivotSubMenu", "Pivot"), 
+				LOCTEXT("PivotSubMenu_ToolTip", "Actor pivoting utils"),
+				FNewMenuDelegate::CreateStatic( &FLevelEditorContextMenuImpl::FillPivotMenu ) );
+		}
+
+		if (GetDefault<UEditorExperimentalSettings>()->bActorMerging && 
+			(SelectionInfo.bHaveStaticMeshComponent || SelectionInfo.bHaveLandscape))
+		{
+			MenuBuilder.AddSubMenu( 
+				LOCTEXT("MergeActorsSubMenu", "Merge"), 
+				LOCTEXT("MergeActorsSubMenu_ToolTip", "Actor merging utils"),
+				FNewMenuDelegate::CreateStatic( &FLevelEditorContextMenuImpl::FillMergeActorsMenu ) );
+		}
+	}
+	MenuBuilder.EndSection();
+
+	FLevelScriptEventMenuHelper::FillLevelBlueprintEventsMenu(MenuBuilder, SelectedActors);
+
+	MenuBuilder.PopCommandList();
+	MenuBuilder.PopExtender();
+}
+
+
+void FLevelEditorContextMenu::SummonMenu( const TSharedRef< SLevelEditor >& LevelEditor, LevelEditorMenuContext ContextType )
+{
+	struct Local
+	{
+		static void ExtendMenu( FMenuBuilder& MenuBuilder )
+		{
+			// one extra entry when summoning the menu this way
+			MenuBuilder.BeginSection("ActorPreview", LOCTEXT("PreviewHeading", "Preview") );
+			{
+				// Note: not using a command for play from here since it requires a mouse click
+				FUIAction PlayFromHereAction( 
+					FExecuteAction::CreateStatic( &FPlayWorldCommandCallbacks::StartPlayFromHere ) );
+
+				const FText PlayFromHereLabel = GEditor->OnlyLoadEditorVisibleLevelsInPIE() ? LOCTEXT("PlayFromHereVisible", "Play From Here (visible levels)") : LOCTEXT("PlayFromHere", "Play From Here");
+				MenuBuilder.AddMenuEntry( PlayFromHereLabel, LOCTEXT("PlayFromHere_ToolTip", "Starts a game preview from the clicked location"),FSlateIcon(), PlayFromHereAction );
+			}
+			MenuBuilder.EndSection();
+		}
+	};
+	
+	TSharedRef<FExtender> Extender = MakeShareable(new FExtender);
+	Extender->AddMenuExtension("LevelViewportAttach", EExtensionHook::After, TSharedPtr< FUICommandList >(), FMenuExtensionDelegate::CreateStatic(&Local::ExtendMenu));
+
+	// Create the context menu!
+	TSharedPtr<SWidget> MenuWidget = BuildMenuWidget( LevelEditor, ContextType, Extender );
+	if ( MenuWidget.IsValid() )
+	{
+		// @todo: Should actually use the location from a click event instead!
+		const FVector2D MouseCursorLocation = FSlateApplication::Get().GetCursorPos();
+	
+		TWeakPtr< SWindow > ContextMenuWindow = FSlateApplication::Get().PushMenu(
+			LevelEditor->GetActiveViewport().ToSharedRef(), MenuWidget.ToSharedRef(), MouseCursorLocation, FPopupTransitionEffect( FPopupTransitionEffect::ContextMenu ) );
+	}
+}
+
+FSlateColor InvertOnHover( const TWeakPtr< SWidget > WidgetPtr )
+{
+	TSharedPtr< SWidget > Widget = WidgetPtr.Pin();
+	if ( Widget.IsValid() && Widget->IsHovered() )
+	{
+		return FEditorStyle::GetSlateColor( "InvertedForeground" );
+	}
+
+	return FSlateColor::UseForeground();
+}
+
+void FLevelEditorContextMenu::BuildGroupMenu( FMenuBuilder& MenuBuilder, const FSelectedActorInfo& SelectedActorInfo )
+{
+	if( GEditor->bGroupingActive )
+	{
+		// Whether or not we added a grouping sub-menu
+		bool bNeedGroupSubMenu = SelectedActorInfo.bHaveSelectedLockedGroup || SelectedActorInfo.bHaveSelectedUnlockedGroup;
+
+		// Grouping based on selection (must have selected at least two actors)
+		if( SelectedActorInfo.NumSelected > 1 )
+		{
+			if( !SelectedActorInfo.bHaveSelectedLockedGroup && !SelectedActorInfo.bHaveSelectedUnlockedGroup )
+			{
+				// Only one menu entry needed so dont use a sub-menu
+				MenuBuilder.AddMenuEntry( FLevelEditorCommands::Get().RegroupActors, NAME_None, FLevelEditorCommands::Get().GroupActors->GetLabel(), FLevelEditorCommands::Get().GroupActors->GetDescription() );
+			}
+			else
+			{
+				// Put everything into a sub-menu
+				bNeedGroupSubMenu = true;
+			}
+		}
+		
+		if( bNeedGroupSubMenu )
+		{
+			MenuBuilder.AddSubMenu( 
+				LOCTEXT("GroupMenu", "Groups"),
+				LOCTEXT("GroupMenu_ToolTip", "Opens the actor grouping menu"),
+				FNewMenuDelegate::CreateStatic( &FLevelEditorContextMenuImpl::FillGroupMenu ) );
+		}
+	}
+}
+
+void FLevelEditorContextMenuImpl::FillSelectActorMenu( FMenuBuilder& MenuBuilder )
+{
+	FText SelectAllActorStr = FText::Format( LOCTEXT("SelectActorsOfSameClass", "Select All {0}(s)"), FText::FromString( SelectionInfo.SelectionStr ) );
+	int32 NumSelectedSurfaces = NumSelectedSurfaces = AssetSelectionUtils::GetNumSelectedSurfaces( SelectionInfo.SharedWorld );
+
+	MenuBuilder.BeginSection("SelectActorGeneral", LOCTEXT("SelectAnyHeading", "General") );
+	{
+		MenuBuilder.AddMenuEntry( FGenericCommands::Get().SelectAll, NAME_None, TAttribute<FText>(), LOCTEXT("SelectAll_ToolTip", "Selects all actors") );
+		MenuBuilder.AddMenuEntry( FLevelEditorCommands::Get().SelectNone );
+		MenuBuilder.AddMenuEntry( FLevelEditorCommands::Get().InvertSelection );
+	}
+	MenuBuilder.EndSection();
+
+	if( !SelectionInfo.bHaveBrush && SelectionInfo.bAllSelectedActorsOfSameType && SelectionInfo.SelectionStr.Len() != 0 )
+	{
+		// These menu options appear if only if all the actors are the same type and we aren't selecting brush
+		MenuBuilder.AddMenuEntry( FLevelEditorCommands::Get().SelectAllActorsOfSameClass, NAME_None, SelectAllActorStr );
+	}
+
+	// Add brush commands when we have a brush or any surfaces selected
+	MenuBuilder.BeginSection("SelectBrush", LOCTEXT("SelectBrushHeading", "Brushes") );
+	{
+		if( SelectionInfo.bHaveBrush || NumSelectedSurfaces > 0 )
+		{
+			if( SelectionInfo.bAllSelectedAreBrushes )
+			{
+				MenuBuilder.AddMenuEntry( FLevelEditorCommands::Get().SelectAllActorsOfSameClass, NAME_None, SelectAllActorStr );
+			}
+		}
+
+		MenuBuilder.AddMenuEntry( FLevelEditorCommands::Get().SelectAllAddditiveBrushes );
+		MenuBuilder.AddMenuEntry( FLevelEditorCommands::Get().SelectAllSubtractiveBrushes );
+		MenuBuilder.AddMenuEntry( FLevelEditorCommands::Get().SelectAllSemiSolidBrushes );
+		MenuBuilder.AddMenuEntry( FLevelEditorCommands::Get().SelectAllNonSolidBrushes );
+	}
+	MenuBuilder.EndSection();
+
+	if( SelectionInfo.NumSelected > 0 || NumSelectedSurfaces > 0 )
+	{
+		// If any actors are selected add lights selection options
+		MenuBuilder.BeginSection("SelectLights", LOCTEXT("SelectLightHeading", "Lights") );
+		{
+			MenuBuilder.AddMenuEntry( FLevelEditorCommands::Get().SelectRelevantLights );
+
+			if ( SelectionInfo.bHaveLight )
+			{
+				MenuBuilder.AddMenuEntry( FLevelEditorCommands::Get().SelectAllLights );
+				MenuBuilder.AddMenuEntry( FLevelEditorCommands::Get().SelectStationaryLightsExceedingOverlap );
+			}
+		}
+		MenuBuilder.EndSection();
+
+		if( SelectionInfo.bHaveStaticMesh )
+		{
+			// if any static meshes are selected allow selecting actors using the same mesh
+			MenuBuilder.BeginSection("SelectMeshes", LOCTEXT("SelectStaticMeshHeading", "Static Meshes") );
+			{
+				MenuBuilder.AddMenuEntry( FLevelEditorCommands::Get().SelectStaticMeshesOfSameClass, NAME_None, LOCTEXT("SelectStaticMeshesOfSameClass_Menu", "Select Matching (Selected Classes)") );
+				MenuBuilder.AddMenuEntry( FLevelEditorCommands::Get().SelectStaticMeshesAllClasses, NAME_None, LOCTEXT("SelectStaticMeshesAllClasses_Menu", "Select Matching (All Classes)") );
+			}
+			MenuBuilder.EndSection();
+		}
+
+		if( SelectionInfo.bHavePawn || SelectionInfo.bHaveSkeletalMesh )
+		{
+			// if any skeletal meshes are selected allow selecting actors using the same mesh
+			MenuBuilder.BeginSection("SelectSkeletalMeshes", LOCTEXT("SelectSkeletalMeshHeading", "Skeletal Meshes") );
+			{
+				MenuBuilder.AddMenuEntry( FLevelEditorCommands::Get().SelectSkeletalMeshesOfSameClass );
+				MenuBuilder.AddMenuEntry( FLevelEditorCommands::Get().SelectSkeletalMeshesAllClasses );
+			}
+			MenuBuilder.EndSection();
+		}
+
+		if( SelectionInfo.bHaveEmitter )
+		{
+			MenuBuilder.BeginSection("SelectEmitters", LOCTEXT("SelectEmitterHeading", "Emitters") );
+			{
+				MenuBuilder.AddMenuEntry( FLevelEditorCommands::Get().SelectMatchingEmitter );
+			}
+			MenuBuilder.EndSection();
+		}
+	}
+
+	if( SelectionInfo.bHaveBrush || SelectionInfo.NumSelected > 0 )
+	{
+		MenuBuilder.BeginSection("SelectMaterial", LOCTEXT("SelectMaterialHeading", "Materials") );
+		{
+			MenuBuilder.AddMenuEntry( FLevelEditorCommands::Get().SelectAllWithSameMaterial );
+		}
+		MenuBuilder.EndSection();
+	}
+
+	// Allow users to select all surfaces in the level in a single click
+	MenuBuilder.BeginSection("SelectSurfaces", LOCTEXT("SelectAllSurfaces", "Surfaces") );
+	{
+		MenuBuilder.AddMenuEntry( FLevelEditorCommands::Get().SelectAllSurfaces );
+	}
+	MenuBuilder.EndSection();
+
+	// build matinee related selection menu
+	FillMatineeSelectActorMenu( MenuBuilder );
+}
+
+void FLevelEditorContextMenuImpl::FillMatineeSelectActorMenu( FMenuBuilder& MenuBuilder )
+{
+	MenuBuilder.BeginSection("SelectMatinee", LOCTEXT("SelectMatineeHeading", "Matinee") );
+	{
+		// show list of Matinee Actors that controls this actor
+		// this is ugly but we don't have good way of knowing which Matinee actor controls me
+		// in the future this can be cached to TMap somewhere and use that list
+		// for now we show only when 1 actor is selected
+		if ( SelectionInfo.SharedLevel && SelectionInfo.NumSelected == 1 )
+		{
+			TArray<AMatineeActor*> MatineeActors;	
+			// first collect all matinee actors
+			for ( AActor* Actor : SelectionInfo.SharedLevel->Actors )
+			{
+				AMatineeActor * CurActor = Cast<AMatineeActor>(Actor);
+				if ( CurActor )
+				{
+					MatineeActors.Add(CurActor);
+				}
+			}
+
+			if ( MatineeActors.Num() > 0 )
+			{
+				FSelectionIterator ActorIter( GEditor->GetSelectedActorIterator() );
+				AActor * SelectedActor = Cast<AActor>(*ActorIter);
+
+				// now delete the matinee actors that don't control currently selected actor
+				for (int32 MatineeActorIter=0; MatineeActorIter<MatineeActors.Num(); ++MatineeActorIter)
+				{
+					AMatineeActor * CurMatineeActor = MatineeActors[MatineeActorIter];
+					TArray<AActor *> CutMatineeControlledActors;
+					CurMatineeActor->GetControlledActors(CutMatineeControlledActors);
+					bool bIsMatineeControlled=false;
+					for ( AActor* ControlledActor : CutMatineeControlledActors )
+					{
+						if (ControlledActor == SelectedActor)
+						{
+							bIsMatineeControlled = true;
+						}
+					}
+
+					// if not, remove it
+					if (!bIsMatineeControlled)
+					{
+						MatineeActors.RemoveAt(MatineeActorIter);
+						--MatineeActorIter;
+					}
+				}
+
+				// if some matinee controls this, add to menu for direct selection
+				if ( MatineeActors.Num() > 0 )
+				{
+					for (int32 MatineeActorIter=0; MatineeActorIter<MatineeActors.Num(); ++MatineeActorIter)
+					{
+						AMatineeActor * CurMatineeActor = MatineeActors[MatineeActorIter];
+						const FText Text = FText::Format( LOCTEXT("SelectMatineeActor", "Select {0}"), FText::FromString( CurMatineeActor->GetName() ) );
+
+						FUIAction CurMatineeActorAction( FExecuteAction::CreateStatic( &FLevelEditorActionCallbacks::OnSelectMatineeActor, CurMatineeActor ) );
+						MenuBuilder.AddMenuEntry( Text, Text, FSlateIcon(), CurMatineeActorAction );
+
+						// if matinee is opened, and if that is CurMatineeActor, show option to go to group
+						if( GEditorModeTools().IsModeActive( FBuiltinEditorModes::EM_InterpEdit ) )
+						{
+							const FEdModeInterpEdit* InterpEditMode = (const FEdModeInterpEdit*)GEditorModeTools().GetActiveMode( FBuiltinEditorModes::EM_InterpEdit );
+
+							if ( InterpEditMode && InterpEditMode->MatineeActor == CurMatineeActor )
+							{
+								FUIAction SelectedActorAction( FExecuteAction::CreateStatic( &FLevelEditorActionCallbacks::OnSelectMatineeGroup, SelectedActor ) );
+								MenuBuilder.AddMenuEntry( LOCTEXT("SelectMatineeGroupForActorMenuTitle", "Select Matinee Group For This Actor"), LOCTEXT("SelectMatineeGroupForActorMenuTooltip", "Selects matinee group controlling this actor"), FSlateIcon(), SelectedActorAction );
+							}
+						}
+					}
+				}
+			}
+		}
+
+		// if this class is Matinee Actor, add option to allow select all controlled actors
+		if ( SelectionInfo.bHaveMatinee )
+		{
+			MenuBuilder.AddMenuEntry( FLevelEditorCommands::Get().SelectAllActorsControlledByMatinee );	
+		}
+	}
+	MenuBuilder.EndSection();
+}
+
+void FLevelEditorContextMenuImpl::FillActorVisibilityMenu( FMenuBuilder& MenuBuilder )
+{
+	MenuBuilder.BeginSection("VisibilitySelected");
+	{
+		// Show 'Show Selected' only if the selection has any hidden actors
+		if ( SelectionInfo.bHaveHidden )
+		{
+			MenuBuilder.AddMenuEntry( FLevelEditorCommands::Get().ShowSelected );
+		}
+		MenuBuilder.AddMenuEntry( FLevelEditorCommands::Get().HideSelected );
+	}
+	MenuBuilder.EndSection();
+
+	MenuBuilder.BeginSection("VisibilityAll");
+	{
+		MenuBuilder.AddMenuEntry( FLevelEditorCommands::Get().ShowSelectedOnly );
+		MenuBuilder.AddMenuEntry( FLevelEditorCommands::Get().ShowAll );
+	}
+	MenuBuilder.EndSection();
+
+	MenuBuilder.BeginSection("VisibilityStartup");
+	{
+		MenuBuilder.AddMenuEntry( FLevelEditorCommands::Get().ShowAllStartup );
+		MenuBuilder.AddMenuEntry( FLevelEditorCommands::Get().ShowSelectedStartup );
+		MenuBuilder.AddMenuEntry( FLevelEditorCommands::Get().HideSelectedStartup );
+	}
+}
+
+void FLevelEditorContextMenuImpl::FillActorLevelMenu( FMenuBuilder& MenuBuilder )
+{
+	MenuBuilder.BeginSection("ActorLevel");
+	{
+		if( SelectionInfo.SharedLevel && SelectionInfo.SharedWorld && SelectionInfo.SharedWorld->GetCurrentLevel() != SelectionInfo.SharedLevel )
+		{
+			// All actors are in the same level and that level is not the current level 
+			// so add a menu entry to make the shared level current
+
+			FText MakeCurrentLevelText = FText::Format( LOCTEXT("MakeCurrentLevelMenu", "Make Current Level: {0}"), FText::FromString( SelectionInfo.SharedLevel->GetOutermost()->GetName() ) );
+			MenuBuilder.AddMenuEntry( FLevelEditorCommands::Get().MakeActorLevelCurrent, NAME_None, MakeCurrentLevelText );
+		}
+
+		if( !SelectionInfo.bAllSelectedActorsBelongToCurrentLevel )
+		{
+			// Only show this menu entry if any actors are not in the current level
+			MenuBuilder.AddMenuEntry( FLevelEditorCommands::Get().MoveSelectedToCurrentLevel );
+		}
+	}
+	MenuBuilder.EndSection();
+
+	MenuBuilder.AddMenuEntry( FLevelEditorCommands::Get().FindActorInLevelScript );
+	MenuBuilder.AddMenuEntry( FLevelEditorCommands::Get().FindLevelsInLevelBrowser );
+	MenuBuilder.AddMenuEntry( FLevelEditorCommands::Get().AddLevelsToSelection );
+	MenuBuilder.AddMenuEntry( FLevelEditorCommands::Get().RemoveLevelsFromSelection );
+}
+
+
+void FLevelEditorContextMenuImpl::FillTransformMenu( FMenuBuilder& MenuBuilder )
+{
+	if ( FLevelEditorActionCallbacks::ActorSelected_CanExecute() )
+	{
+		MenuBuilder.BeginSection("TransformSnapAlign");
+		{
+			MenuBuilder.AddSubMenu( 
+				LOCTEXT("SnapAlignSubMenu", "Snap/Align"), 
+				LOCTEXT("SnapAlignSubMenu_ToolTip", "Actor snap/align utils"),
+				FNewMenuDelegate::CreateStatic( &FLevelEditorContextMenuImpl::FillSnapAlignMenu ) );
+		}
+		MenuBuilder.EndSection();
+
+		MenuBuilder.BeginSection("DeltaTransformToActors");
+		{
+			MenuBuilder.AddMenuEntry( FLevelEditorCommands::Get().DeltaTransformToActors );
+		}
+		MenuBuilder.EndSection();
+	}
+
+	MenuBuilder.BeginSection("MirrorLock");
+	{
+		MenuBuilder.AddMenuEntry( FLevelEditorCommands::Get().MirrorActorX );
+		MenuBuilder.AddMenuEntry( FLevelEditorCommands::Get().MirrorActorY );
+		MenuBuilder.AddMenuEntry( FLevelEditorCommands::Get().MirrorActorZ );
+		MenuBuilder.AddMenuEntry( FLevelEditorCommands::Get().LockActorMovement );
+	}
+}
+
+void FLevelEditorContextMenuImpl::FillActorMenu( FMenuBuilder& MenuBuilder )
+{
+	FSceneOutlinerInitializationOptions InitOptions;
+	{
+		InitOptions.Mode = ESceneOutlinerMode::ActorPicker;			
+		InitOptions.bShowHeaderRow = false;
+		InitOptions.bFocusSearchBoxWhenOpened = true;
+
+		// Only display Actors that we can attach too
+		InitOptions.ActorFilters = MakeShareable( new TFilterCollection< const AActor* const >() );
+
+		struct Local
+		{
+			static bool IsAttachableActor( const AActor* const ParentActor )
+			{					
+				for ( FSelectionIterator It( GEditor->GetSelectedActorIterator() ) ; It ; ++It )
+				{
+					AActor* Actor = static_cast<AActor*>( *It );
+					if (!GEditor->CanParentActors(ParentActor, Actor))
+					{
+						return false;
+					}
+
+					USceneComponent* ChildRoot = Actor->GetRootComponent();
+					USceneComponent* ParentRoot = ParentActor->GetRootComponent();
+
+					if (ChildRoot != nullptr && ParentRoot != nullptr && ChildRoot->IsAttachedTo(ParentRoot))
+					{
+						return false;
+					}
+				}
+				return true;
+			}
+		};
+
+		InitOptions.ActorFilters->Add( MakeShareable( new TDelegateFilter< const AActor* const >( TDelegateFilter< const AActor* const >::FPredicate::CreateStatic( &Local::IsAttachableActor ) ) ) );
+	}		
+
+	// Actor selector to allow the user to choose a parent actor
+	FSceneOutlinerModule& SceneOutlinerModule = FModuleManager::LoadModuleChecked<FSceneOutlinerModule>( "SceneOutliner" );
+	TSharedRef< SWidget > MiniSceneOutliner =
+		SNew( SVerticalBox )
+		+SVerticalBox::Slot()
+		.MaxHeight(400.0f)
+		[
+			SceneOutlinerModule.CreateSceneOutliner(
+			InitOptions,
+			FOnActorPicked::CreateStatic( &FLevelEditorActionCallbacks::AttachToActor ) )
+		];
+
+	// Give the scene outliner a border and background
+	const FSlateBrush* BackgroundBrush = FEditorStyle::GetBrush( "Menu.Background" );
+	TSharedRef< SBorder > RootBorder =
+		SNew( SBorder )
+		.Padding(3)
+		.BorderImage( BackgroundBrush )
+		.ForegroundColor( FEditorStyle::GetSlateColor("DefaultForeground") )
+
+		// Assign the box panel as the child
+		[
+			SNew( SVerticalBox )
+			+SVerticalBox::Slot()
+			.Padding( 5 )
+			.HAlign( HAlign_Center )
+			[
+				SNew( STextBlock )
+				.Text( LOCTEXT( "SelectMatineeActorToEdit", "Select a Matinee actor" ) )
+			]
+
+			+SVerticalBox::Slot()
+			.Padding( 2 )
+			[
+				MiniSceneOutliner
+			]
+		];
+	if(SelectionInfo.bHaveAttachedActor)
+	{
+		MenuBuilder.AddMenuEntry( FLevelEditorCommands::Get().DetachFromParent, NAME_None, LOCTEXT( "None", "None" ) );
+	}
+
+	MenuBuilder.AddWidget(MiniSceneOutliner, FText::GetEmpty(), false);
+}
+
+void FLevelEditorContextMenuImpl::FillSnapAlignMenu( FMenuBuilder& MenuBuilder )
+{
+	MenuBuilder.AddMenuEntry( FLevelEditorCommands::Get().SnapOriginToGrid );
+	MenuBuilder.AddMenuEntry( FLevelEditorCommands::Get().SnapOriginToGridPerActor );
+	MenuBuilder.AddMenuEntry( FLevelEditorCommands::Get().AlignOriginToGrid );
+	MenuBuilder.AddMenuEntry( FLevelEditorCommands::Get().SnapToFloor );
+	MenuBuilder.AddMenuEntry( FLevelEditorCommands::Get().AlignToFloor );
+	MenuBuilder.AddMenuEntry( FLevelEditorCommands::Get().SnapPivotToFloor );
+	MenuBuilder.AddMenuEntry( FLevelEditorCommands::Get().AlignPivotToFloor );
+	MenuBuilder.AddMenuEntry( FLevelEditorCommands::Get().SnapBottomCenterBoundsToFloor );
+	MenuBuilder.AddMenuEntry( FLevelEditorCommands::Get().AlignBottomCenterBoundsToFloor );
+/*
+	MenuBuilder.AddMenuSeparator();
+	AActor* Actor = GEditor->GetSelectedActors()->GetBottom<AActor>();
+	if( Actor && FLevelEditorActionCallbacks::ActorsSelected_CanExecute())
+	{
+		const FString Label = Actor->GetActorLabel();	// Update the options to show the actors label
+		
+		TSharedPtr< FUICommandInfo > SnapOriginToActor = FLevelEditorCommands::Get().SnapOriginToActor;
+		TSharedPtr< FUICommandInfo > AlignOriginToActor = FLevelEditorCommands::Get().AlignOriginToActor;
+		TSharedPtr< FUICommandInfo > SnapToActor = FLevelEditorCommands::Get().SnapToActor;
+		TSharedPtr< FUICommandInfo > AlignToActor = FLevelEditorCommands::Get().AlignToActor;
+		TSharedPtr< FUICommandInfo > SnapPivotToActor = FLevelEditorCommands::Get().SnapPivotToActor;
+		TSharedPtr< FUICommandInfo > AlignPivotToActor = FLevelEditorCommands::Get().AlignPivotToActor;
+		TSharedPtr< FUICommandInfo > SnapBottomCenterBoundsToActor = FLevelEditorCommands::Get().SnapBottomCenterBoundsToActor;
+		TSharedPtr< FUICommandInfo > AlignBottomCenterBoundsToActor = FLevelEditorCommands::Get().AlignBottomCenterBoundsToActor;
+
+		SnapOriginToActor->Label = FString::Printf( *LOCTEXT("Snap Origin To", "Snap Origin to %s"), *Label);
+		AlignOriginToActor->Label = FString::Printf( *LOCTEXT("Align Origin To", "Align Origin to %s"), *Label);
+		SnapToActor->Label = FString::Printf( *LOCTEXT("Snap To", "Snap to %s"), *Label);
+		AlignToActor->Label = FString::Printf( *LOCTEXT("Align To", "Align to %s"), *Label);
+		SnapPivotToActor->Label = FString::Printf( *LOCTEXT("Snap Pivot To", "Snap Pivot to %s"), *Label);
+		AlignPivotToActor->Label = FString::Printf( *LOCTEXT("Align Pivot To", "Align Pivot to %s"), *Label);
+		SnapBottomCenterBoundsToActor->Label = FString::Printf( *LOCTEXT("Snap Bottom Center Bounds To", "Snap Bottom Center Bounds to %s"), *Label);
+		AlignBottomCenterBoundsToActor->Label = FString::Printf( *LOCTEXT("Align Bottom Center Bounds To", "Align Bottom Center Bounds to %s"), *Label);
+
+		MenuBuilder.AddMenuEntry( SnapOriginToActor );
+		MenuBuilder.AddMenuEntry( AlignOriginToActor );
+		MenuBuilder.AddMenuEntry( SnapToActor );
+		MenuBuilder.AddMenuEntry( AlignToActor );
+		MenuBuilder.AddMenuEntry( SnapPivotToActor );
+		MenuBuilder.AddMenuEntry( AlignPivotToActor );
+		MenuBuilder.AddMenuEntry( SnapBottomCenterBoundsToActor );
+		MenuBuilder.AddMenuEntry( AlignBottomCenterBoundsToActor );
+	}
+	else
+	{
+		MenuBuilder.AddMenuEntry( FLevelEditorCommands::Get().SnapOriginToActor );
+		MenuBuilder.AddMenuEntry( FLevelEditorCommands::Get().AlignOriginToActor );
+		MenuBuilder.AddMenuEntry( FLevelEditorCommands::Get().SnapToActor );
+		MenuBuilder.AddMenuEntry( FLevelEditorCommands::Get().AlignToActor );
+		MenuBuilder.AddMenuEntry( FLevelEditorCommands::Get().SnapPivotToActor );
+		MenuBuilder.AddMenuEntry( FLevelEditorCommands::Get().AlignPivotToActor );
+		MenuBuilder.AddMenuEntry( FLevelEditorCommands::Get().SnapBottomCenterBoundsToActor );
+		MenuBuilder.AddMenuEntry( FLevelEditorCommands::Get().AlignBottomCenterBoundsToActor );
+	}
+*/
+}
+
+void FLevelEditorContextMenuImpl::FillPivotMenu( FMenuBuilder& MenuBuilder )
+{
+	MenuBuilder.BeginSection("SaveResetPivot");
+	{
+		MenuBuilder.AddMenuEntry( FLevelEditorCommands::Get().SavePivotToPrePivot );
+		MenuBuilder.AddMenuEntry( FLevelEditorCommands::Get().ResetPrePivot );
+		MenuBuilder.AddMenuEntry( FLevelEditorCommands::Get().ResetPivot );
+	}
+	MenuBuilder.EndSection();
+
+	MenuBuilder.BeginSection("MovePivot");
+	{
+		MenuBuilder.AddMenuEntry( FLevelEditorCommands::Get().MovePivotHere );
+		MenuBuilder.AddMenuEntry( FLevelEditorCommands::Get().MovePivotHereSnapped );
+		MenuBuilder.AddMenuEntry( FLevelEditorCommands::Get().MovePivotToCenter );
+	}
+	MenuBuilder.EndSection();
+}
+
+void FLevelEditorContextMenuImpl::FillGroupMenu( class FMenuBuilder& MenuBuilder )
+{
+	if( SelectionInfo.NumSelectedUngroupedActors > 1 )
+	{
+		// Only show this menu item if we have more than one actor.
+		MenuBuilder.AddMenuEntry( FLevelEditorCommands::Get().GroupActors  );
+	}
+
+	if( SelectionInfo.bHaveSelectedLockedGroup || SelectionInfo.bHaveSelectedUnlockedGroup )
+	{
+		const int32 NumActiveGroups = AGroupActor::NumActiveGroups(true);
+
+		// Regroup will clear any existing groups and create a new one from the selection
+		// Only allow regrouping if multiple groups are selected, or a group and ungrouped actors are selected
+		if( NumActiveGroups > 1 || (NumActiveGroups && SelectionInfo.NumSelectedUngroupedActors) )
+		{
+			MenuBuilder.AddMenuEntry( FLevelEditorCommands::Get().RegroupActors );
+		}
+
+		MenuBuilder.AddMenuEntry( FLevelEditorCommands::Get().UngroupActors );
+
+		if( SelectionInfo.bHaveSelectedUnlockedGroup )
+		{
+			// Only allow removal of loose actors or locked subgroups
+			if( !SelectionInfo.bHaveSelectedLockedGroup || ( SelectionInfo.bHaveSelectedLockedGroup && SelectionInfo.bHaveSelectedSubGroup ) )
+			{
+				MenuBuilder.AddMenuEntry( FLevelEditorCommands::Get().RemoveActorsFromGroup );
+			}
+			MenuBuilder.AddMenuEntry( FLevelEditorCommands::Get().LockGroup );
+		}
+
+		if( SelectionInfo.bHaveSelectedLockedGroup )
+		{
+			MenuBuilder.AddMenuEntry( FLevelEditorCommands::Get().UnlockGroup );
+		}
+
+		// Only allow group adds if a single group is selected in addition to ungrouped actors
+		if( AGroupActor::NumActiveGroups(true, false) == 1 && SelectionInfo.NumSelectedUngroupedActors )
+		{ 
+			MenuBuilder.AddMenuEntry( FLevelEditorCommands::Get().AddActorsToGroup );
+		}
+	}
+}
+
+void FLevelEditorContextMenuImpl::FillEditMenu( class FMenuBuilder& MenuBuilder, LevelEditorMenuContext ContextType )
+{
+	MenuBuilder.AddMenuEntry( FGenericCommands::Get().Cut );
+	MenuBuilder.AddMenuEntry( FGenericCommands::Get().Copy );
+	MenuBuilder.AddMenuEntry( FGenericCommands::Get().Paste );
+	if (ContextType == LevelEditorMenuContext::Viewport)
+	{
+		MenuBuilder.AddMenuEntry( FLevelEditorCommands::Get().PasteHere );
+	}
+
+	MenuBuilder.AddMenuEntry( FGenericCommands::Get().Duplicate );
+	MenuBuilder.AddMenuEntry( FGenericCommands::Get().Delete );
+	MenuBuilder.AddMenuEntry( FGenericCommands::Get().Rename );
+}
+
+void FLevelEditorContextMenuImpl::FillMergeActorsMenu( class FMenuBuilder& MenuBuilder )
+{
+	MenuBuilder.AddMenuEntry( FLevelEditorCommands::Get().MergeActorsByMaterials );
+
+	IMeshUtilities& MeshUtilities = FModuleManager::Get().LoadModuleChecked<IMeshUtilities>("MeshUtilities");
+	IMeshReduction* MeshReduction = MeshUtilities.GetMeshReductionInterface();
+
+	if (MeshReduction && MeshReduction->IsSupported())
+	{
+		MenuBuilder.BeginSection("ProxySimplygon", LOCTEXT("SimplygonHeading", "Simplygon"));
+		{
+			MenuBuilder.AddMenuEntry(FLevelEditorCommands::Get().MergeActors);
+		}
+		MenuBuilder.EndSection();
+	}
+}
+
+void FLevelScriptEventMenuHelper::FillLevelBlueprintEventsMenu(class FMenuBuilder& MenuBuilder, const TArray<AActor*>& SelectedActors)
+{
+	AActor* SelectedActor = (1 == SelectedActors.Num()) ? SelectedActors[0] : NULL;
+
+	if (FKismetEditorUtilities::IsActorValidForLevelScript(SelectedActor))
+	{
+		const bool bAnyEventExists = FKismetEditorUtilities::AnyBoundLevelScriptEventForActor(SelectedActor, false);
+		const bool bAnyEventCanBeAdded = FKismetEditorUtilities::AnyBoundLevelScriptEventForActor(SelectedActor, true);
+
+		if (bAnyEventExists || bAnyEventCanBeAdded)
+		{
+			TWeakObjectPtr<AActor> ActorPtr(SelectedActor);
+
+			MenuBuilder.BeginSection("LevelBlueprintEvents", LOCTEXT("LevelBlueprintEvents", "Level Blueprint Events"));
+
+			if (bAnyEventExists)
+			{
+				MenuBuilder.AddSubMenu(
+					LOCTEXT("JumpEventSubMenu", "Jump to Event"),
+					FText::GetEmpty(),
+					FNewMenuDelegate::CreateStatic(&FKismetEditorUtilities::AddLevelScriptEventOptionsForActor
+					, ActorPtr
+					, true
+					, false
+					, true));
+			}
+
+			if (bAnyEventCanBeAdded)
+			{
+				MenuBuilder.AddSubMenu(
+					LOCTEXT("AddEventSubMenu", "Add Event"),
+					FText::GetEmpty(),
+					FNewMenuDelegate::CreateStatic(&FKismetEditorUtilities::AddLevelScriptEventOptionsForActor
+					, ActorPtr
+					, false
+					, true
+					, true));
+			}
+
+			MenuBuilder.EndSection();
+		}
+	}
+}
+
+
+#undef LOCTEXT_NAMESPACE
diff --git a/Engine/Source/Editor/LevelEditor/Private/LevelEditorContextMenu.h b/Engine/Source/Editor/LevelEditor/Private/LevelEditorContextMenu.h
new file mode 100644
--- /dev/null
+++ b/Engine/Source/Editor/LevelEditor/Private/LevelEditorContextMenu.h
@@ -0,0 +1,64 @@
+// Copyright 1998-2014 Epic Games, Inc. All Rights Reserved.
+
+
+#ifndef __LevelViewportContextMenu_h__
+#define __LevelViewportContextMenu_h__
+
+#pragma once
+
+/** Enum to describe what a menu should be built for */
+enum class LevelEditorMenuContext
+{
+	/** This context menu is applicable to a viewport */
+	Viewport,
+	/** This context menu is applicable to an external UI or dialog (disables click-position-based menu items) */
+	NonViewport,
+};
+
+/**
+ * Context menu construction class 
+ */
+class FLevelEditorContextMenu
+{
+
+public:
+
+	/**
+	 * Summons the level viewport context menu
+	 * @param	LevelEditor		The level editor using this menu.
+	 * @param	ContextType		The context we should use to specialize this menu
+	 */
+	static void SummonMenu( const TSharedRef< class SLevelEditor >& LevelEditor, LevelEditorMenuContext ContextType );
+
+	/**
+	 * Creates a widget for the context menu that can be inserted into a pop-up window
+	 *
+	 * @param	LevelEditor		The level editor using this menu.
+	 * @param	ContextType		The context we should use to specialize this menu
+	 * @param	Extender		Allows extension of this menu based on context.
+	 * @return	Widget for this context menu
+	 */
+	static TSharedPtr< SWidget > BuildMenuWidget(TWeakPtr< SLevelEditor > LevelEditor, LevelEditorMenuContext ContextType, TSharedPtr<FExtender> Extender = TSharedPtr<FExtender>());
+
+	/**
+	 * Populates the specified menu builder for the context menu that can be inserted into a pop-up window
+	 *
+	 * @param	MenuBuilder		The menu builder to fill the menu with
+	 * @param	LevelEditor		The level editor using this menu.
+	 * @param	ContextType		The context we should use to specialize this menu
+	 * @param	Extender		Allows extension of this menu based on context.
+	 */
+	static void FillMenu(FMenuBuilder& MenuBuilder, TWeakPtr< SLevelEditor > LevelEditor, LevelEditorMenuContext ContextType, TSharedPtr<FExtender> Extender);
+
+private:
+
+	/**
+	 * Builds the actor group menu
+	 *
+	 * @param MenuBuilder		The menu builder to add items to.
+	 * @param SelectedActorInfo	Information about the selected actors.
+	 */
+	static void BuildGroupMenu( FMenuBuilder& MenuBuilder, const struct FSelectedActorInfo& SelectedActorInfo );
+};
+
+#endif		// __LevelViewportContextMenu_h__
\ No newline at end of file
diff --git a/Engine/Source/Editor/LevelEditor/Private/LevelViewportContextMenu.cpp b/Engine/Source/Editor/LevelEditor/Private/LevelViewportContextMenu.cpp
deleted file mode 100644
diff --git a/Engine/Source/Editor/LevelEditor/Private/LevelViewportContextMenu.h b/Engine/Source/Editor/LevelEditor/Private/LevelViewportContextMenu.h
deleted file mode 100644
diff --git a/Engine/Source/Editor/LevelEditor/Private/SLevelEditor.cpp b/Engine/Source/Editor/LevelEditor/Private/SLevelEditor.cpp
--- a/Engine/Source/Editor/LevelEditor/Private/SLevelEditor.cpp
+++ b/Engine/Source/Editor/LevelEditor/Private/SLevelEditor.cpp
@@ -11,7 +11,7 @@
 #include "SLevelViewport.h"
 #include "LevelViewportTabContent.h"
 #include "AssetSelection.h"
-#include "LevelViewportContextMenu.h"
+#include "LevelEditorContextMenu.h"
 #include "LevelEditorToolBar.h"
 #include "ScopedTransaction.h"
 #include "SLevelEditorToolBox.h"
@@ -612,7 +612,7 @@ TSharedRef<SDockTab> SLevelEditor::SpawnLevelEditorTab( const FSpawnTabArgs& Arg
 					// Only extend the menu if we have actors selected
 					if (GEditor->GetSelectedActors()->Num())
 					{
-						FLevelViewportContextMenu::FillMenu(MenuBuilder, WeakLevelEditor, TSharedPtr<FExtender>());
+						FLevelEditorContextMenu::FillMenu(MenuBuilder, WeakLevelEditor, LevelEditorMenuContext::NonViewport, TSharedPtr<FExtender>());
 					}
 				}, WeakLevelEditor)
 			);
@@ -1393,7 +1393,7 @@ void SLevelEditor::OnLayoutHasChanged()
 
 void SLevelEditor::SummonLevelViewportContextMenu()
 {
-	FLevelViewportContextMenu::SummonMenu( SharedThis( this ) );
+	FLevelEditorContextMenu::SummonMenu( SharedThis( this ), LevelEditorMenuContext::Viewport );
 }
 
 
diff --git a/Engine/Source/Editor/LevelEditor/Private/SLevelViewport.cpp b/Engine/Source/Editor/LevelEditor/Private/SLevelViewport.cpp
--- a/Engine/Source/Editor/LevelEditor/Private/SLevelViewport.cpp
+++ b/Engine/Source/Editor/LevelEditor/Private/SLevelViewport.cpp
@@ -845,7 +845,7 @@ void FLevelViewportDropContextMenuImpl::FillDropAddReplaceActorMenu( bool bRepla
 				{
 					FUIAction Action( FExecuteAction::CreateStatic( &FLevelEditorActionCallbacks::ReplaceActors_Clicked, MenuItem.FactoryToUse, MenuItem.AssetData ) );
 	
-					FText MenuEntryName = FText::Format( NSLOCTEXT("FLevelViewportContextMenuImpl", "ReplaceActorMenuFormat", "Replace with {0}"), MenuItem.FactoryToUse->GetDisplayName() );
+					FText MenuEntryName = FText::Format( NSLOCTEXT("LevelEditor", "ReplaceActorMenuFormat", "Replace with {0}"), MenuItem.FactoryToUse->GetDisplayName() );
 					if ( MenuItem.AssetData.IsValid() )
 					{
 						MenuEntryName = FText::Format( NSLOCTEXT("LevelEditor", "ReplaceActorUsingAssetMenuFormat", "Replace with {0}: {1}"), 
