#ttp 331268 - SLATE: Replace instances of DragDrop::IsTypeMatch with FDragDropEvent::GetOperationAs()
#summary  Replaced instances of DragDrop::IsTypeMatch with FDragDropEvent::GetOperationAs()
#note Please change any instances of DragDrop::IsTypeMatch for the templated FDragDropEvent::GetOperationAs<>() or FDragDropOperation::IsOfType<>() where necessary
#proj Editor
#branch UE4

diff --git a/Engine/Source/Developer/Profiler/Private/Widgets/SDataGraph.cpp b/Engine/Source/Developer/Profiler/Private/Widgets/SDataGraph.cpp
--- a/Engine/Source/Developer/Profiler/Private/Widgets/SDataGraph.cpp
+++ b/Engine/Source/Developer/Profiler/Private/Widgets/SDataGraph.cpp
@@ -1238,10 +1238,9 @@ void SDataGraph::OnDragEnter( const FGeometry& MyGeometry, const FDragDropEvent&
 {
 	SCompoundWidget::OnDragEnter(MyGeometry, DragDropEvent);
 
-	if( DragDrop::IsTypeMatch<FStatIDDragDropOp>( DragDropEvent.GetOperation() ) )
+	TSharedPtr<FStatIDDragDropOp> Operation = DragDropEvent.GetOperationAs<FStatIDDragDropOp>();
+	if (Operation.IsValid())
 	{
-		// D&D decorator icon
-		auto Operation = StaticCastSharedPtr<FStatIDDragDropOp>( DragDropEvent.GetOperation() );
 		Operation->ShowOK();
 	}
 }
@@ -1250,10 +1249,9 @@ void SDataGraph::OnDragLeave( const FDragDropEvent& DragDropEvent )
 {
 	SCompoundWidget::OnDragLeave(DragDropEvent);
 
-	if( DragDrop::IsTypeMatch<FStatIDDragDropOp>( DragDropEvent.GetOperation() ) )
+	TSharedPtr<FStatIDDragDropOp> Operation = DragDropEvent.GetOperationAs<FStatIDDragDropOp>();
+	if (Operation.IsValid())
 	{
-		// D&D decorator icon
-		auto Operation = StaticCastSharedPtr<FStatIDDragDropOp>( DragDropEvent.GetOperation() );
 		Operation->ShowError();
 	}
 }
@@ -1265,9 +1263,10 @@ FReply SDataGraph::OnDragOver( const FGeometry& MyGeometry, const FDragDropEvent
 
 FReply SDataGraph::OnDrop( const FGeometry& MyGeometry, const FDragDropEvent& DragDropEvent )
 {
-	if( DragDrop::IsTypeMatch<FStatIDDragDropOp>( DragDropEvent.GetOperation() ) )
+	TSharedPtr<FStatIDDragDropOp> Operation = DragDropEvent.GetOperationAs<FStatIDDragDropOp>();
+
+	if(Operation.IsValid())
 	{
-		TSharedPtr<FStatIDDragDropOp> Operation = StaticCastSharedPtr<FStatIDDragDropOp>( DragDropEvent.GetOperation() );
 		if( Operation->IsSingleStatID() )
 		{
 			FProfilerManager::Get()->TrackStat( Operation->GetSingleStatID() );
diff --git a/Engine/Source/Developer/Profiler/Private/Widgets/SProfilerWindow.cpp b/Engine/Source/Developer/Profiler/Private/Widgets/SProfilerWindow.cpp
--- a/Engine/Source/Developer/Profiler/Private/Widgets/SProfilerWindow.cpp
+++ b/Engine/Source/Developer/Profiler/Private/Widgets/SProfilerWindow.cpp
@@ -455,10 +455,9 @@ FReply SProfilerWindow::OnKeyDown( const FGeometry& MyGeometry, const FKeyboardE
 
 FReply SProfilerWindow::OnDragOver( const FGeometry& MyGeometry, const FDragDropEvent& DragDropEvent )
 {
-	if( DragDrop::IsTypeMatch<FExternalDragOperation>( DragDropEvent.GetOperation() ) )
+	TSharedPtr<FExternalDragOperation> DragDropOp = DragDropEvent.GetOperationAs<FExternalDragOperation>();
+	if(DragDropOp.IsValid())
 	{
-		TSharedPtr<FExternalDragOperation> DragDropOp = StaticCastSharedPtr<FExternalDragOperation>(DragDropEvent.GetOperation());
-
 		if( DragDropOp->HasFiles() )
 		{
 			const TArray<FString>& Files = DragDropOp->GetFiles();
@@ -483,10 +482,9 @@ FReply SProfilerWindow::OnDragOver( const FGeometry& MyGeometry, const FDragDrop
 
 FReply SProfilerWindow::OnDrop( const FGeometry& MyGeometry, const FDragDropEvent& DragDropEvent )
 {
-	if( DragDrop::IsTypeMatch<FExternalDragOperation>( DragDropEvent.GetOperation() ) )
+	TSharedPtr<FExternalDragOperation> DragDropOp = DragDropEvent.GetOperationAs<FExternalDragOperation>();
+	if(DragDropOp.IsValid())
 	{
-		TSharedPtr<FExternalDragOperation> DragDropOp = StaticCastSharedPtr<FExternalDragOperation>(DragDropEvent.GetOperation());
-
 		if( DragDropOp->HasFiles() )
 		{
 			// For now, only allow a single file.
diff --git a/Engine/Source/Editor/BehaviorTreeEditor/Private/SGraphNode_BehaviorTree.cpp b/Engine/Source/Editor/BehaviorTreeEditor/Private/SGraphNode_BehaviorTree.cpp
--- a/Engine/Source/Editor/BehaviorTreeEditor/Private/SGraphNode_BehaviorTree.cpp
+++ b/Engine/Source/Editor/BehaviorTreeEditor/Private/SGraphNode_BehaviorTree.cpp
@@ -606,10 +606,10 @@ EVisibility SGraphNode_BehaviorTree::GetDragOverMarkerVisibility() const
 void SGraphNode_BehaviorTree::OnDragEnter( const FGeometry& MyGeometry, const FDragDropEvent& DragDropEvent )
 {
 	// Is someone dragging a node?
-	if ( DragDrop::IsTypeMatch<FDragNode>(DragDropEvent.GetOperation()) )
+	TSharedPtr<FDragNode> DragConnectionOp = DragDropEvent.GetOperationAs<FDragNode>();
+	if (DragConnectionOp.IsValid())
 	{
 		// Inform the Drag and Drop operation that we are hovering over this node.
-		TSharedPtr<FDragNode> DragConnectionOp = StaticCastSharedPtr<FDragNode>(DragDropEvent.GetOperation());
 		TSharedPtr<SGraphNode> SubNode = GetSubNodeUnderCursor(MyGeometry, DragDropEvent);
 		DragConnectionOp->SetHoveredNode( SubNode.IsValid() ? SubNode : SharedThis(this) );
 		if (DragConnectionOp->IsValidOperation() && Cast<UBehaviorTreeGraphNode>(GraphNode)->ParentNode)
@@ -624,10 +624,10 @@ void SGraphNode_BehaviorTree::OnDragEnter( const FGeometry& MyGeometry, const FD
 FReply SGraphNode_BehaviorTree::OnDragOver( const FGeometry& MyGeometry, const FDragDropEvent& DragDropEvent )
 {
 	// Is someone dragging a node?
-	if ( DragDrop::IsTypeMatch<FDragNode>(DragDropEvent.GetOperation()) )
+	TSharedPtr<FDragNode> DragConnectionOp = DragDropEvent.GetOperationAs<FDragNode>();
+	if (DragConnectionOp.IsValid())
 	{
 		// Inform the Drag and Drop operation that we are hovering over this node.
-		TSharedPtr<FDragNode> DragConnectionOp = StaticCastSharedPtr<FDragNode>(DragDropEvent.GetOperation());
 		TSharedPtr<SGraphNode> SubNode = GetSubNodeUnderCursor(MyGeometry, DragDropEvent);
 		DragConnectionOp->SetHoveredNode( SubNode.IsValid() ? SubNode : SharedThis(this) );
 	}
@@ -636,10 +636,10 @@ FReply SGraphNode_BehaviorTree::OnDragOver( const FGeometry& MyGeometry, const F
 
 void SGraphNode_BehaviorTree::OnDragLeave( const FDragDropEvent& DragDropEvent )
 {
-	if ( DragDrop::IsTypeMatch<FDragNode>(DragDropEvent.GetOperation()) )
+	TSharedPtr<FDragNode> DragConnectionOp = DragDropEvent.GetOperationAs<FDragNode>();
+	if (DragConnectionOp.IsValid())
 	{
 		// Inform the Drag and Drop operation that we are not hovering any pins
-		TSharedPtr<FDragNode> DragConnectionOp = StaticCastSharedPtr<FDragNode>(DragDropEvent.GetOperation());
 		DragConnectionOp->SetHoveredNode( TSharedPtr<SGraphNode>(NULL) );
 	}
 	SetDragMarker(false);
@@ -649,9 +649,9 @@ void SGraphNode_BehaviorTree::OnDragLeave( const FDragDropEvent& DragDropEvent )
 
 FReply SGraphNode_BehaviorTree::OnDrop( const FGeometry& MyGeometry, const FDragDropEvent& DragDropEvent )
 {
-	if ( DragDrop::IsTypeMatch<FDragNode>(DragDropEvent.GetOperation()) )
+	TSharedPtr<FDragNode> DragNodeOp = DragDropEvent.GetOperationAs<FDragNode>();
+	if (DragNodeOp.IsValid())
 	{
-		TSharedPtr<FDragNode> DragNodeOp = StaticCastSharedPtr<FDragNode>(DragDropEvent.GetOperation());
 		if (!DragNodeOp->IsValidOperation())
 		{
 			return FReply::Handled();
diff --git a/Engine/Source/Editor/ContentBrowser/Private/AssetViewWidgets.cpp b/Engine/Source/Editor/ContentBrowser/Private/AssetViewWidgets.cpp
--- a/Engine/Source/Editor/ContentBrowser/Private/AssetViewWidgets.cpp
+++ b/Engine/Source/Editor/ContentBrowser/Private/AssetViewWidgets.cpp
@@ -174,7 +174,8 @@ void SAssetViewItem::OnDragEnter( const FGeometry& MyGeometry, const FDragDropEv
 
 	if(IsFolder())
 	{
-		if( DragDrop::IsTypeMatch<FAssetDragDropOp>(DragDropEvent.GetOperation()) )
+		TSharedPtr<FDragDropOperation> Operation = DragDropEvent.GetOperation();
+		if (Operation->IsOfType<FAssetDragDropOp>())
 		{
 			TArray< FAssetData > AssetDatas = AssetUtil::ExtractAssetDataFromDrag( DragDropEvent );
 
@@ -185,7 +186,7 @@ void SAssetViewItem::OnDragEnter( const FGeometry& MyGeometry, const FDragDropEv
 			}
 			bDraggedOver = true;
 		}
-		else if( DragDrop::IsTypeMatch<FAssetPathDragDropOp>(DragDropEvent.GetOperation()) )
+		else if (Operation->IsOfType<FAssetPathDragDropOp>())
 		{
 			TSharedPtr<FAssetPathDragDropOp> DragDropOp = StaticCastSharedPtr<FAssetPathDragDropOp>( DragDropEvent.GetOperation() );
 			bool bCanDrop = DragDropOp->PathNames.Num() > 0;
@@ -207,9 +208,9 @@ void SAssetViewItem::OnDragLeave( const FDragDropEvent& DragDropEvent )
 {
 	if(IsFolder())
 	{
-		if( DragDrop::IsTypeMatch<FAssetDragDropOp>(DragDropEvent.GetOperation()) )
+		TSharedPtr< FAssetDragDropOp > DragDropOp = DragDropEvent.GetOperationAs< FAssetDragDropOp >();
+		if(DragDropOp.IsValid())
 		{
-			TSharedPtr< FAssetDragDropOp > DragDropOp = StaticCastSharedPtr< FAssetDragDropOp >( DragDropEvent.GetOperation() );	
 			DragDropOp->ResetToDefaultToolTip();
 		}
 	}
@@ -221,23 +222,29 @@ FReply SAssetViewItem::OnDragOver( const FGeometry& MyGeometry, const FDragDropE
 {
 	if(IsFolder())
 	{
-		if ( DragDrop::IsTypeMatch<FExternalDragOperation>(DragDropEvent.GetOperation()) )
+		TSharedPtr<FDragDropOperation> Operation = DragDropEvent.GetOperation();
+		if (!Operation.IsValid())
 		{
-			TSharedPtr<FExternalDragOperation> DragDropOp = StaticCastSharedPtr<FExternalDragOperation>( DragDropEvent.GetOperation() );
+			return FReply::Unhandled();
+		}
+
+		if (Operation->IsOfType<FExternalDragOperation>())
+		{
+			TSharedPtr<FExternalDragOperation> DragDropOp = StaticCastSharedPtr<FExternalDragOperation>(Operation);
 			if ( DragDropOp->HasFiles() )
 			{
 				bDraggedOver = true;
 				return FReply::Handled();
 			}
 		}
-		else if( DragDrop::IsTypeMatch<FAssetDragDropOp>(DragDropEvent.GetOperation()) )
+		else if (Operation->IsOfType<FAssetDragDropOp>())
 		{
 			bDraggedOver = true;
 			return FReply::Handled();
 		}
-		else if( DragDrop::IsTypeMatch<FAssetPathDragDropOp>(DragDropEvent.GetOperation()) )
+		else if (Operation->IsOfType<FAssetPathDragDropOp>())
 		{
-			TSharedPtr<FAssetPathDragDropOp> DragDropOp = StaticCastSharedPtr<FAssetPathDragDropOp>( DragDropEvent.GetOperation() );
+			TSharedPtr<FAssetPathDragDropOp> DragDropOp = StaticCastSharedPtr<FAssetPathDragDropOp>(Operation);
 			bool bCanDrop = DragDropOp->PathNames.Num() > 0;
 			if ( DragDropOp->PathNames.Contains(StaticCastSharedPtr<FAssetViewFolder>(AssetItem)->FolderPath) )
 			{
@@ -264,9 +271,15 @@ FReply SAssetViewItem::OnDrop( const FGeometry& MyGeometry, const FDragDropEvent
 	{
 		check(AssetItem->GetType() == EAssetItemType::Folder);
 
-		if ( DragDrop::IsTypeMatch<FExternalDragOperation>(DragDropEvent.GetOperation()) )
+		TSharedPtr<FDragDropOperation> Operation = DragDropEvent.GetOperation();
+		if (!Operation.IsValid())
 		{
-			TSharedPtr<FExternalDragOperation> DragDropOp = StaticCastSharedPtr<FExternalDragOperation>( DragDropEvent.GetOperation() );
+			return FReply::Unhandled();
+		}
+
+		if (Operation->IsOfType<FExternalDragOperation>())
+		{
+			TSharedPtr<FExternalDragOperation> DragDropOp = StaticCastSharedPtr<FExternalDragOperation>(Operation);
 
 			if ( DragDropOp->HasFiles() )
 			{
@@ -275,9 +288,9 @@ FReply SAssetViewItem::OnDrop( const FGeometry& MyGeometry, const FDragDropEvent
 
 			return FReply::Handled();
 		}
-		else if ( DragDrop::IsTypeMatch<FAssetPathDragDropOp>(DragDropEvent.GetOperation()) )
+		else if (Operation->IsOfType<FAssetPathDragDropOp>())
 		{
-			TSharedPtr<FAssetPathDragDropOp> DragDropOp = StaticCastSharedPtr<FAssetPathDragDropOp>( DragDropEvent.GetOperation() );
+			TSharedPtr<FAssetPathDragDropOp> DragDropOp = StaticCastSharedPtr<FAssetPathDragDropOp>(Operation);
 
 			bool bCanDrop = DragDropOp->PathNames.Num() > 0;
 
@@ -294,9 +307,9 @@ FReply SAssetViewItem::OnDrop( const FGeometry& MyGeometry, const FDragDropEvent
 
 			return FReply::Handled();
 		}
-		else if ( DragDrop::IsTypeMatch<FAssetDragDropOp>(DragDropEvent.GetOperation()) )
+		else if (Operation->IsOfType<FAssetDragDropOp>())
 		{
-			TSharedPtr<FAssetDragDropOp> DragDropOp = StaticCastSharedPtr<FAssetDragDropOp>( DragDropEvent.GetOperation() );
+			TSharedPtr<FAssetDragDropOp> DragDropOp = StaticCastSharedPtr<FAssetDragDropOp>(Operation);
 
 			OnAssetsDragDropped.ExecuteIfBound(DragDropOp->AssetData, StaticCastSharedPtr<FAssetViewFolder>(AssetItem)->FolderPath);
 
diff --git a/Engine/Source/Editor/ContentBrowser/Private/SAssetView.cpp b/Engine/Source/Editor/ContentBrowser/Private/SAssetView.cpp
--- a/Engine/Source/Editor/ContentBrowser/Private/SAssetView.cpp
+++ b/Engine/Source/Editor/ContentBrowser/Private/SAssetView.cpp
@@ -1030,19 +1030,18 @@ void SAssetView::ProcessQueriedItems( const double TickStartTime )
 
 void SAssetView::OnDragLeave( const FDragDropEvent& DragDropEvent )
 {
-	if( DragDrop::IsTypeMatch<FAssetDragDropOp>( DragDropEvent.GetOperation() ) )
+	TSharedPtr< FAssetDragDropOp > DragAssetOp = DragDropEvent.GetOperationAs< FAssetDragDropOp >();
+	if( DragAssetOp.IsValid() )
 	{
-		TSharedPtr< FAssetDragDropOp > DragAssetOp = StaticCastSharedPtr< FAssetDragDropOp >( DragDropEvent.GetOperation() );	
 		DragAssetOp->ResetToDefaultToolTip();
 	}
 }
 
 FReply SAssetView::OnDragOver( const FGeometry& MyGeometry, const FDragDropEvent& DragDropEvent )
 {
-	if ( DragDrop::IsTypeMatch<FExternalDragOperation>(DragDropEvent.GetOperation()) )
+	TSharedPtr< FExternalDragOperation > DragDropOp = DragDropEvent.GetOperationAs< FExternalDragOperation >();
+	if ( DragDropOp.IsValid() )
 	{
-		TSharedPtr<FExternalDragOperation> DragDropOp = StaticCastSharedPtr<FExternalDragOperation>( DragDropEvent.GetOperation() );
-
 		if ( DragDropOp->HasFiles() )
 		{
 			return FReply::Handled();
@@ -1054,10 +1053,9 @@ FReply SAssetView::OnDragOver( const FGeometry& MyGeometry, const FDragDropEvent
 
 		if ( AssetDatas.Num() > 0 )
 		{
-			if( DragDrop::IsTypeMatch<FAssetDragDropOp>( DragDropEvent.GetOperation() ) )
+			TSharedPtr< FAssetDragDropOp > DragAssetOp = DragDropEvent.GetOperationAs< FAssetDragDropOp >();
+			if( DragAssetOp.IsValid() )
 			{
-				TSharedPtr< FAssetDragDropOp > DragAssetOp = StaticCastSharedPtr< FAssetDragDropOp >( DragDropEvent.GetOperation() );	
-
 				TArray< FName > ObjectPaths;
 				FCollectionManagerModule& CollectionManagerModule = FModuleManager::LoadModuleChecked<FCollectionManagerModule>(TEXT("CollectionManager"));
 				CollectionManagerModule.Get().GetObjectsInCollection( SourcesData.Collections[0].Name, SourcesData.Collections[0].Type, ObjectPaths );
@@ -1090,10 +1088,9 @@ FReply SAssetView::OnDrop( const FGeometry& MyGeometry, const FDragDropEvent& Dr
 	// Handle drag drop for import
 	if ( IsAssetPathSelected() )
 	{
-		if ( DragDrop::IsTypeMatch<FExternalDragOperation>(DragDropEvent.GetOperation()) )
+		TSharedPtr<FExternalDragOperation> DragDropOp = DragDropEvent.GetOperationAs<FExternalDragOperation>();
+		if (DragDropOp.IsValid())
 		{
-			TSharedPtr<FExternalDragOperation> DragDropOp = StaticCastSharedPtr<FExternalDragOperation>( DragDropEvent.GetOperation() );
-
 			if ( DragDropOp->HasFiles() )
 			{
 				FAssetToolsModule& AssetToolsModule = FModuleManager::Get().LoadModuleChecked<FAssetToolsModule>("AssetTools");
diff --git a/Engine/Source/Editor/ContentBrowser/Private/SourcesViewWidgets.cpp b/Engine/Source/Editor/ContentBrowser/Private/SourcesViewWidgets.cpp
--- a/Engine/Source/Editor/ContentBrowser/Private/SourcesViewWidgets.cpp
+++ b/Engine/Source/Editor/ContentBrowser/Private/SourcesViewWidgets.cpp
@@ -87,9 +87,15 @@ void SAssetTreeItem::OnDragEnter( const FGeometry& MyGeometry, const FDragDropEv
 {
 	if ( IsValidAssetPath() )
 	{
-		if ( DragDrop::IsTypeMatch<FAssetDragDropOp>(DragDropEvent.GetOperation()) )
+		TSharedPtr<FDragDropOperation> Operation = DragDropEvent.GetOperation();
+		if (!Operation.IsValid())
 		{
-			TSharedPtr<FAssetDragDropOp> DragDropOp = StaticCastSharedPtr<FAssetDragDropOp>( DragDropEvent.GetOperation() );
+			return;
+		}
+
+		if (Operation->IsOfType<FAssetDragDropOp>())
+		{
+			TSharedPtr<FAssetDragDropOp> DragDropOp = StaticCastSharedPtr<FAssetDragDropOp>(Operation);
 			bool bCanDrop = DragDropOp->AssetData.Num() > 0;
 
 			if (bCanDrop)
@@ -98,9 +104,9 @@ void SAssetTreeItem::OnDragEnter( const FGeometry& MyGeometry, const FDragDropEv
 				bDraggedOver = true;
 			}
 		}
-		else if ( DragDrop::IsTypeMatch<FAssetPathDragDropOp>(DragDropEvent.GetOperation()) )
+		else if (Operation->IsOfType<FAssetPathDragDropOp>())
 		{
-			TSharedPtr<FAssetPathDragDropOp> DragDropOp = StaticCastSharedPtr<FAssetPathDragDropOp>( DragDropEvent.GetOperation() );
+			TSharedPtr<FAssetPathDragDropOp> DragDropOp = StaticCastSharedPtr<FAssetPathDragDropOp>(Operation);
 			bool bCanDrop = DragDropOp->PathNames.Num() > 0;
 
 			if ( DragDropOp->PathNames.Num() == 1 && DragDropOp->PathNames[0] == TreeItem.Pin()->FolderPath )
@@ -114,9 +120,9 @@ void SAssetTreeItem::OnDragEnter( const FGeometry& MyGeometry, const FDragDropEv
 				bDraggedOver = true;
 			}
 		}
-		else if ( DragDrop::IsTypeMatch<FExternalDragOperation>(DragDropEvent.GetOperation()) )
+		else if (Operation->IsOfType<FExternalDragOperation>())
 		{
-			TSharedPtr<FExternalDragOperation> DragDropOp = StaticCastSharedPtr<FExternalDragOperation>( DragDropEvent.GetOperation() );
+			TSharedPtr<FExternalDragOperation> DragDropOp = StaticCastSharedPtr<FExternalDragOperation>(Operation);
 			bool bCanDrop = DragDropOp->HasFiles();
 
 			if (bCanDrop)
@@ -129,9 +135,9 @@ void SAssetTreeItem::OnDragEnter( const FGeometry& MyGeometry, const FDragDropEv
 
 void SAssetTreeItem::OnDragLeave( const FDragDropEvent& DragDropEvent )
 {
-	if ( DragDrop::IsTypeMatch<FAssetDragDropOp>(DragDropEvent.GetOperation()) )
+	TSharedPtr<FAssetDragDropOp> DragDropOp = DragDropEvent.GetOperationAs<FAssetDragDropOp>();
+	if ( DragDropOp.IsValid() )
 	{
-		TSharedPtr<FAssetDragDropOp> DragDropOp = StaticCastSharedPtr<FAssetDragDropOp>( DragDropEvent.GetOperation() );
 		DragDropOp->ResetToDefaultToolTip();
 	}
 	bDraggedOver = false;
@@ -141,7 +147,13 @@ FReply SAssetTreeItem::OnDragOver( const FGeometry& MyGeometry, const FDragDropE
 {
 	if ( IsValidAssetPath() )
 	{
-		if ( DragDrop::IsTypeMatch<FAssetDragDropOp>(DragDropEvent.GetOperation()) )
+		TSharedPtr<FDragDropOperation> Operation = DragDropEvent.GetOperation();
+		if (!Operation.IsValid())
+		{
+			return FReply::Unhandled();
+		}
+
+		if (Operation->IsOfType<FAssetDragDropOp>())
 		{
 			TSharedPtr<FAssetDragDropOp> DragDropOp = StaticCastSharedPtr<FAssetDragDropOp>( DragDropEvent.GetOperation() );
 			bool bCanDrop = DragDropOp->AssetData.Num() > 0;
@@ -151,7 +163,7 @@ FReply SAssetTreeItem::OnDragOver( const FGeometry& MyGeometry, const FDragDropE
 				return FReply::Handled();
 			}
 		}
-		else if ( DragDrop::IsTypeMatch<FAssetPathDragDropOp>(DragDropEvent.GetOperation()) )
+		else if (Operation->IsOfType<FAssetPathDragDropOp>())
 		{
 			TSharedPtr<FAssetPathDragDropOp> DragDropOp = StaticCastSharedPtr<FAssetPathDragDropOp>( DragDropEvent.GetOperation() );
 			bool bCanDrop = DragDropOp->PathNames.Num() > 0;
@@ -167,7 +179,7 @@ FReply SAssetTreeItem::OnDragOver( const FGeometry& MyGeometry, const FDragDropE
 				return FReply::Handled();
 			}
 		}
-		else if ( DragDrop::IsTypeMatch<FExternalDragOperation>(DragDropEvent.GetOperation()) )
+		else if (Operation->IsOfType<FExternalDragOperation>())
 		{
 			TSharedPtr<FExternalDragOperation> DragDropOp = StaticCastSharedPtr<FExternalDragOperation>( DragDropEvent.GetOperation() );
 			bool bCanDrop = DragDropOp->HasFiles();
@@ -188,7 +200,13 @@ FReply SAssetTreeItem::OnDrop( const FGeometry& MyGeometry, const FDragDropEvent
 
 	if ( IsValidAssetPath() )
 	{
-		if ( DragDrop::IsTypeMatch<FAssetDragDropOp>(DragDropEvent.GetOperation()) )
+		TSharedPtr<FDragDropOperation> Operation = DragDropEvent.GetOperation();
+		if (!Operation.IsValid())
+		{
+			return FReply::Unhandled();
+		}
+
+		if (Operation->IsOfType<FAssetDragDropOp>())
 		{
 			TSharedPtr<FAssetDragDropOp> DragDropOp = StaticCastSharedPtr<FAssetDragDropOp>( DragDropEvent.GetOperation() );
 
@@ -196,7 +214,7 @@ FReply SAssetTreeItem::OnDrop( const FGeometry& MyGeometry, const FDragDropEvent
 
 			return FReply::Handled();
 		}
-		else if ( DragDrop::IsTypeMatch<FAssetPathDragDropOp>(DragDropEvent.GetOperation()) )
+		else if (Operation->IsOfType<FAssetPathDragDropOp>())
 		{
 			TSharedPtr<FAssetPathDragDropOp> DragDropOp = StaticCastSharedPtr<FAssetPathDragDropOp>( DragDropEvent.GetOperation() );
 
@@ -215,7 +233,7 @@ FReply SAssetTreeItem::OnDrop( const FGeometry& MyGeometry, const FDragDropEvent
 
 			return FReply::Handled();
 		}
-		else if ( DragDrop::IsTypeMatch<FExternalDragOperation>(DragDropEvent.GetOperation()) )
+		else if (Operation->IsOfType<FExternalDragOperation>())
 		{
 			TSharedPtr<FExternalDragOperation> DragDropOp = StaticCastSharedPtr<FExternalDragOperation>( DragDropEvent.GetOperation() );
 
@@ -450,9 +468,9 @@ void SCollectionListItem::Tick( const FGeometry& AllottedGeometry, const double
 
 void SCollectionListItem::OnDragEnter( const FGeometry& MyGeometry, const FDragDropEvent& DragDropEvent )
 {
-	if ( DragDrop::IsTypeMatch<FAssetDragDropOp>(DragDropEvent.GetOperation()) )
+	TSharedPtr<FAssetDragDropOp> DragDropOp = DragDropEvent.GetOperationAs<FAssetDragDropOp>();
+	if ( DragDropOp.IsValid() )
 	{
-		TSharedPtr<FAssetDragDropOp> DragDropOp = StaticCastSharedPtr<FAssetDragDropOp>( DragDropEvent.GetOperation() );
 		bool bCanDrop = DragDropOp->AssetData.Num() > 0;
 
 		if (bCanDrop)
@@ -469,9 +487,9 @@ void SCollectionListItem::OnDragLeave( const FDragDropEvent& DragDropEvent )
 
 FReply SCollectionListItem::OnDragOver( const FGeometry& MyGeometry, const FDragDropEvent& DragDropEvent )
 {
-	if ( DragDrop::IsTypeMatch<FAssetDragDropOp>(DragDropEvent.GetOperation()) )
+	TSharedPtr<FAssetDragDropOp> DragDropOp = DragDropEvent.GetOperationAs<FAssetDragDropOp>();
+	if (DragDropOp.IsValid())
 	{
-		TSharedPtr<FAssetDragDropOp> DragDropOp = StaticCastSharedPtr<FAssetDragDropOp>( DragDropEvent.GetOperation() );
 		bool bCanDrop = DragDropOp->AssetData.Num() > 0;
 
 		if (bCanDrop)
@@ -487,10 +505,9 @@ FReply SCollectionListItem::OnDrop( const FGeometry& MyGeometry, const FDragDrop
 {
 	bDraggedOver = false;
 
-	if ( DragDrop::IsTypeMatch<FAssetDragDropOp>(DragDropEvent.GetOperation()) && ensure(CollectionItem.IsValid()) )
+	TSharedPtr<FAssetDragDropOp> DragDropOp = DragDropEvent.GetOperationAs<FAssetDragDropOp>();
+	if (DragDropOp.IsValid() && ensure(CollectionItem.IsValid()))
 	{
-		TSharedPtr<FAssetDragDropOp> DragDropOp = StaticCastSharedPtr<FAssetDragDropOp>( DragDropEvent.GetOperation() );
-
 		FText Message;
 		OnAssetsDragDropped.ExecuteIfBound(DragDropOp->AssetData, CollectionItem.Pin(), Message);
 
diff --git a/Engine/Source/Editor/EditorWidgets/Private/SAssetDropTarget.cpp b/Engine/Source/Editor/EditorWidgets/Private/SAssetDropTarget.cpp
--- a/Engine/Source/Editor/EditorWidgets/Private/SAssetDropTarget.cpp
+++ b/Engine/Source/Editor/EditorWidgets/Private/SAssetDropTarget.cpp
@@ -107,11 +107,13 @@ UObject* SAssetDropTarget::GetDroppedObject( const FDragDropEvent& DragDropEvent
 {
 	bOutRecognizedEvent = false;
 	UObject* DroppedObject = NULL;
+
+	TSharedPtr<FDragDropOperation> Operation = DragDropEvent.GetOperation();
 	// Asset being dragged from content browser
-	if ( DragDrop::IsTypeMatch<FAssetDragDropOp>(DragDropEvent.GetOperation()) )
+	if (Operation->IsOfType<FAssetDragDropOp>())
 	{
 		bOutRecognizedEvent = true;
-		TSharedPtr<FAssetDragDropOp> DragDropOp = StaticCastSharedPtr<FAssetDragDropOp>( DragDropEvent.GetOperation() );
+		TSharedPtr<FAssetDragDropOp> DragDropOp = StaticCastSharedPtr<FAssetDragDropOp>(Operation);
 
 		bool bCanDrop = DragDropOp->AssetData.Num() == 1;
 
@@ -122,16 +124,16 @@ UObject* SAssetDropTarget::GetDroppedObject( const FDragDropEvent& DragDropEvent
 			// Make sure the asset is loaded
 			DroppedObject = AssetData.GetAsset();
 		}
-		// Asset being dragged from some external source
-		else if ( DragDrop::IsTypeMatch<FExternalDragOperation>( DragDropEvent.GetOperation() ) )
-		{
-			bOutRecognizedEvent = true;
-			TArray<FAssetData> DroppedAssetData = AssetUtil::ExtractAssetDataFromDrag(DragDropEvent);
+	}
+	// Asset being dragged from some external source
+	else if (Operation->IsOfType<FExternalDragOperation>())
+	{
+		bOutRecognizedEvent = true;
+		TArray<FAssetData> DroppedAssetData = AssetUtil::ExtractAssetDataFromDrag(DragDropEvent);
 
-			if( DroppedAssetData.Num() == 1 )
-			{
-				DroppedObject = DroppedAssetData[0].GetAsset();
-			}
+		if (DroppedAssetData.Num() == 1)
+		{
+			DroppedObject = DroppedAssetData[0].GetAsset();
 		}
 	}
 
diff --git a/Engine/Source/Editor/EnvironmentQueryEditor/Private/SGraphNode_EnvironmentQuery.cpp b/Engine/Source/Editor/EnvironmentQueryEditor/Private/SGraphNode_EnvironmentQuery.cpp
--- a/Engine/Source/Editor/EnvironmentQueryEditor/Private/SGraphNode_EnvironmentQuery.cpp
+++ b/Engine/Source/Editor/EnvironmentQueryEditor/Private/SGraphNode_EnvironmentQuery.cpp
@@ -542,10 +542,10 @@ void SGraphNode_EnvironmentQuery::OnTestToggleChanged(ESlateCheckBoxState::Type
 void SGraphNode_EnvironmentQuery::OnDragEnter( const FGeometry& MyGeometry, const FDragDropEvent& DragDropEvent )
 {
 	// Is someone dragging a node?
-	if ( DragDrop::IsTypeMatch<FDragNode>(DragDropEvent.GetOperation()) )
+	TSharedPtr<FDragNode> DragConnectionOp = DragDropEvent.GetOperationAs<FDragNode>();
+	if (DragConnectionOp.IsValid())
 	{
 		// Inform the Drag and Drop operation that we are hovering over this node.
-		TSharedPtr<FDragNode> DragConnectionOp = StaticCastSharedPtr<FDragNode>(DragDropEvent.GetOperation());
 		TSharedPtr<SGraphNode> SubNode = GetSubNodeUnderCursor(MyGeometry, DragDropEvent);
 		DragConnectionOp->SetHoveredNode( SubNode.IsValid() ? SubNode : SharedThis(this) );
 
@@ -562,10 +562,10 @@ void SGraphNode_EnvironmentQuery::OnDragEnter( const FGeometry& MyGeometry, cons
 FReply SGraphNode_EnvironmentQuery::OnDragOver( const FGeometry& MyGeometry, const FDragDropEvent& DragDropEvent )
 {
 	// Is someone dragging a node?
-	if ( DragDrop::IsTypeMatch<FDragNode>(DragDropEvent.GetOperation()) )
+	TSharedPtr<FDragNode> DragConnectionOp = DragDropEvent.GetOperationAs<FDragNode>();
+	if (DragConnectionOp.IsValid())
 	{
 		// Inform the Drag and Drop operation that we are hovering over this node.
-		TSharedPtr<FDragNode> DragConnectionOp = StaticCastSharedPtr<FDragNode>(DragDropEvent.GetOperation());
 		TSharedPtr<SGraphNode> SubNode = GetSubNodeUnderCursor(MyGeometry, DragDropEvent);
 		DragConnectionOp->SetHoveredNode( SubNode.IsValid() ? SubNode : SharedThis(this) );
 	}
@@ -574,10 +574,10 @@ FReply SGraphNode_EnvironmentQuery::OnDragOver( const FGeometry& MyGeometry, con
 
 void SGraphNode_EnvironmentQuery::OnDragLeave( const FDragDropEvent& DragDropEvent )
 {
-	if ( DragDrop::IsTypeMatch<FDragNode>(DragDropEvent.GetOperation()) )
+	TSharedPtr<FDragNode> DragConnectionOp = DragDropEvent.GetOperationAs<FDragNode>();
+	if (DragConnectionOp.IsValid())
 	{
 		// Inform the Drag and Drop operation that we are not hovering any pins
-		TSharedPtr<FDragNode> DragConnectionOp = StaticCastSharedPtr<FDragNode>(DragDropEvent.GetOperation());
 		DragConnectionOp->SetHoveredNode( TSharedPtr<SGraphNode>(NULL) );
 	}
 	SetDragMarker(false);
@@ -587,9 +587,9 @@ void SGraphNode_EnvironmentQuery::OnDragLeave( const FDragDropEvent& DragDropEve
 
 FReply SGraphNode_EnvironmentQuery::OnDrop( const FGeometry& MyGeometry, const FDragDropEvent& DragDropEvent )
 {
-	if ( DragDrop::IsTypeMatch<FDragNode>(DragDropEvent.GetOperation()) )
+	TSharedPtr<FDragNode> DragNodeOp = DragDropEvent.GetOperationAs<FDragNode>();
+	if (DragNodeOp.IsValid())
 	{
-		TSharedPtr<FDragNode> DragNodeOp = StaticCastSharedPtr<FDragNode>(DragDropEvent.GetOperation());
 		if (!DragNodeOp->IsValidOperation())
 		{
 			return FReply::Handled();
diff --git a/Engine/Source/Editor/GraphEditor/Private/SGraphActionMenu.cpp b/Engine/Source/Editor/GraphEditor/Private/SGraphActionMenu.cpp
--- a/Engine/Source/Editor/GraphEditor/Private/SGraphActionMenu.cpp
+++ b/Engine/Source/Editor/GraphEditor/Private/SGraphActionMenu.cpp
@@ -108,9 +108,9 @@ public:
 	// SWidget interface
 	virtual FReply OnDrop( const FGeometry& MyGeometry, const FDragDropEvent& DragDropEvent ) OVERRIDE
 	{
-		if ( DragDrop::IsTypeMatch<FGraphEditorDragDropAction>(DragDropEvent.GetOperation()) )
+		TSharedPtr<FGraphEditorDragDropAction> GraphDropOp = DragDropEvent.GetOperationAs<FGraphEditorDragDropAction>();
+		if (GraphDropOp.IsValid())
 		{
-			TSharedPtr<FGraphEditorDragDropAction> GraphDropOp = StaticCastSharedPtr<FGraphEditorDragDropAction>(DragDropEvent.GetOperation());
 			GraphDropOp->DroppedOnCategory( ActionNode.Pin()->Category );
 			return FReply::Handled();
 		}
@@ -119,18 +119,18 @@ public:
 
 	virtual void OnDragEnter( const FGeometry& MyGeometry, const FDragDropEvent& DragDropEvent ) OVERRIDE
 	{
-		if ( DragDrop::IsTypeMatch<FGraphEditorDragDropAction>(DragDropEvent.GetOperation()) )
+		TSharedPtr<FGraphEditorDragDropAction> GraphDropOp = DragDropEvent.GetOperationAs<FGraphEditorDragDropAction>();
+		if (GraphDropOp.IsValid())
 		{
-			TSharedPtr<FGraphEditorDragDropAction> GraphDropOp = StaticCastSharedPtr<FGraphEditorDragDropAction>(DragDropEvent.GetOperation());
 			GraphDropOp->SetHoveredCategoryName( ActionNode.Pin()->Category );
 		}
 	}
 
 	virtual void OnDragLeave( const FDragDropEvent& DragDropEvent ) OVERRIDE
 	{
-		if ( DragDrop::IsTypeMatch<FGraphEditorDragDropAction>(DragDropEvent.GetOperation()) )
+		TSharedPtr<FGraphEditorDragDropAction> GraphDropOp = DragDropEvent.GetOperationAs<FGraphEditorDragDropAction>();
+		if (GraphDropOp.IsValid())
 		{
-			TSharedPtr<FGraphEditorDragDropAction> GraphDropOp = StaticCastSharedPtr<FGraphEditorDragDropAction>(DragDropEvent.GetOperation());
 			GraphDropOp->SetHoveredCategoryName( FString(TEXT("")) );
 		}
 	}
diff --git a/Engine/Source/Editor/GraphEditor/Private/SGraphNode.cpp b/Engine/Source/Editor/GraphEditor/Private/SGraphNode.cpp
--- a/Engine/Source/Editor/GraphEditor/Private/SGraphNode.cpp
+++ b/Engine/Source/Editor/GraphEditor/Private/SGraphNode.cpp
@@ -135,16 +135,22 @@ void SGraphNode::SetDisallowedPinConnectionEvent(SGraphEditor::FOnDisallowedPinC
 
 void SGraphNode::OnDragEnter( const FGeometry& MyGeometry, const FDragDropEvent& DragDropEvent )
 {
+	TSharedPtr<FDragDropOperation> Operation = DragDropEvent.GetOperation();
+	if (!Operation.IsValid())
+	{
+		return;
+	}
+
 	// Is someone dragging a connection?
-	if ( DragDrop::IsTypeMatch<FGraphEditorDragDropAction>(DragDropEvent.GetOperation()) )
+	if (Operation->IsOfType<FGraphEditorDragDropAction>())
 	{
 		// Inform the Drag and Drop operation that we are hovering over this pin.
-		TSharedPtr<FGraphEditorDragDropAction> DragConnectionOp = StaticCastSharedPtr<FGraphEditorDragDropAction>(DragDropEvent.GetOperation());
+		TSharedPtr<FGraphEditorDragDropAction> DragConnectionOp = StaticCastSharedPtr<FGraphEditorDragDropAction>(Operation);
 		DragConnectionOp->SetHoveredNode( SharedThis(this) );
 	}
-	else if( DragDrop::IsTypeMatch<FActorDragDropGraphEdOp>(DragDropEvent.GetOperation()) )
+	else if (Operation->IsOfType<FActorDragDropGraphEdOp>())
 	{
-		TSharedPtr<FActorDragDropGraphEdOp> DragConnectionOp = StaticCastSharedPtr<FActorDragDropGraphEdOp>(DragDropEvent.GetOperation());
+		TSharedPtr<FActorDragDropGraphEdOp> DragConnectionOp = StaticCastSharedPtr<FActorDragDropGraphEdOp>(Operation);
 		if( GraphNode->IsA( UK2Node_Literal::StaticClass() ) )
 		{
 			//Show tool tip only if a single actor is dragged
@@ -169,7 +175,7 @@ void SGraphNode::OnDragEnter( const FGeometry& MyGeometry, const FDragDropEvent&
 			DragConnectionOp->SetToolTip( (DragConnectionOp->Actors.Num() == 1) ? FActorDragDropGraphEdOp::ToolTip_Incompatible : FActorDragDropGraphEdOp::ToolTip_MultipleSelection_Incompatible);
 		}
 	}
-	else if ( DragDrop::IsTypeMatch<FBoneDragDropOp>(DragDropEvent.GetOperation()) )
+	else if (Operation->IsOfType<FBoneDragDropOp>())
 	{
 		//@TODO: A2REMOVAL: No support for A3 nodes handling this drag-drop op yet!
 	}
@@ -177,25 +183,31 @@ void SGraphNode::OnDragEnter( const FGeometry& MyGeometry, const FDragDropEvent&
 
 void SGraphNode::OnDragLeave( const FDragDropEvent& DragDropEvent )
 {
+	TSharedPtr<FDragDropOperation> Operation = DragDropEvent.GetOperation();
+	if (!Operation.IsValid())
+	{
+		return;
+	}
+
 	// Is someone dragging a connection?
-	if ( DragDrop::IsTypeMatch<FGraphEditorDragDropAction>(DragDropEvent.GetOperation()) )
+	if (Operation->IsOfType<FGraphEditorDragDropAction>())
 	{
 		// Inform the Drag and Drop operation that we are not hovering any pins
-		TSharedPtr<FGraphEditorDragDropAction> DragConnectionOp = StaticCastSharedPtr<FGraphEditorDragDropAction>(DragDropEvent.GetOperation());
+		TSharedPtr<FGraphEditorDragDropAction> DragConnectionOp = StaticCastSharedPtr<FGraphEditorDragDropAction>(Operation);
 		DragConnectionOp->SetHoveredNode( TSharedPtr<SGraphNode>(NULL) );
 	}
-	else if( DragDrop::IsTypeMatch<FActorDragDropGraphEdOp>(DragDropEvent.GetOperation()) )
+	else if (Operation->IsOfType<FActorDragDropGraphEdOp>())
 	{
 		//Default tool tip
-		TSharedPtr<FActorDragDropGraphEdOp> DragConnectionOp = StaticCastSharedPtr<FActorDragDropGraphEdOp>(DragDropEvent.GetOperation());
+		TSharedPtr<FActorDragDropGraphEdOp> DragConnectionOp = StaticCastSharedPtr<FActorDragDropGraphEdOp>(Operation);
 		DragConnectionOp->ResetToDefaultToolTip();
 	}
-	else if( DragDrop::IsTypeMatch<FAssetDragDropOp>(DragDropEvent.GetOperation()) )
+	else if (Operation->IsOfType<FAssetDragDropOp>())
 	{
-		TSharedPtr<FAssetDragDropOp> AssetOp = StaticCastSharedPtr<FAssetDragDropOp>(DragDropEvent.GetOperation());
+		TSharedPtr<FAssetDragDropOp> AssetOp = StaticCastSharedPtr<FAssetDragDropOp>(Operation);
 		AssetOp->ResetToDefaultToolTip();
 	}
-	else if ( DragDrop::IsTypeMatch<FBoneDragDropOp>(DragDropEvent.GetOperation()) )
+	else if (Operation->IsOfType<FBoneDragDropOp>())
 	{
 		//@TODO: A2REMOVAL: No support for A3 nodes handling this drag-drop op yet!
 	}
@@ -203,11 +215,11 @@ void SGraphNode::OnDragLeave( const FDragDropEvent& DragDropEvent )
 
 FReply SGraphNode::OnDragOver( const FGeometry& MyGeometry, const FDragDropEvent& DragDropEvent )
 {
-	if( DragDrop::IsTypeMatch<FAssetDragDropOp>(DragDropEvent.GetOperation()) )
+	TSharedPtr<FAssetDragDropOp> AssetOp = DragDropEvent.GetOperationAs<FAssetDragDropOp>();
+	if (AssetOp.IsValid())
 	{
 		if(GraphNode != NULL && GraphNode->GetSchema() != NULL)
 		{
-			TSharedPtr<FAssetDragDropOp> AssetOp = StaticCastSharedPtr<FAssetDragDropOp>(DragDropEvent.GetOperation());
 			bool bOkIcon = false;
 			FString TooltipText;
 			GraphNode->GetSchema()->GetAssetsNodeHoverMessage(AssetOp->AssetData, GraphNode, TooltipText, bOkIcon);
@@ -237,10 +249,16 @@ FVector2D SGraphNode::NodeCoordToGraphCoord( const FVector2D& NodeSpaceCoordinat
 
 FReply SGraphNode::OnDrop( const FGeometry& MyGeometry, const FDragDropEvent& DragDropEvent )
 {
+	TSharedPtr<FDragDropOperation> Operation = DragDropEvent.GetOperation();
+	if (!Operation.IsValid())
+	{
+		return FReply::Unhandled();
+	}
+
 	// Is someone dropping a connection onto this node?
-	if ( DragDrop::IsTypeMatch<FGraphEditorDragDropAction>(DragDropEvent.GetOperation()) )
+	if (Operation->IsOfType<FGraphEditorDragDropAction>())
 	{
-		TSharedPtr<FGraphEditorDragDropAction> DragConnectionOp = StaticCastSharedPtr<FGraphEditorDragDropAction>(DragDropEvent.GetOperation());
+		TSharedPtr<FGraphEditorDragDropAction> DragConnectionOp = StaticCastSharedPtr<FGraphEditorDragDropAction>(Operation);
 
 		const FVector2D NodeAddPosition = NodeCoordToGraphCoord( MyGeometry.AbsoluteToLocal( DragDropEvent.GetScreenSpacePosition() ) );
 
@@ -252,9 +270,9 @@ FReply SGraphNode::OnDrop( const FGeometry& MyGeometry, const FDragDropEvent& Dr
 		}
 		return Result;
 	}
-	else if( DragDrop::IsTypeMatch<FActorDragDropGraphEdOp>(DragDropEvent.GetOperation()) )
+	else if (Operation->IsOfType<FActorDragDropGraphEdOp>())
 	{
-		TSharedPtr<FActorDragDropGraphEdOp> DragConnectionOp = StaticCastSharedPtr<FActorDragDropGraphEdOp>(DragDropEvent.GetOperation());
+		TSharedPtr<FActorDragDropGraphEdOp> DragConnectionOp = StaticCastSharedPtr<FActorDragDropGraphEdOp>(Operation);
 		if( CanAllowInteractionUsingDragDropOp( GraphNode, DragConnectionOp ) )
 		{
 			UK2Node_Literal* LiteralNode = CastChecked< UK2Node_Literal > ( GraphNode );
@@ -271,17 +289,17 @@ FReply SGraphNode::OnDrop( const FGeometry& MyGeometry, const FDragDropEvent& Dr
 		}
 		return FReply::Handled();
 	}
-	else if( DragDrop::IsTypeMatch<FAssetDragDropOp>(DragDropEvent.GetOperation()) )
+	else if (Operation->IsOfType<FAssetDragDropOp>())
 	{
 		UEdGraphNode* Node = GetNodeObj();
 		if(Node != NULL && Node->GetSchema() != NULL)
 		{
-			TSharedPtr<FAssetDragDropOp> AssetOp = StaticCastSharedPtr<FAssetDragDropOp>(DragDropEvent.GetOperation());
+			TSharedPtr<FAssetDragDropOp> AssetOp = StaticCastSharedPtr<FAssetDragDropOp>(Operation);
 			Node->GetSchema()->DroppedAssetsOnNode(AssetOp->AssetData, DragDropEvent.GetScreenSpacePosition(), Node);
 		}
 		return FReply::Handled();
 	} 
-	else if ( DragDrop::IsTypeMatch<FBoneDragDropOp>(DragDropEvent.GetOperation()) )
+	else if (Operation->IsOfType<FBoneDragDropOp>())
 	{
 		//@TODO: A2REMOVAL: No support for A3 nodes handling this drag-drop op yet!
 	}
diff --git a/Engine/Source/Editor/GraphEditor/Private/SGraphPalette.cpp b/Engine/Source/Editor/GraphEditor/Private/SGraphPalette.cpp
--- a/Engine/Source/Editor/GraphEditor/Private/SGraphPalette.cpp
+++ b/Engine/Source/Editor/GraphEditor/Private/SGraphPalette.cpp
@@ -53,10 +53,10 @@ void SGraphPaletteItem::Construct(const FArguments& InArgs, FCreateWidgetForActi
 
 FReply SGraphPaletteItem::OnDrop( const FGeometry& MyGeometry, const FDragDropEvent& DragDropEvent )
 {
-	if ( DragDrop::IsTypeMatch<FGraphEditorDragDropAction>(DragDropEvent.GetOperation()) )
+	TSharedPtr<FGraphEditorDragDropAction> GraphDropOp = DragDropEvent.GetOperationAs<FGraphEditorDragDropAction>();
+	if (GraphDropOp.IsValid())
 	{
-		TSharedPtr<FGraphEditorDragDropAction> GrapDropOp = StaticCastSharedPtr<FGraphEditorDragDropAction>(DragDropEvent.GetOperation());
-		GrapDropOp->DroppedOnAction( ActionPtr.Pin().ToSharedRef() );
+		GraphDropOp->DroppedOnAction(ActionPtr.Pin().ToSharedRef());
 		return FReply::Handled();
 	}
 	return FReply::Unhandled();
@@ -64,18 +64,18 @@ FReply SGraphPaletteItem::OnDrop( const FGeometry& MyGeometry, const FDragDropEv
 
 void SGraphPaletteItem::OnDragEnter( const FGeometry& MyGeometry, const FDragDropEvent& DragDropEvent )
 {
-	if ( DragDrop::IsTypeMatch<FGraphEditorDragDropAction>(DragDropEvent.GetOperation()) )
+	TSharedPtr<FGraphEditorDragDropAction> GraphDropOp = DragDropEvent.GetOperationAs<FGraphEditorDragDropAction>();
+	if (GraphDropOp.IsValid())
 	{
-		TSharedPtr<FGraphEditorDragDropAction> GraphDropOp = StaticCastSharedPtr<FGraphEditorDragDropAction>(DragDropEvent.GetOperation());
 		GraphDropOp->SetHoveredAction( ActionPtr.Pin() );
 	}
 }
 
 void SGraphPaletteItem::OnDragLeave( const FDragDropEvent& DragDropEvent )
 {
-	if ( DragDrop::IsTypeMatch<FGraphEditorDragDropAction>(DragDropEvent.GetOperation()) )
+	TSharedPtr<FGraphEditorDragDropAction> GraphDropOp = DragDropEvent.GetOperationAs<FGraphEditorDragDropAction>();
+	if (GraphDropOp.IsValid())
 	{
-		TSharedPtr<FGraphEditorDragDropAction> GraphDropOp = StaticCastSharedPtr<FGraphEditorDragDropAction>(DragDropEvent.GetOperation());
 		GraphDropOp->SetHoveredAction( NULL );
 	}
 }
diff --git a/Engine/Source/Editor/GraphEditor/Private/SGraphPanel.cpp b/Engine/Source/Editor/GraphEditor/Private/SGraphPanel.cpp
--- a/Engine/Source/Editor/GraphEditor/Private/SGraphPanel.cpp
+++ b/Engine/Source/Editor/GraphEditor/Private/SGraphPanel.cpp
@@ -632,43 +632,52 @@ bool SGraphPanel::OnHandleLeftMouseRelease(const FGeometry& MyGeometry, const FP
 
 void SGraphPanel::OnDragEnter( const FGeometry& MyGeometry, const FDragDropEvent& DragDropEvent )
 {
-	if ( DragDrop::IsTypeMatch<FGraphEditorDragDropAction>(DragDropEvent.GetOperation()) )
+	TSharedPtr<FGraphEditorDragDropAction> DragConnectionOp = DragDropEvent.GetOperationAs<FGraphEditorDragDropAction>();
+	if (DragConnectionOp.IsValid())
 	{
-		TSharedPtr<FGraphEditorDragDropAction> DragConnectionOp = StaticCastSharedPtr<FGraphEditorDragDropAction>(DragDropEvent.GetOperation());
 		DragConnectionOp->SetHoveredGraph( SharedThis(this) );
 	}
 }
 
 void SGraphPanel::OnDragLeave( const FDragDropEvent& DragDropEvent )
 {
-	if ( DragDrop::IsTypeMatch<FGraphEditorDragDropAction>(DragDropEvent.GetOperation()) )
+	TSharedPtr<FGraphEditorDragDropAction> DragConnectionOp = DragDropEvent.GetOperationAs<FGraphEditorDragDropAction>();
+	if (DragConnectionOp.IsValid())
 	{
-		TSharedPtr<FGraphEditorDragDropAction> DragConnectionOp = StaticCastSharedPtr<FGraphEditorDragDropAction>(DragDropEvent.GetOperation());
 		DragConnectionOp->SetHoveredGraph( TSharedPtr<SGraphPanel>(NULL) );
 	}
-	else if( DragDrop::IsTypeMatch<FAssetDragDropOp>(DragDropEvent.GetOperation()) )
+	else
 	{
-		TSharedPtr<FAssetDragDropOp> AssetOp = StaticCastSharedPtr<FAssetDragDropOp>(DragDropEvent.GetOperation());
-		AssetOp->ResetToDefaultToolTip();
+		TSharedPtr<FAssetDragDropOp> AssetOp = DragDropEvent.GetOperationAs<FAssetDragDropOp>();
+		if (AssetOp.IsValid())
+		{
+			AssetOp->ResetToDefaultToolTip();
+		}
 	}
 }
 
 FReply SGraphPanel::OnDragOver( const FGeometry& MyGeometry, const FDragDropEvent& DragDropEvent )
 {
-	if( DragDrop::IsTypeMatch<FGraphEditorDragDropAction>(DragDropEvent.GetOperation()) )
+	TSharedPtr<FDragDropOperation> Operation = DragDropEvent.GetOperation();
+	if (!Operation.IsValid())
+	{
+		return FReply::Unhandled();
+	}
+
+	if( Operation->IsOfType<FGraphEditorDragDropAction>() )
 	{
 		PreviewConnectorEndpoint = MyGeometry.AbsoluteToLocal( DragDropEvent.GetScreenSpacePosition() );
 		return FReply::Handled();
 	}
-	else if(DragDrop::IsTypeMatch<FExternalDragOperation>( DragDropEvent.GetOperation() ) )
+	else if (Operation->IsOfType<FExternalDragOperation>())
 	{
 		return AssetUtil::CanHandleAssetDrag(DragDropEvent);
 	}
-	else if( DragDrop::IsTypeMatch<FAssetDragDropOp>(DragDropEvent.GetOperation()) )
+	else if (Operation->IsOfType<FAssetDragDropOp>())
 	{
 		if(GraphObj != NULL && GraphObj->GetSchema())
 		{
-			TSharedPtr<FAssetDragDropOp> AssetOp = StaticCastSharedPtr<FAssetDragDropOp>(DragDropEvent.GetOperation());
+			TSharedPtr<FAssetDragDropOp> AssetOp = StaticCastSharedPtr<FAssetDragDropOp>(Operation);
 			bool bOkIcon = false;
 			FString TooltipText;
 			GraphObj->GetSchema()->GetAssetsGraphHoverMessage(AssetOp->AssetData, GraphObj, TooltipText, bOkIcon);
@@ -689,22 +698,28 @@ FReply SGraphPanel::OnDrop( const FGeometry& MyGeometry, const FDragDropEvent& D
 
 	FSlateApplication::Get().SetKeyboardFocus(AsShared(), EKeyboardFocusCause::SetDirectly);
 
-	if (DragDrop::IsTypeMatch<FGraphEditorDragDropAction>(DragDropEvent.GetOperation()))
+	TSharedPtr<FDragDropOperation> Operation = DragDropEvent.GetOperation();
+	if (!Operation.IsValid())
+	{
+		return FReply::Unhandled();
+	}
+
+	if (Operation->IsOfType<FGraphEditorDragDropAction>())
 	{
 		check(GraphObj);
-		TSharedPtr<FGraphEditorDragDropAction> DragConn = StaticCastSharedPtr<FGraphEditorDragDropAction>(DragDropEvent.GetOperation());
+		TSharedPtr<FGraphEditorDragDropAction> DragConn = StaticCastSharedPtr<FGraphEditorDragDropAction>(Operation);
 		return DragConn->DroppedOnPanel( SharedThis( this ), DragDropEvent.GetScreenSpacePosition(), NodeAddPosition, *GraphObj);
 	}
-	else if (DragDrop::IsTypeMatch<FActorDragDropGraphEdOp>(DragDropEvent.GetOperation()))
+	else if (Operation->IsOfType<FActorDragDropGraphEdOp>())
 	{
-		TSharedPtr<FActorDragDropGraphEdOp> ActorOp = StaticCastSharedPtr<FActorDragDropGraphEdOp>(DragDropEvent.GetOperation());
+		TSharedPtr<FActorDragDropGraphEdOp> ActorOp = StaticCastSharedPtr<FActorDragDropGraphEdOp>(Operation);
 		OnDropActor.ExecuteIfBound(ActorOp->Actors, GraphObj, NodeAddPosition);
 		return FReply::Handled();
 	}
 
-	else if (DragDrop::IsTypeMatch<FLevelDragDropOp>(DragDropEvent.GetOperation()))
+	else if (Operation->IsOfType<FLevelDragDropOp>())
 	{
-		TSharedPtr<FLevelDragDropOp> LevelOp = StaticCastSharedPtr<FLevelDragDropOp>(DragDropEvent.GetOperation());
+		TSharedPtr<FLevelDragDropOp> LevelOp = StaticCastSharedPtr<FLevelDragDropOp>(Operation);
 		OnDropStreamingLevel.ExecuteIfBound(LevelOp->StreamingLevelsToDrop, GraphObj, NodeAddPosition);
 		return FReply::Handled();
 	}
diff --git a/Engine/Source/Editor/GraphEditor/Private/SGraphPin.cpp b/Engine/Source/Editor/GraphEditor/Private/SGraphPin.cpp
--- a/Engine/Source/Editor/GraphEditor/Private/SGraphPin.cpp
+++ b/Engine/Source/Editor/GraphEditor/Private/SGraphPin.cpp
@@ -499,21 +499,27 @@ void SGraphPin::OnMouseLeave( const FPointerEvent& MouseEvent )
 
 void SGraphPin::OnDragEnter( const FGeometry& MyGeometry, const FDragDropEvent& DragDropEvent )
 {
+	TSharedPtr<FDragDropOperation> Operation = DragDropEvent.GetOperation();
+	if (!Operation.IsValid())
+	{
+		return;
+	}
+
 	// Is someone dragging a connection?
-	if ( DragDrop::IsTypeMatch<FGraphEditorDragDropAction>(DragDropEvent.GetOperation()) )
+	if (Operation->IsOfType<FGraphEditorDragDropAction>())
 	{
 		// Ensure that the pin is valid before using it
 		if(GraphPinObj != NULL && GraphPinObj->GetOuter() != NULL && GraphPinObj->GetOuter()->IsA(UEdGraphNode::StaticClass()))
 		{
 			// Inform the Drag and Drop operation that we are hovering over this pin.
-			TSharedPtr<FGraphEditorDragDropAction> DragConnectionOp = StaticCastSharedPtr<FGraphEditorDragDropAction>(DragDropEvent.GetOperation());
+			TSharedPtr<FGraphEditorDragDropAction> DragConnectionOp = StaticCastSharedPtr<FGraphEditorDragDropAction>(Operation);
 			DragConnectionOp->SetHoveredPin( SharedThis(this) );
 		}	
 
 		// Pins treat being dragged over the same as being hovered outside of drag and drop if they know how to respond to the drag action.
 		SBorder::OnMouseEnter( MyGeometry, DragDropEvent );
 	}
-	else if( DragDrop::IsTypeMatch<FAssetDragDropOp>(DragDropEvent.GetOperation()) )
+	else if (Operation->IsOfType<FAssetDragDropOp>())
 	{
 		TSharedPtr<SGraphNode> NodeWidget = OwnerNodePtr.Pin();
 		if (NodeWidget.IsValid())
@@ -521,7 +527,7 @@ void SGraphPin::OnDragEnter( const FGeometry& MyGeometry, const FDragDropEvent&
 			UEdGraphNode* Node = NodeWidget->GetNodeObj();
 			if(Node != NULL && Node->GetSchema() != NULL)
 			{
-				TSharedPtr<FAssetDragDropOp> AssetOp = StaticCastSharedPtr<FAssetDragDropOp>(DragDropEvent.GetOperation());
+				TSharedPtr<FAssetDragDropOp> AssetOp = StaticCastSharedPtr<FAssetDragDropOp>(Operation);
 				bool bOkIcon = false;
 				FString TooltipText;
 				Node->GetSchema()->GetAssetsPinHoverMessage(AssetOp->AssetData, GraphPinObj, TooltipText, bOkIcon);
@@ -534,18 +540,25 @@ void SGraphPin::OnDragEnter( const FGeometry& MyGeometry, const FDragDropEvent&
 
 void SGraphPin::OnDragLeave( const FDragDropEvent& DragDropEvent )
 {
+	TSharedPtr<FDragDropOperation> Operation = DragDropEvent.GetOperation();
+	if (!Operation.IsValid())
+	{
+		return;
+	}
+
 	// Is someone dragging a connection?
-	if ( DragDrop::IsTypeMatch<FGraphEditorDragDropAction>(DragDropEvent.GetOperation()) )
+	if (Operation->IsOfType<FGraphEditorDragDropAction>())
 	{
 		// Inform the Drag and Drop operation that we are not hovering any pins
-		TSharedPtr<FGraphEditorDragDropAction> DragConnectionOp = StaticCastSharedPtr<FGraphEditorDragDropAction>(DragDropEvent.GetOperation());
+		TSharedPtr<FGraphEditorDragDropAction> DragConnectionOp = StaticCastSharedPtr<FGraphEditorDragDropAction>(Operation);
 		DragConnectionOp->SetHoveredPin( TSharedPtr<SGraphPin>(NULL) );
 
 		SBorder::OnMouseLeave( DragDropEvent );
 	}
-	else if( DragDrop::IsTypeMatch<FAssetDragDropOp>(DragDropEvent.GetOperation()) )
+
+	else if (Operation->IsOfType<FAssetDragDropOp>())
 	{
-		TSharedPtr<FAssetDragDropOp> AssetOp = StaticCastSharedPtr<FAssetDragDropOp>(DragDropEvent.GetOperation());
+		TSharedPtr<FAssetDragDropOp> AssetOp = StaticCastSharedPtr<FAssetDragDropOp>(Operation);
 		AssetOp->ResetToDefaultToolTip();
 	}
 }
@@ -566,10 +579,16 @@ bool SGraphPin::TryHandlePinConnection(SGraphPin& OtherSPin)
 
 FReply SGraphPin::OnDrop( const FGeometry& MyGeometry, const FDragDropEvent& DragDropEvent )
 {
+	TSharedPtr<FDragDropOperation> Operation = DragDropEvent.GetOperation();
+	if (!Operation.IsValid())
+	{
+		return FReply::Unhandled();
+	}
+
 	// Is someone dropping a connection onto this pin?
-	if ( DragDrop::IsTypeMatch<FGraphEditorDragDropAction>(DragDropEvent.GetOperation()) )
+	if (Operation->IsOfType<FGraphEditorDragDropAction>())
 	{
-		TSharedPtr<FGraphEditorDragDropAction> DragConnectionOp = StaticCastSharedPtr<FGraphEditorDragDropAction>(DragDropEvent.GetOperation());
+		TSharedPtr<FGraphEditorDragDropAction> DragConnectionOp = StaticCastSharedPtr<FGraphEditorDragDropAction>(Operation);
 
 		FVector2D NodeAddPosition = FVector2D::ZeroVector;
 		TSharedPtr<SGraphNode> OwnerNode = OwnerNodePtr.Pin();
@@ -601,7 +620,7 @@ FReply SGraphPin::OnDrop( const FGeometry& MyGeometry, const FDragDropEvent& Dra
 		return DragConnectionOp->DroppedOnPin(DragDropEvent.GetScreenSpacePosition(), NodeAddPosition);
 	}
 	// handle dropping an asset on the pin
-	else if( DragDrop::IsTypeMatch<FAssetDragDropOp>(DragDropEvent.GetOperation()) )
+	else if (Operation->IsOfType<FAssetDragDropOp>())
 	{
 		TSharedPtr<SGraphNode> NodeWidget = OwnerNodePtr.Pin();
 		if (NodeWidget.IsValid())
@@ -609,7 +628,7 @@ FReply SGraphPin::OnDrop( const FGeometry& MyGeometry, const FDragDropEvent& Dra
 			UEdGraphNode* Node = NodeWidget->GetNodeObj();
 			if(Node != NULL && Node->GetSchema() != NULL)
 			{
-				TSharedPtr<FAssetDragDropOp> AssetOp = StaticCastSharedPtr<FAssetDragDropOp>(DragDropEvent.GetOperation());
+				TSharedPtr<FAssetDragDropOp> AssetOp = StaticCastSharedPtr<FAssetDragDropOp>(Operation);
 
 				Node->GetSchema()->DroppedAssetsOnPin(AssetOp->AssetData, DragDropEvent.GetScreenSpacePosition(), GraphPinObj);
 			}
diff --git a/Engine/Source/Editor/Kismet/Private/SSCSEditor.cpp b/Engine/Source/Editor/Kismet/Private/SSCSEditor.cpp
--- a/Engine/Source/Editor/Kismet/Private/SSCSEditor.cpp
+++ b/Engine/Source/Editor/Kismet/Private/SSCSEditor.cpp
@@ -60,7 +60,8 @@ FReply SSCSEditorDragDropTree::OnDragOver( const FGeometry& MyGeometry, const FD
 
 	if ( SCSEditor != NULL )
 	{
-		if(DragDrop::IsTypeMatch<FExternalDragOperation>( DragDropEvent.GetOperation() ) )
+		const bool bIsValidDrag = DragDropEvent.GetOperationAs<FExternalDragOperation>().IsValid();
+		if (bIsValidDrag)
 		{
 			Handled = AssetUtil::CanHandleAssetDrag(DragDropEvent);
 		}
@@ -75,7 +76,8 @@ FReply SSCSEditorDragDropTree::OnDrop( const FGeometry& MyGeometry, const FDragD
 
 	if ( SCSEditor != NULL )
 	{
-		if(DragDrop::IsTypeMatch<FExternalDragOperation>( DragDropEvent.GetOperation() ) || DragDrop::IsTypeMatch<FAssetDragDropOp>(DragDropEvent.GetOperation()))
+		TSharedPtr<FDragDropOperation> Operation = DragDropEvent.GetOperation();
+		if (Operation.IsValid() && (Operation->IsOfType<FExternalDragOperation>() || Operation->IsOfType<FAssetDragDropOp>()))
 		{
 			TArray< FAssetData > DroppedAssetData = AssetUtil::ExtractAssetDataFromDrag( DragDropEvent );
 			const int32 NumAssets = DroppedAssetData.Num();
@@ -998,9 +1000,15 @@ FReply SSCS_RowWidget::OnDragDetected( const FGeometry& MyGeometry, const FPoint
 
 void SSCS_RowWidget::OnDragEnter( const FGeometry& MyGeometry, const FDragDropEvent& DragDropEvent )
 {
-	if ( DragDrop::IsTypeMatch<FSCSRowDragDropOp>(DragDropEvent.GetOperation()) )
+	TSharedPtr<FDragDropOperation> Operation = DragDropEvent.GetOperation();
+	if (!Operation.IsValid())
 	{
-		TSharedPtr<FSCSRowDragDropOp> DragRowOp = StaticCastSharedPtr<FSCSRowDragDropOp>( DragDropEvent.GetOperation() );
+		return;
+	}
+
+	if ( Operation->IsOfType<FSCSRowDragDropOp>() )
+	{
+		TSharedPtr<FSCSRowDragDropOp> DragRowOp = StaticCastSharedPtr<FSCSRowDragDropOp>( Operation );
 		check(DragRowOp.IsValid());
 
 		// If the hover text is already set, skip everything below, because it means we already know we can't drag-and-drop one or more of the selected nodes.
@@ -1199,7 +1207,7 @@ void SSCS_RowWidget::OnDragEnter( const FGeometry& MyGeometry, const FDragDropEv
 			}
 		}
 	}
-	else if ( DragDrop::IsTypeMatch<FExternalDragOperation>( DragDropEvent.GetOperation() ) || DragDrop::IsTypeMatch<FAssetDragDropOp>(DragDropEvent.GetOperation() ) )
+	else if ( Operation->IsOfType<FExternalDragOperation>() || Operation->IsOfType<FAssetDragDropOp>() )
 	{
 		// defer to the tree widget's handler for this type of operation
 		TSharedPtr<SSCSEditor> PinnedEditor = SCSEditor.Pin();
@@ -1212,11 +1220,9 @@ void SSCS_RowWidget::OnDragEnter( const FGeometry& MyGeometry, const FDragDropEv
 
 void SSCS_RowWidget::OnDragLeave( const FDragDropEvent& DragDropEvent )
 {
-	if ( DragDrop::IsTypeMatch<FSCSRowDragDropOp>(DragDropEvent.GetOperation()) )
+	TSharedPtr<FSCSRowDragDropOp> DragRowOp = DragDropEvent.GetOperationAs<FSCSRowDragDropOp>();
+	if (DragRowOp.IsValid())
 	{
-		TSharedPtr<FSCSRowDragDropOp> DragRowOp = StaticCastSharedPtr<FSCSRowDragDropOp>( DragDropEvent.GetOperation() );
-		check(DragRowOp.IsValid());
-
 		bool bCanReparentAllNodes = true;
 		for(auto SourceNodeIter = DragRowOp->SourceNodes.CreateConstIterator(); SourceNodeIter && bCanReparentAllNodes; ++SourceNodeIter)
 		{
@@ -1237,9 +1243,15 @@ void SSCS_RowWidget::OnDragLeave( const FDragDropEvent& DragDropEvent )
 
 FReply SSCS_RowWidget::OnDrop( const FGeometry& MyGeometry, const FDragDropEvent& DragDropEvent )
 {
-	if ( DragDrop::IsTypeMatch<FSCSRowDragDropOp>(DragDropEvent.GetOperation()) && NodePtr->GetComponentTemplate()->IsA<USceneComponent>())	
+	TSharedPtr<FDragDropOperation> Operation = DragDropEvent.GetOperation();
+	if (!Operation.IsValid())
+	{
+		return FReply::Handled();
+	}
+
+	if ( Operation->IsOfType<FSCSRowDragDropOp>() && NodePtr->GetComponentTemplate()->IsA<USceneComponent>())	
 	{
-		TSharedPtr<FSCSRowDragDropOp> DragRowOp = StaticCastSharedPtr<FSCSRowDragDropOp>( DragDropEvent.GetOperation() );	
+		TSharedPtr<FSCSRowDragDropOp> DragRowOp = StaticCastSharedPtr<FSCSRowDragDropOp>( Operation );	
 		check(DragRowOp.IsValid());
 
 		switch(DragRowOp->PendingDropAction)
@@ -1272,7 +1284,7 @@ FReply SSCS_RowWidget::OnDrop( const FGeometry& MyGeometry, const FDragDropEvent
 			break;
 		}
 	}
-	else if ( DragDrop::IsTypeMatch<FExternalDragOperation>( DragDropEvent.GetOperation() ) || DragDrop::IsTypeMatch<FAssetDragDropOp>(DragDropEvent.GetOperation() ) )
+	else if (Operation->IsOfType<FExternalDragOperation>() || Operation->IsOfType<FAssetDragDropOp>())
 	{
 		// defer to the tree widget's handler for this type of operation
 		TSharedPtr<SSCSEditor> PinnedEditor = SCSEditor.Pin();
diff --git a/Engine/Source/Editor/Layers/Private/SLayerBrowser.h b/Engine/Source/Editor/Layers/Private/SLayerBrowser.h
--- a/Engine/Source/Editor/Layers/Private/SLayerBrowser.h
+++ b/Engine/Source/Editor/Layers/Private/SLayerBrowser.h
@@ -208,13 +208,11 @@ protected:
 	 */
 	virtual void OnDragLeave( const FDragDropEvent& DragDropEvent ) OVERRIDE
 	{
-		if( !DragDrop::IsTypeMatch<FActorDragDropGraphEdOp>( DragDropEvent.GetOperation() ) )
+		TSharedPtr< FActorDragDropGraphEdOp > DragActorOp = DragDropEvent.GetOperationAs< FActorDragDropGraphEdOp >();
+		if(DragActorOp.IsValid())
 		{
-			return;
+			DragActorOp->ResetToDefaultToolTip();
 		}
-
-		TSharedPtr< FActorDragDropGraphEdOp > DragActorOp = StaticCastSharedPtr< FActorDragDropGraphEdOp >( DragDropEvent.GetOperation() );	
-		DragActorOp->ResetToDefaultToolTip();
 	}
 
 	/**
@@ -227,7 +225,8 @@ protected:
 	 */
 	virtual FReply OnDragOver( const FGeometry& MyGeometry, const FDragDropEvent& DragDropEvent ) OVERRIDE
 	{
-		if( !SelectedLayerViewModel->GetDataSource().IsValid() || !DragDrop::IsTypeMatch<FActorDragDropGraphEdOp>( DragDropEvent.GetOperation() ) )
+		TSharedPtr< FActorDragDropGraphEdOp > DragActorOp = DragDropEvent.GetOperationAs< FActorDragDropGraphEdOp >();
+		if (!SelectedLayerViewModel->GetDataSource().IsValid() || !DragActorOp.IsValid())
 		{
 			return FReply::Unhandled();
 		}
@@ -246,8 +245,6 @@ protected:
 			return FReply::Unhandled();
 		}
 
-		TSharedPtr< FActorDragDropGraphEdOp > DragActorOp = StaticCastSharedPtr< FActorDragDropGraphEdOp >( DragDropEvent.GetOperation() );	
-
 		if ( !DragActorOp.IsValid() || DragActorOp->Actors.Num() == 0 )
 		{
 			return FReply::Unhandled();
@@ -287,7 +284,9 @@ protected:
 	 */
 	virtual FReply OnDrop( const FGeometry& MyGeometry, const FDragDropEvent& DragDropEvent ) OVERRIDE
 	{
-		if ( !SelectedLayerViewModel->GetDataSource().IsValid() || !DragDrop::IsTypeMatch<FActorDragDropGraphEdOp>( DragDropEvent.GetOperation() ) )	
+		TSharedPtr< FActorDragDropGraphEdOp > DragActorOp = DragDropEvent.GetOperationAs< FActorDragDropGraphEdOp >();
+
+		if (!SelectedLayerViewModel->GetDataSource().IsValid() || !DragActorOp.IsValid())
 		{
 			return FReply::Unhandled();
 		}
@@ -306,7 +305,6 @@ protected:
 			return FReply::Unhandled();
 		}
 
-		TSharedPtr< FActorDragDropGraphEdOp > DragActorOp = StaticCastSharedPtr< FActorDragDropGraphEdOp >( DragDropEvent.GetOperation() );		
 		SelectedLayerViewModel->AddActors( DragActorOp->Actors );
 
 		return FReply::Handled();
diff --git a/Engine/Source/Editor/Layers/Private/SLayersView.h b/Engine/Source/Editor/Layers/Private/SLayersView.h
--- a/Engine/Source/Editor/Layers/Private/SLayersView.h
+++ b/Engine/Source/Editor/Layers/Private/SLayersView.h
@@ -151,13 +151,11 @@ protected:
 	 */
 	virtual void OnDragLeave( const FDragDropEvent& DragDropEvent ) OVERRIDE
 	{
-		if( !DragDrop::IsTypeMatch<FActorDragDropGraphEdOp>( DragDropEvent.GetOperation() ) )
+		TSharedPtr< FActorDragDropGraphEdOp > DragActorOp = DragDropEvent.GetOperationAs< FActorDragDropGraphEdOp >();
+		if (DragActorOp.IsValid())
 		{
-			return;
+			DragActorOp->ResetToDefaultToolTip();
 		}
-
-		TSharedPtr< FActorDragDropGraphEdOp > DragActorOp = StaticCastSharedPtr< FActorDragDropGraphEdOp >( DragDropEvent.GetOperation() );	
-		DragActorOp->ResetToDefaultToolTip();
 	}
 
 	/**
@@ -170,14 +168,12 @@ protected:
 	 */
 	virtual FReply OnDragOver( const FGeometry& MyGeometry, const FDragDropEvent& DragDropEvent ) OVERRIDE
 	{
-		if( !DragDrop::IsTypeMatch<FActorDragDropGraphEdOp>( DragDropEvent.GetOperation() ) )
+		TSharedPtr< FActorDragDropGraphEdOp > DragActorOp = DragDropEvent.GetOperationAs< FActorDragDropGraphEdOp >();
+		if (DragActorOp.IsValid())
 		{
-			return FReply::Unhandled();
+			DragActorOp->SetToolTip(FActorDragDropGraphEdOp::ToolTip_CompatibleGeneric, LOCTEXT("OnDragOver", "Add Actors to New Layer"));
 		}
 
-		TSharedPtr< FActorDragDropGraphEdOp > DragActorOp = StaticCastSharedPtr< FActorDragDropGraphEdOp >( DragDropEvent.GetOperation() );	
-		DragActorOp->SetToolTip( FActorDragDropGraphEdOp::ToolTip_CompatibleGeneric, LOCTEXT("OnDragOver", "Add Actors to New Layer") );
-
 		// We leave the event unhandled so the children of the ListView get a chance to grab the drag/drop
 		return FReply::Unhandled();
 	}
@@ -192,13 +188,12 @@ protected:
 	 */
 	virtual FReply OnDrop( const FGeometry& MyGeometry, const FDragDropEvent& DragDropEvent ) OVERRIDE
 	{
-		if ( !DragDrop::IsTypeMatch<FActorDragDropGraphEdOp>( DragDropEvent.GetOperation() ) )	
+		TSharedPtr< FActorDragDropGraphEdOp > DragActorOp = DragDropEvent.GetOperationAs< FActorDragDropGraphEdOp >();
+		if ( !DragActorOp.IsValid() )	
 		{
 			return FReply::Unhandled();
 		}
 
-		TSharedPtr< FActorDragDropGraphEdOp > DragActorOp = StaticCastSharedPtr< FActorDragDropGraphEdOp >( DragDropEvent.GetOperation() );		
-
 		ViewModel->AddActorsToNewLayer( DragActorOp->Actors );
 
 		return FReply::Handled();
diff --git a/Engine/Source/Editor/Layers/Private/SLayersViewRow.h b/Engine/Source/Editor/Layers/Private/SLayersViewRow.h
--- a/Engine/Source/Editor/Layers/Private/SLayersViewRow.h
+++ b/Engine/Source/Editor/Layers/Private/SLayersViewRow.h
@@ -145,13 +145,11 @@ protected:
 	 */
 	virtual void OnDragLeave( const FDragDropEvent& DragDropEvent ) OVERRIDE
 	{
-		if( !DragDrop::IsTypeMatch<FActorDragDropGraphEdOp>( DragDropEvent.GetOperation() ) )
+		TSharedPtr< FActorDragDropGraphEdOp > DragActorOp = DragDropEvent.GetOperationAs< FActorDragDropGraphEdOp >();
+		if (DragActorOp.IsValid())
 		{
-			return;
+			DragActorOp->ResetToDefaultToolTip();
 		}
-
-		TSharedPtr< FActorDragDropGraphEdOp > DragActorOp = StaticCastSharedPtr< FActorDragDropGraphEdOp >( DragDropEvent.GetOperation() );	
-		DragActorOp->ResetToDefaultToolTip();
 	}
 
 	/**
@@ -164,13 +162,12 @@ protected:
 	 */
 	virtual FReply OnDragOver( const FGeometry& MyGeometry, const FDragDropEvent& DragDropEvent ) OVERRIDE
 	{
-		if( !DragDrop::IsTypeMatch<FActorDragDropGraphEdOp>( DragDropEvent.GetOperation() ) )
+		TSharedPtr< FActorDragDropGraphEdOp > DragActorOp = DragDropEvent.GetOperationAs< FActorDragDropGraphEdOp >();
+		if (!DragActorOp.IsValid())
 		{
 			return FReply::Unhandled();
 		}
 
-		TSharedPtr< FActorDragDropGraphEdOp > DragActorOp = StaticCastSharedPtr< FActorDragDropGraphEdOp >( DragDropEvent.GetOperation() );		
-
 		bool bCanAssign = false;
 		FString Message;
 		if( DragActorOp->Actors.Num() > 1 )
@@ -204,12 +201,12 @@ protected:
 	 */
 	virtual FReply OnDrop( const FGeometry& MyGeometry, const FDragDropEvent& DragDropEvent ) OVERRIDE
 	{
-		if ( !DragDrop::IsTypeMatch<FActorDragDropGraphEdOp>( DragDropEvent.GetOperation() ) )	
+		TSharedPtr< FActorDragDropGraphEdOp > DragActorOp = DragDropEvent.GetOperationAs< FActorDragDropGraphEdOp >();
+		if (!DragActorOp.IsValid())
 		{
 			return FReply::Unhandled();
 		}
 
-		TSharedPtr< FActorDragDropGraphEdOp > DragActorOp = StaticCastSharedPtr< FActorDragDropGraphEdOp >( DragDropEvent.GetOperation() );		
 		ViewModel->AddActors( DragActorOp->Actors );
 
 		return FReply::Handled();
@@ -245,7 +242,7 @@ private:
 		bool bCanAcceptDrop = false;
 		TSharedPtr<FDragDropOperation> DragDropOp = FSlateApplication::Get().GetDragDroppingContent();
 
-		if (DragDrop::IsTypeMatch<FActorDragDropGraphEdOp>(DragDropOp))	
+		if (DragDropOp.IsValid() && DragDropOp->IsOfType<FActorDragDropGraphEdOp>())
 		{
 			TSharedPtr<FActorDragDropGraphEdOp> DragDropActorOp = StaticCastSharedPtr<FActorDragDropGraphEdOp>(DragDropOp);
 
diff --git a/Engine/Source/Editor/LevelEditor/Private/SLevelViewport.cpp b/Engine/Source/Editor/LevelEditor/Private/SLevelViewport.cpp
--- a/Engine/Source/Editor/LevelEditor/Private/SLevelViewport.cpp
+++ b/Engine/Source/Editor/LevelEditor/Private/SLevelViewport.cpp
@@ -376,34 +376,38 @@ bool SLevelViewport::HandleDragObjects(const FGeometry& MyGeometry, const FDragD
 	bool bValidDrag = false;
 	TArray<FAssetData> SelectedAssetDatas;
 
-	if ( DragDrop::IsTypeMatch<FClassDragDropOp>(DragDropEvent.GetOperation()) )
+	TSharedPtr< FDragDropOperation > Operation = DragDropEvent.GetOperation();
+	if (!Operation.IsValid())
 	{
-		auto Operation = StaticCastSharedPtr<FClassDragDropOp>( DragDropEvent.GetOperation() );
+		return false;
+	}
+
+	if (Operation->IsOfType<FClassDragDropOp>())
+	{
+		auto ClassOperation = StaticCastSharedPtr<FClassDragDropOp>( Operation );
 
 		bValidDrag = true;
 
-		for (int32 DroppedAssetIdx = 0; DroppedAssetIdx < Operation->ClassesToDrop.Num(); ++DroppedAssetIdx)
+		for (int32 DroppedAssetIdx = 0; DroppedAssetIdx < ClassOperation->ClassesToDrop.Num(); ++DroppedAssetIdx)
 		{
-			new(SelectedAssetDatas) FAssetData(Operation->ClassesToDrop[DroppedAssetIdx].Get());
+			new(SelectedAssetDatas)FAssetData(ClassOperation->ClassesToDrop[DroppedAssetIdx].Get());
 		}
 	}
-	else if ( DragDrop::IsTypeMatch<FUnloadedClassDragDropOp>(DragDropEvent.GetOperation()) )
+	else if (Operation->IsOfType<FUnloadedClassDragDropOp>())
 	{
-		TSharedPtr<FUnloadedClassDragDropOp> DragDropOp = StaticCastSharedPtr<FUnloadedClassDragDropOp>( DragDropEvent.GetOperation() );
-
 		bValidDrag = true;
 	}
-	else if ( DragDrop::IsTypeMatch<FExportTextDragDropOp>(DragDropEvent.GetOperation()) )
+	else if (Operation->IsOfType<FExportTextDragDropOp>())
 	{
 		bValidDrag = true;
 	}
-	else if ( DragDrop::IsTypeMatch<FBrushBuilderDragDropOp>(DragDropEvent.GetOperation()) )
+	else if (Operation->IsOfType<FBrushBuilderDragDropOp>())
 	{
 		bValidDrag = true;
 
-		auto Operation = StaticCastSharedPtr<FBrushBuilderDragDropOp>( DragDropEvent.GetOperation() );
+		auto BrushOperation = StaticCastSharedPtr<FBrushBuilderDragDropOp>( Operation );
 
-		new(SelectedAssetDatas) FAssetData(Operation->GetBrushBuilder().Get());
+		new(SelectedAssetDatas) FAssetData(BrushOperation->GetBrushBuilder().Get());
 	}
 	else
 	{
@@ -433,10 +437,10 @@ bool SLevelViewport::HandleDragObjects(const FGeometry& MyGeometry, const FDragD
 	if (LevelViewportClient->UpdateDropPreviewActors(CachedOnDropLocalMousePos.X, CachedOnDropLocalMousePos.Y, DroppedObjects, bDroppedObjectsVisible))
 	{
 		// if dragged actors were hidden, show decorator
-		DragDropEvent.GetOperation()->SetDecoratorVisibility(! bDroppedObjectsVisible);
+		Operation->SetDecoratorVisibility(! bDroppedObjectsVisible);
 	}
 
-	DragDropEvent.GetOperation()->SetCursorOverride(TOptional<EMouseCursor::Type>());
+	Operation->SetCursorOverride(TOptional<EMouseCursor::Type>());
 
 	FText HintText;
 
@@ -456,7 +460,7 @@ bool SLevelViewport::HandleDragObjects(const FGeometry& MyGeometry, const FDragD
 		if ( !DropResult.bCanDrop )
 		{
 			// At least one of the assets can't be dropped.
-			DragDropEvent.GetOperation()->SetCursorOverride(EMouseCursor::SlashedCircle);
+			Operation->SetCursorOverride(EMouseCursor::SlashedCircle);
 			return false;
 		}
 		else
@@ -468,10 +472,10 @@ bool SLevelViewport::HandleDragObjects(const FGeometry& MyGeometry, const FDragD
 		}
 	}
 
-	if ( DragDrop::IsTypeMatch<FAssetDragDropOp>(DragDropEvent.GetOperation()) )
+	if ( Operation->IsOfType<FAssetDragDropOp>() )
 	{
-		auto Operation = StaticCastSharedPtr<FAssetDragDropOp>(DragDropEvent.GetOperation());
-		Operation->SetToolTip(HintText, NULL);
+		auto AssetOperation = StaticCastSharedPtr<FAssetDragDropOp>(DragDropEvent.GetOperation());
+		AssetOperation->SetToolTip(HintText, NULL);
 	}
 
 	return bValidDrag;
@@ -493,18 +497,24 @@ bool SLevelViewport::HandlePlaceDraggedObjects(const FDragDropEvent& DragDropEve
 	bool bValidDrop = false;
 	UActorFactory* ActorFactory = NULL;
 
-	if ( DragDrop::IsTypeMatch<FClassDragDropOp>(DragDropEvent.GetOperation()) )
+	TSharedPtr< FDragDropOperation > Operation = DragDropEvent.GetOperation();
+	if (!Operation.IsValid())
+	{
+		return false;
+	}
+
+	if (Operation->IsOfType<FClassDragDropOp>())
 	{
-		auto Operation = StaticCastSharedPtr<FClassDragDropOp>( DragDropEvent.GetOperation() );
+		auto ClassOperation = StaticCastSharedPtr<FClassDragDropOp>( Operation );
 
 		DroppedObjects.Empty();
 
 		// Check if the asset is loaded, used to see if the context menu should be available
 		bAllAssetWereLoaded = true;
 
-		for (int32 DroppedAssetIdx = 0; DroppedAssetIdx < Operation->ClassesToDrop.Num(); ++DroppedAssetIdx)
+		for (int32 DroppedAssetIdx = 0; DroppedAssetIdx < ClassOperation->ClassesToDrop.Num(); ++DroppedAssetIdx)
 		{
-			UObject* Object = Operation->ClassesToDrop[DroppedAssetIdx].Get();
+			UObject* Object = ClassOperation->ClassesToDrop[DroppedAssetIdx].Get();
 
 			if(Object)
 			{
@@ -518,9 +528,9 @@ bool SLevelViewport::HandlePlaceDraggedObjects(const FDragDropEvent& DragDropEve
 
 		bValidDrop = true;
 	}
-	else if ( DragDrop::IsTypeMatch<FUnloadedClassDragDropOp>(DragDropEvent.GetOperation()) )
+	else if (Operation->IsOfType<FUnloadedClassDragDropOp>())
 	{
-		TSharedPtr<FUnloadedClassDragDropOp> DragDropOp = StaticCastSharedPtr<FUnloadedClassDragDropOp>( DragDropEvent.GetOperation() );
+		TSharedPtr<FUnloadedClassDragDropOp> DragDropOp = StaticCastSharedPtr<FUnloadedClassDragDropOp>( Operation );
 
 		DroppedObjects.Empty();
 
@@ -568,12 +578,12 @@ bool SLevelViewport::HandlePlaceDraggedObjects(const FDragDropEvent& DragDropEve
 			}
 		}
 	}
-	else if ( DragDrop::IsTypeMatch<FAssetDragDropOp>(DragDropEvent.GetOperation()) )
+	else if (Operation->IsOfType<FAssetDragDropOp>())
 	{
 		bValidDrop = true;
 		DroppedObjects.Empty();
 
-		TSharedPtr<FAssetDragDropOp> DragDropOp = StaticCastSharedPtr<FAssetDragDropOp>( DragDropEvent.GetOperation() );
+		TSharedPtr<FAssetDragDropOp> DragDropOp = StaticCastSharedPtr<FAssetDragDropOp>( Operation );
 
 		ActorFactory = DragDropOp->ActorFactory.Get();
 
@@ -595,7 +605,7 @@ bool SLevelViewport::HandlePlaceDraggedObjects(const FDragDropEvent& DragDropEve
 	}
 	// OLE drops are blocking which causes problem when positioning and maintaining the drop preview
 	// Drop preview is disabled when dragging from external sources
-	else if ( !bCreateDropPreview && DragDrop::IsTypeMatch<FExternalDragOperation>( DragDropEvent.GetOperation() ) )
+	else if ( !bCreateDropPreview && Operation->IsOfType<FExternalDragOperation>() )
 	{
 		bValidDrop = true;
 		DroppedObjects.Empty();
@@ -618,11 +628,11 @@ bool SLevelViewport::HandlePlaceDraggedObjects(const FDragDropEvent& DragDropEve
 			}
 		}
 	}
-	else if ( DragDrop::IsTypeMatch<FExportTextDragDropOp>(DragDropEvent.GetOperation()) )
+	else if ( Operation->IsOfType<FExportTextDragDropOp>() )
 	{
 		bValidDrop = true;
 
-		TSharedPtr<FExportTextDragDropOp> DragDropOp = StaticCastSharedPtr<FExportTextDragDropOp>( DragDropEvent.GetOperation() );
+		TSharedPtr<FExportTextDragDropOp> DragDropOp = StaticCastSharedPtr<FExportTextDragDropOp>( Operation );
 
 		// Check if the asset is loaded, used to see if the context menu should be available
 		bAllAssetWereLoaded = true;
@@ -633,12 +643,12 @@ bool SLevelViewport::HandlePlaceDraggedObjects(const FDragDropEvent& DragDropEve
 		NewContainer->ExportText = DragDropOp->ActorExportText;
 		DroppedObjects.Add(NewContainer);
 	}
-	else if ( DragDrop::IsTypeMatch<FBrushBuilderDragDropOp>(DragDropEvent.GetOperation()) )
+	else if ( Operation->IsOfType<FBrushBuilderDragDropOp>() )
 	{
 		bValidDrop = true;
 		DroppedObjects.Empty();
 
-		TSharedPtr<FBrushBuilderDragDropOp> DragDropOp = StaticCastSharedPtr<FBrushBuilderDragDropOp>( DragDropEvent.GetOperation() );
+		TSharedPtr<FBrushBuilderDragDropOp> DragDropOp = StaticCastSharedPtr<FBrushBuilderDragDropOp>( Operation );
 
 		if(DragDropOp->GetBrushBuilder().IsValid())
 		{
diff --git a/Engine/Source/Editor/Levels/Private/SLevelBrowser.h b/Engine/Source/Editor/Levels/Private/SLevelBrowser.h
--- a/Engine/Source/Editor/Levels/Private/SLevelBrowser.h
+++ b/Engine/Source/Editor/Levels/Private/SLevelBrowser.h
@@ -160,13 +160,11 @@ protected:
 	 */
 	virtual void OnDragLeave( const FDragDropEvent& DragDropEvent ) OVERRIDE
 	{
-		if( !DragDrop::IsTypeMatch<FActorDragDropGraphEdOp>( DragDropEvent.GetOperation() ) )
+		TSharedPtr< FActorDragDropGraphEdOp > DragActorOp = DragDropEvent.GetOperationAs< FActorDragDropGraphEdOp >();
+		if (DragActorOp.IsValid())
 		{
-			return;
+			DragActorOp->ResetToDefaultToolTip();
 		}
-
-		TSharedPtr< FActorDragDropGraphEdOp > DragActorOp = StaticCastSharedPtr< FActorDragDropGraphEdOp >( DragDropEvent.GetOperation() );	
-		DragActorOp->ResetToDefaultToolTip();
 	}
 
 	/**
diff --git a/Engine/Source/Editor/Levels/Private/SLevelsView.h b/Engine/Source/Editor/Levels/Private/SLevelsView.h
--- a/Engine/Source/Editor/Levels/Private/SLevelsView.h
+++ b/Engine/Source/Editor/Levels/Private/SLevelsView.h
@@ -217,13 +217,11 @@ protected:
 	 */
 	virtual void OnDragLeave( const FDragDropEvent& DragDropEvent ) OVERRIDE
 	{
-		if( !DragDrop::IsTypeMatch<FActorDragDropGraphEdOp>( DragDropEvent.GetOperation() ) )
+		TSharedPtr< FActorDragDropGraphEdOp > DragActorOp = DragDropEvent.GetOperationAs< FActorDragDropGraphEdOp >();
+		if (DragActorOp.IsValid())
 		{
-			return;
+			DragActorOp->ResetToDefaultToolTip();
 		}
-
-		TSharedPtr< FActorDragDropGraphEdOp > DragActorOp = StaticCastSharedPtr< FActorDragDropGraphEdOp >( DragDropEvent.GetOperation() );	
-		DragActorOp->ResetToDefaultToolTip();
 	}
 
 	//@todo: add actor drag + drop support
@@ -237,12 +235,13 @@ protected:
 	 */
 	virtual FReply OnDragOver( const FGeometry& MyGeometry, const FDragDropEvent& DragDropEvent ) OVERRIDE
 	{
-		if( !DragDrop::IsTypeMatch<FActorDragDropGraphEdOp>( DragDropEvent.GetOperation() ) )
+
+		TSharedPtr< FActorDragDropGraphEdOp > DragActorOp = DragDropEvent.GetOperationAs< FActorDragDropGraphEdOp >();
+		if (!DragActorOp.IsValid())
 		{
 			return FReply::Unhandled();
 		}
 
-		TSharedPtr< FActorDragDropGraphEdOp > DragActorOp = StaticCastSharedPtr< FActorDragDropGraphEdOp >( DragDropEvent.GetOperation() );	
 		DragActorOp->SetToolTip( FActorDragDropGraphEdOp::ToolTip_CompatibleGeneric, LOCTEXT("OnDragOver", "Add Actors to New Level") );
 
 		// We leave the event unhandled so the children of the ListView get a chance to grab the drag/drop
@@ -260,14 +259,8 @@ protected:
 	 */
 	virtual FReply OnDrop( const FGeometry& MyGeometry, const FDragDropEvent& DragDropEvent ) OVERRIDE
 	{
-		if ( !DragDrop::IsTypeMatch<FActorDragDropGraphEdOp>( DragDropEvent.GetOperation() ) )	
-		{
-			return FReply::Unhandled();
-		}
-
-		TSharedPtr< FActorDragDropGraphEdOp > DragActorOp = StaticCastSharedPtr< FActorDragDropGraphEdOp >( DragDropEvent.GetOperation() );		
-
-		return FReply::Handled();
+		const bool bIsValidDrop = DragDropEvent.GetOperationAs<FActorDragDropGraphEdOp>().IsValid();
+		return bIsValidDrop ? FReply::Handled() : FReply::Unhandled();
 	}
 
 
diff --git a/Engine/Source/Editor/Levels/Private/SLevelsViewRow.h b/Engine/Source/Editor/Levels/Private/SLevelsViewRow.h
--- a/Engine/Source/Editor/Levels/Private/SLevelsViewRow.h
+++ b/Engine/Source/Editor/Levels/Private/SLevelsViewRow.h
@@ -380,13 +380,11 @@ protected:
 	 */
 	virtual void OnDragLeave( const FDragDropEvent& DragDropEvent ) OVERRIDE
 	{
-		if( !DragDrop::IsTypeMatch<FActorDragDropGraphEdOp>( DragDropEvent.GetOperation() ) )
+		TSharedPtr< FActorDragDropGraphEdOp > DragActorOp = DragDropEvent.GetOperationAs< FActorDragDropGraphEdOp >();
+		if (DragActorOp.IsValid())
 		{
-			return;
+			DragActorOp->ResetToDefaultToolTip();
 		}
-
-		TSharedPtr< FActorDragDropGraphEdOp > DragActorOp = StaticCastSharedPtr< FActorDragDropGraphEdOp >( DragDropEvent.GetOperation() );	
-		DragActorOp->ResetToDefaultToolTip();
 	}
 
 	/**
@@ -399,13 +397,12 @@ protected:
 	 */
 	virtual FReply OnDragOver( const FGeometry& MyGeometry, const FDragDropEvent& DragDropEvent ) OVERRIDE
 	{
-		if( !DragDrop::IsTypeMatch<FActorDragDropGraphEdOp>( DragDropEvent.GetOperation() ) )
+		TSharedPtr< FActorDragDropGraphEdOp > DragActorOp = DragDropEvent.GetOperationAs< FActorDragDropGraphEdOp >();
+		if (!DragActorOp.IsValid())
 		{
 			return FReply::Unhandled();
 		}
 
-		TSharedPtr< FActorDragDropGraphEdOp > DragActorOp = StaticCastSharedPtr< FActorDragDropGraphEdOp >( DragDropEvent.GetOperation() );		
-
 		bool bCanAssign = false;
 		FString Message;
 		if( DragActorOp->Actors.Num() > 1 )
@@ -439,12 +436,12 @@ protected:
 	 */
 	virtual FReply OnDrop( const FGeometry& MyGeometry, const FDragDropEvent& DragDropEvent ) OVERRIDE
 	{
-		if ( !DragDrop::IsTypeMatch<FActorDragDropGraphEdOp>( DragDropEvent.GetOperation() ) )	
+		TSharedPtr< FActorDragDropGraphEdOp > DragActorOp = DragDropEvent.GetOperationAs< FActorDragDropGraphEdOp >();
+		if (!DragActorOp.IsValid())
 		{
 			return FReply::Unhandled();
 		}
 
-		TSharedPtr< FActorDragDropGraphEdOp > DragActorOp = StaticCastSharedPtr< FActorDragDropGraphEdOp >( DragDropEvent.GetOperation() );		
 		ViewModel->AddActors( DragActorOp->Actors );
 
 		return FReply::Handled();
@@ -512,7 +509,7 @@ private:
 		bool bCanAcceptDrop = false;
 		TSharedPtr<FDragDropOperation> DragDropOp = FSlateApplication::Get().GetDragDroppingContent();
 
-		if (DragDrop::IsTypeMatch<FActorDragDropGraphEdOp>(DragDropOp))	
+		if (DragDropOp.IsValid() && DragDropOp->IsOfType<FActorDragDropGraphEdOp>())
 		{
 			TSharedPtr<FActorDragDropGraphEdOp> DragDropActorOp = StaticCastSharedPtr<FActorDragDropGraphEdOp>(DragDropOp);
 
diff --git a/Engine/Source/Editor/Persona/Private/SAnimNotifyPanel.cpp b/Engine/Source/Editor/Persona/Private/SAnimNotifyPanel.cpp
--- a/Engine/Source/Editor/Persona/Private/SAnimNotifyPanel.cpp
+++ b/Engine/Source/Editor/Persona/Private/SAnimNotifyPanel.cpp
@@ -2247,14 +2247,13 @@ float SAnimNotifyTrack::CalculateTime( const FGeometry& MyGeometry, FVector2D No
 
 FReply SAnimNotifyTrack::OnDrop( const FGeometry& MyGeometry, const FDragDropEvent& DragDropEvent )
 {
-	if ( DragDrop::IsTypeMatch<FNotifyDragDropOp>( DragDropEvent.GetOperation() ) )
+	TSharedPtr<FNotifyDragDropOp> DragDropOp = DragDropEvent.GetOperationAs<FNotifyDragDropOp>();
+	if (DragDropOp.IsValid())
 	{
-		TSharedPtr<FNotifyDragDropOp> DragDropOp = StaticCastSharedPtr<FNotifyDragDropOp>(DragDropEvent.GetOperation());
-
 		HandleNodeDrop(*DragDropOp);
 
 		return FReply::Handled();
-}
+	}
 
 	return FReply::Unhandled();
 }
diff --git a/Engine/Source/Editor/Persona/Private/SAnimSegmentsPanel.cpp b/Engine/Source/Editor/Persona/Private/SAnimSegmentsPanel.cpp
--- a/Engine/Source/Editor/Persona/Private/SAnimSegmentsPanel.cpp
+++ b/Engine/Source/Editor/Persona/Private/SAnimSegmentsPanel.cpp
@@ -220,7 +220,7 @@ void SAnimSegmentsPanel::RemoveAnimSegment(int32 AnimSegmentIndex)
 
 void SAnimSegmentsPanel::OnTrackDragDrop( TSharedPtr<FDragDropOperation> DragDropOp, float DataPos )
 {
-	if ( DragDrop::IsTypeMatch<FAssetDragDropOp>(DragDropOp) )
+	if (DragDropOp.IsValid() && DragDropOp->IsOfType<FAssetDragDropOp>())
 	{
 		TSharedPtr<FAssetDragDropOp> AssetOp = StaticCastSharedPtr<FAssetDragDropOp>(DragDropOp);
 		UAnimSequence *DroppedSequence = FAssetData::GetFirstAsset<UAnimSequence>(AssetOp->AssetData);
diff --git a/Engine/Source/Editor/Persona/Private/SAnimationBlendSpaceBase.cpp b/Engine/Source/Editor/Persona/Private/SAnimationBlendSpaceBase.cpp
--- a/Engine/Source/Editor/Persona/Private/SAnimationBlendSpaceBase.cpp
+++ b/Engine/Source/Editor/Persona/Private/SAnimationBlendSpaceBase.cpp
@@ -692,9 +692,15 @@ FReply SBlendSpaceWidget::OnDrop( const FGeometry& MyGeometry, const FDragDropEv
 {
 	FSlateRect WindowRect = GetWindowRectFromGeometry(MyGeometry);
 
-	if ( DragDrop::IsTypeMatch<FAssetDragDropOp>( DragDropEvent.GetOperation() ) )
+	TSharedPtr< FDragDropOperation > Operation = DragDropEvent.GetOperation();
+	if (!Operation.IsValid())
 	{
-		TSharedPtr<FAssetDragDropOp> AssetOp = StaticCastSharedPtr<FAssetDragDropOp>(DragDropEvent.GetOperation());
+		return FReply::Unhandled();
+	}
+
+	if ( Operation->IsOfType<FAssetDragDropOp>() )
+	{
+		TSharedPtr<FAssetDragDropOp> AssetOp = StaticCastSharedPtr<FAssetDragDropOp>(Operation);
 		UAnimSequence *DroppedSequence = FAssetData::GetFirstAsset<UAnimSequence>(AssetOp->AssetData);
 
 		if (DroppedSequence)
@@ -703,9 +709,9 @@ FReply SBlendSpaceWidget::OnDrop( const FGeometry& MyGeometry, const FDragDropEv
 			return FReply::Handled();
 		}
 	}
-	else if ( DragDrop::IsTypeMatch<FSampleDragDropOp>( DragDropEvent.GetOperation() ) )
+	else if (Operation->IsOfType<FSampleDragDropOp>())
 	{
-		TSharedPtr<FSampleDragDropOp> DragDropOp = StaticCastSharedPtr<FSampleDragDropOp>(DragDropEvent.GetOperation());
+		TSharedPtr<FSampleDragDropOp> DragDropOp = StaticCastSharedPtr<FSampleDragDropOp>(Operation);
 		UAnimSequence *DroppedSequence = DragDropOp->Sample.Animation;
 		int32 OriginalIndex = DragDropOp->OriginalSampleIndex;
 
@@ -783,7 +789,7 @@ FReply SBlendSpaceWidget::OnMouseMove( const FGeometry& MyGeometry, const FPoint
 FReply SBlendSpaceWidget::OnDragOver( const FGeometry& MyGeometry, const FDragDropEvent& DragDropEvent )
 {
 	// if being dragged, update the mouse pos
-	if ( DragDrop::IsTypeMatch<FAssetDragDropOp>( DragDropEvent.GetOperation() ) || DragDrop::IsTypeMatch<FSampleDragDropOp>( DragDropEvent.GetOperation() ) )
+	if (DragDropEvent.GetOperationAs<FAssetDragDropOp>().IsValid() || DragDropEvent.GetOperationAs<FSampleDragDropOp>().IsValid())
 	{
 		bBeingDragged = true;
 		return UpdateLastMousePosition(MyGeometry, DragDropEvent.GetScreenSpacePosition());
diff --git a/Engine/Source/Editor/Persona/Private/SSkeletonTree.cpp b/Engine/Source/Editor/Persona/Private/SSkeletonTree.cpp
--- a/Engine/Source/Editor/Persona/Private/SSkeletonTree.cpp
+++ b/Engine/Source/Editor/Persona/Private/SSkeletonTree.cpp
@@ -180,11 +180,11 @@ TSharedRef< SWidget > SSkeletonTreeRow::GenerateWidgetForColumn( const FName& Co
 
 void SSkeletonTreeRow::OnDragEnter( const FGeometry& MyGeometry, const FDragDropEvent& DragDropEvent )
 {
+	TSharedPtr<FSocketDragDropOp> DragConnectionOp = DragDropEvent.GetOperationAs<FSocketDragDropOp>();
+
 	// Is someone dragging a socket onto a bone?
-	if ( DragDrop::IsTypeMatch<FSocketDragDropOp>( DragDropEvent.GetOperation() ) )
+	if (DragConnectionOp.IsValid())
 	{
-		TSharedPtr<FSocketDragDropOp> DragConnectionOp = StaticCastSharedPtr<FSocketDragDropOp>( DragDropEvent.GetOperation() );
-
 		if ( Item->GetType() == ESkeletonTreeRowType::Bone &&
 			( *static_cast<FName*>( Item->GetData() ) != DragConnectionOp->GetSocketInfo().Socket->BoneName ) )
 		{
@@ -205,20 +205,19 @@ void SSkeletonTreeRow::OnDragEnter( const FGeometry& MyGeometry, const FDragDrop
 
 void SSkeletonTreeRow::OnDragLeave( const FDragDropEvent& DragDropEvent )
 {
-	if ( DragDrop::IsTypeMatch<FSocketDragDropOp>( DragDropEvent.GetOperation() ) )
+	TSharedPtr<FSocketDragDropOp> DragConnectionOp = DragDropEvent.GetOperationAs<FSocketDragDropOp>();
+	if (DragConnectionOp.IsValid())
 	{
 		// Reset the drag/drop icon when leaving this row
-		TSharedPtr<FSocketDragDropOp> DragConnectionOp = StaticCastSharedPtr<FSocketDragDropOp>( DragDropEvent.GetOperation() );
 		DragConnectionOp->SetIcon( FEditorStyle::GetBrush( TEXT( "Graph.ConnectorFeedback.Error" ) ) );
 	}
 }
 
 FReply SSkeletonTreeRow::OnDrop( const FGeometry& MyGeometry, const FDragDropEvent& DragDropEvent )
 {
-	if ( DragDrop::IsTypeMatch<FSocketDragDropOp>( DragDropEvent.GetOperation() ) )
+	TSharedPtr<FSocketDragDropOp> DragConnectionOp = DragDropEvent.GetOperationAs<FSocketDragDropOp>();
+	if (DragConnectionOp.IsValid())
 	{
-		TSharedPtr<FSocketDragDropOp> DragConnectionOp = StaticCastSharedPtr<FSocketDragDropOp>( DragDropEvent.GetOperation() );
-
 		FSelectedSocketInfo SocketInfo = DragConnectionOp->GetSocketInfo();
 
 		if ( DragConnectionOp->IsAltDrag() && Item->GetType() == ESkeletonTreeRowType::Bone )
@@ -248,9 +247,12 @@ FReply SSkeletonTreeRow::OnDrop( const FGeometry& MyGeometry, const FDragDropEve
 			return FReply::Handled();
 		}
 	}
-	else if( DragDrop::IsTypeMatch<FAssetDragDropOp>( DragDropEvent.GetOperation() ) )
+	else
 	{
-		SkeletonTree.Pin()->OnDropAssetToSkeletonTree( Item, DragDropEvent );
+		if (DragDropEvent.GetOperationAs<FAssetDragDropOp>().IsValid())
+		{
+			SkeletonTree.Pin()->OnDropAssetToSkeletonTree( Item, DragDropEvent );
+		}
 	}
 
 	return FReply::Unhandled();
@@ -1969,10 +1971,9 @@ void SSkeletonTree::OnSelectionChanged(TSharedPtr<FDisplayedTreeRowInfo> Selecti
 
 FReply SSkeletonTree::OnDropAssetToSkeletonTree(const FDisplayedTreeRowInfoPtr TargetItem, const FDragDropEvent& DragDropEvent)
 {
-	if( DragDrop::IsTypeMatch<FAssetDragDropOp>(DragDropEvent.GetOperation()) )
+	TSharedPtr<FAssetDragDropOp> DragDropOp = DragDropEvent.GetOperationAs<FAssetDragDropOp>();
+	if (DragDropOp.IsValid())
 	{
-		TSharedPtr<FAssetDragDropOp> DragDropOp = StaticCastSharedPtr<FAssetDragDropOp>( DragDropEvent.GetOperation() );
-		
 		//Do we have some assets to attach?
 		if(DragDropOp->AssetData.Num() > 0)
 		{
diff --git a/Engine/Source/Editor/Persona/Private/STrack.cpp b/Engine/Source/Editor/Persona/Private/STrack.cpp
--- a/Engine/Source/Editor/Persona/Private/STrack.cpp
+++ b/Engine/Source/Editor/Persona/Private/STrack.cpp
@@ -516,9 +516,9 @@ FReply STrack::OnDrop( const FGeometry& MyGeometry, const FDragDropEvent& DragDr
 	float CusorDataPos = LocalToDataX(CursorPos.X, MyGeometry);
 
 	// Handle TrackNodes that were dropped
-	if ( DragDrop::IsTypeMatch<FTrackNodeDragDropOp>( DragDropEvent.GetOperation() ) )
+	TSharedPtr<FTrackNodeDragDropOp> DragDropOp = DragDropEvent.GetOperationAs<FTrackNodeDragDropOp>();
+	if (DragDropOp.IsValid())
 	{
-		TSharedPtr<FTrackNodeDragDropOp> DragDropOp = StaticCastSharedPtr<FTrackNodeDragDropOp>(DragDropEvent.GetOperation());
 		TSharedPtr<STrackNode> TrackNode = DragDropOp->OriginalTrackNode.Pin();
 
 		float DataPos = GetNodeDragDropDataPos(MyGeometry, DragDropEvent);
@@ -535,9 +535,9 @@ FReply STrack::OnDrop( const FGeometry& MyGeometry, const FDragDropEvent& DragDr
 FReply STrack::OnDragOver( const FGeometry& MyGeometry, const FDragDropEvent& DragDropEvent )
 {
 	// Handle TrackNodes that were dropped
-	if ( DragDrop::IsTypeMatch<FTrackNodeDragDropOp>( DragDropEvent.GetOperation() ) )
+	TSharedPtr<FTrackNodeDragDropOp> DragDropOp = DragDropEvent.GetOperationAs<FTrackNodeDragDropOp>();
+	if (DragDropOp.IsValid())
 	{
-		TSharedPtr<FTrackNodeDragDropOp> DragDropOp = StaticCastSharedPtr<FTrackNodeDragDropOp>(DragDropEvent.GetOperation());
 		TSharedPtr<STrackNode> TrackNode = DragDropOp->OriginalTrackNode.Pin();
 
 		float DataPos = GetNodeDragDropDataPos(MyGeometry, DragDropEvent);
diff --git a/Engine/Source/Editor/PropertyEditor/Private/UserInterface/PropertyEditor/SPropertyEditorClass.cpp b/Engine/Source/Editor/PropertyEditor/Private/UserInterface/PropertyEditor/SPropertyEditorClass.cpp
--- a/Engine/Source/Editor/PropertyEditor/Private/UserInterface/PropertyEditor/SPropertyEditorClass.cpp
+++ b/Engine/Source/Editor/PropertyEditor/Private/UserInterface/PropertyEditor/SPropertyEditorClass.cpp
@@ -241,12 +241,11 @@ void SPropertyEditorClass::SendToObjects(const FString& NewValue)
 
 FReply SPropertyEditorClass::OnDrop(const FGeometry& MyGeometry, const FDragDropEvent& DragDropEvent)
 {
-	if(DragDrop::IsTypeMatch<FClassDragDropOp>(DragDropEvent.GetOperation()))
+	TSharedPtr<FClassDragDropOp> ClassOperation = DragDropEvent.GetOperationAs<FClassDragDropOp>();
+	if (ClassOperation.IsValid())
 	{
-		auto Operation = StaticCastSharedPtr<FClassDragDropOp>(DragDropEvent.GetOperation());
-
 		// We can only drop one item into the combo box, so drop the first one.
-		FString AssetName = Operation->ClassesToDrop[0]->GetName();
+		FString AssetName = ClassOperation->ClassesToDrop[0]->GetName();
 
 		// Set the property, it will be verified as valid.
 		SendToObjects(AssetName);
@@ -254,14 +253,13 @@ FReply SPropertyEditorClass::OnDrop(const FGeometry& MyGeometry, const FDragDrop
 		return FReply::Handled();
 	}
 	
-	if(DragDrop::IsTypeMatch<FUnloadedClassDragDropOp>(DragDropEvent.GetOperation()))
+	TSharedPtr<FUnloadedClassDragDropOp> UnloadedClassOp = DragDropEvent.GetOperationAs<FUnloadedClassDragDropOp>();
+	if (UnloadedClassOp.IsValid())
 	{
-		TSharedPtr<FUnloadedClassDragDropOp> DragDropOp = StaticCastSharedPtr<FUnloadedClassDragDropOp>(DragDropEvent.GetOperation());
-
 		// Check if the asset is loaded, used to see if the context menu should be available
 		bool bAllAssetWereLoaded = true;
 
-		TArray<FClassPackageData>& AssetArray = *(DragDropOp->AssetsToDrop.Get());
+		TArray<FClassPackageData>& AssetArray = *(UnloadedClassOp->AssetsToDrop.Get());
 
 		// We can only drop one item into the combo box, so drop the first one.
 		FString& AssetName = AssetArray[0].AssetName;
diff --git a/Engine/Source/Editor/SceneOutliner/Private/SOutlinerTreeView.cpp b/Engine/Source/Editor/SceneOutliner/Private/SOutlinerTreeView.cpp
--- a/Engine/Source/Editor/SceneOutliner/Private/SOutlinerTreeView.cpp
+++ b/Engine/Source/Editor/SceneOutliner/Private/SOutlinerTreeView.cpp
@@ -61,16 +61,16 @@ namespace SceneOutliner
 			return;
 		}
 
-		if (DragDrop::IsTypeMatch<FFolderActorDragDropOp>(Operation))
+		if (Operation->IsOfType<FFolderActorDragDropOp>())
 		{
 			DraggedFolders = &StaticCastSharedPtr<FFolderActorDragDropOp>(Operation)->Folders;
 		}
-		else if (DragDrop::IsTypeMatch<FFolderDragDropOp>(Operation))
+		else if (Operation->IsOfType<FFolderDragDropOp>())
 		{
 			DraggedFolders = &StaticCastSharedPtr<FFolderDragDropOp>(Operation)->Folders;
 		}
 
-		if (DragDrop::IsTypeMatch<FActorDragDropGraphEdOp>(Operation))
+		if (Operation->IsOfType<FActorDragDropGraphEdOp>())
 		{
 			DraggedActors = &StaticCastSharedPtr<FActorDragDropGraphEdOp>(Operation)->Actors;
 		}
@@ -433,14 +433,14 @@ namespace SceneOutliner
 		}
 
 
-		if (DragDrop::IsTypeMatch<FDecoratedDragDropOp>(DragDropEvent.GetOperation()))
+		TSharedPtr<FDecoratedDragDropOp> DecoratedDragOp = DragDropEvent.GetOperationAs<FDecoratedDragDropOp>();
+		if (DecoratedDragOp.IsValid())
 		{
 			const FDragValidationInfo ValidationInfo = ValidateRootDragDropEvent(SceneOutlinerPtr.ToSharedRef(), DragDropEvent);
 
 			if (ValidationInfo.TooltipType == FActorDragDropGraphEdOp::ToolTip_CompatibleGeneric || ValidationInfo.TooltipType == FActorDragDropGraphEdOp::ToolTip_IncompatibleGeneric)
 			{
 				// Generic messages can go through the decorated drag/drop operation base class
-				auto DecoratedDragOp = StaticCastSharedPtr<FDecoratedDragDropOp>(DragDropEvent.GetOperation());
 				DecoratedDragOp->CurrentHoverText = ValidationInfo.ValidationText;
 				if (ValidationInfo.TooltipType == FActorDragDropGraphEdOp::ToolTip_CompatibleGeneric)
 				{
@@ -451,10 +451,10 @@ namespace SceneOutliner
 					DecoratedDragOp->CurrentIconBrush = FEditorStyle::GetBrush(TEXT("Graph.ConnectorFeedback.Error"));
 				}
 			}
-			else if (DragDrop::IsTypeMatch<FActorDragDropGraphEdOp>(DragDropEvent.GetOperation()))
+			else if (DecoratedDragOp->IsOfType<FActorDragDropGraphEdOp>())
 			{
 				// Display actor specific messages
-				auto ActorDragOp = StaticCastSharedPtr<FActorDragDropGraphEdOp>(DragDropEvent.GetOperation());
+				auto ActorDragOp = StaticCastSharedPtr<FActorDragDropGraphEdOp>(DecoratedDragOp);
 				ActorDragOp->SetToolTip(ValidationInfo.TooltipType, ValidationInfo.ValidationText);
 			}
 		}
@@ -651,14 +651,16 @@ namespace SceneOutliner
 		// Perform the validation and update the tooltip if necessary
 		const FDragValidationInfo ValidationInfo = ValidateDragDropEvent(SceneOutlinerPtr.ToSharedRef(), DragDropEvent, Item.ToSharedRef());
 
-		if (DragDrop::IsTypeMatch<FActorDragDropGraphEdOp>(DragDropEvent.GetOperation()))
+		TSharedPtr<FActorDragDropGraphEdOp> ActorDragOp = DragDropEvent.GetOperationAs<FActorDragDropGraphEdOp>();
+		if (ActorDragOp.IsValid())
 		{
-			auto ActorDragOp = StaticCastSharedPtr<FActorDragDropGraphEdOp>(DragDropEvent.GetOperation());
 			ActorDragOp->SetToolTip(ValidationInfo.TooltipType, ValidationInfo.ValidationText);
+			return;
 		}
-		else if (DragDrop::IsTypeMatch<FDecoratedDragDropOp>(DragDropEvent.GetOperation()))
+
+		TSharedPtr<FDecoratedDragDropOp> DecoratedDragOp = DragDropEvent.GetOperationAs<FDecoratedDragDropOp>();
+		if (DecoratedDragOp.IsValid())
 		{
-			auto DecoratedDragOp = StaticCastSharedPtr<FDecoratedDragDropOp>(DragDropEvent.GetOperation());
 			DecoratedDragOp->CurrentHoverText = ValidationInfo.ValidationText;
 			DecoratedDragOp->CurrentIconBrush = ValidationInfo.IsValid() ? FEditorStyle::GetBrush(TEXT("Graph.ConnectorFeedback.OK")) : FEditorStyle::GetBrush(TEXT("Graph.ConnectorFeedback.Error"));
 		}
diff --git a/Engine/Source/Editor/SceneOutliner/Private/SSceneOutliner.cpp b/Engine/Source/Editor/SceneOutliner/Private/SSceneOutliner.cpp
--- a/Engine/Source/Editor/SceneOutliner/Private/SSceneOutliner.cpp
+++ b/Engine/Source/Editor/SceneOutliner/Private/SSceneOutliner.cpp
@@ -1161,7 +1161,7 @@ namespace SceneOutliner
 			{
 				TSharedPtr<FDragDropOperation> DragDropOp = FSlateApplication::Get().GetDragDroppingContent();
 
-				if (DragDrop::IsTypeMatch<FActorDragDropGraphEdOp>(DragDropOp))	
+				if (DragDropOp->IsOfType<FActorDragDropGraphEdOp>())
 				{
 					TSharedPtr<FActorDragDropGraphEdOp> DragDropActorOp = StaticCastSharedPtr<FActorDragDropGraphEdOp>(DragDropOp);
 
diff --git a/Engine/Source/Editor/Sequencer/Private/SSequencer.cpp b/Engine/Source/Editor/Sequencer/Private/SSequencer.cpp
--- a/Engine/Source/Editor/Sequencer/Private/SSequencer.cpp
+++ b/Engine/Source/Editor/Sequencer/Private/SSequencer.cpp
@@ -428,10 +428,12 @@ FReply SSequencer::OnDragOver( const FGeometry& MyGeometry, const FDragDropEvent
 {
 	bool bIsDragSupported = false;
 
-	if( DragDrop::IsTypeMatch<FAssetDragDropOp>( DragDropEvent.GetOperation() ) ||
-		DragDrop::IsTypeMatch<FClassDragDropOp>( DragDropEvent.GetOperation() ) ||
-		DragDrop::IsTypeMatch<FUnloadedClassDragDropOp>( DragDropEvent.GetOperation() ) ||
-		DragDrop::IsTypeMatch<FActorDragDropGraphEdOp>( DragDropEvent.GetOperation() ) )
+	TSharedPtr<FDragDropOperation> Operation = DragDropEvent.GetOperation();
+	if (Operation.IsValid() && (
+		Operation->IsOfType<FAssetDragDropOp>() ||
+		Operation->IsOfType<FClassDragDropOp>() ||
+		Operation->IsOfType<FUnloadedClassDragDropOp>() ||
+		Operation->IsOfType<FActorDragDropGraphEdOp>() ) )
 	{
 		bIsDragSupported = true;
 	}
@@ -455,33 +457,38 @@ FReply SSequencer::OnDrop( const FGeometry& MyGeometry, const FDragDropEvent& Dr
 	//		- Basically, when Sequencer is open it would take over drops into the level and auto-add puppets for these instead of regular actors
 	//		- This would let people drag smoothly and precisely into the view to drop assets/classes into the scene
 
-	if( DragDrop::IsTypeMatch<FAssetDragDropOp>( DragDropEvent.GetOperation() ) )
+	TSharedPtr<FDragDropOperation> Operation = DragDropEvent.GetOperation();
+	if ( !Operation.IsValid() )
 	{
-		const auto& DragDropOp = StaticCastSharedPtr<FAssetDragDropOp>( DragDropEvent.GetOperation() );
+
+	}
+	else if ( Operation->IsOfType<FAssetDragDropOp>() )
+	{
+		const auto& DragDropOp = StaticCastSharedPtr<FAssetDragDropOp>( Operation );
 
 		OnAssetsDropped( *DragDropOp );
 
 		bWasDropHandled = true;
 	}
-	else if( DragDrop::IsTypeMatch<FClassDragDropOp>( DragDropEvent.GetOperation() ) )
+	else if( Operation->IsOfType<FClassDragDropOp>() )
 	{
-		const auto& DragDropOp = StaticCastSharedPtr<FClassDragDropOp>( DragDropEvent.GetOperation() );
+		const auto& DragDropOp = StaticCastSharedPtr<FClassDragDropOp>( Operation );
 
 		OnClassesDropped( *DragDropOp );
 
 		bWasDropHandled = true;
 	}
-	else if( DragDrop::IsTypeMatch<FUnloadedClassDragDropOp>( DragDropEvent.GetOperation() ) )
+	else if( Operation->IsOfType<FUnloadedClassDragDropOp>() )
 	{
-		const auto& DragDropOp = StaticCastSharedPtr<FUnloadedClassDragDropOp>( DragDropEvent.GetOperation() );
+		const auto& DragDropOp = StaticCastSharedPtr<FUnloadedClassDragDropOp>( Operation );
 
 		OnUnloadedClassesDropped( *DragDropOp );
 
 		bWasDropHandled = true;
 	}
-	else if( DragDrop::IsTypeMatch<FActorDragDropGraphEdOp>( DragDropEvent.GetOperation() ) )
+	else if( Operation->IsOfType<FActorDragDropGraphEdOp>() )
 	{
-		const auto& DragDropOp = StaticCastSharedPtr<FActorDragDropGraphEdOp>( DragDropEvent.GetOperation() );
+		const auto& DragDropOp = StaticCastSharedPtr<FActorDragDropGraphEdOp>( Operation );
 
 		OnActorsDropped( *DragDropOp );
 
diff --git a/Engine/Source/Editor/SourceControlWindows/Private/SSourceControlHistory.cpp b/Engine/Source/Editor/SourceControlWindows/Private/SSourceControlHistory.cpp
--- a/Engine/Source/Editor/SourceControlWindows/Private/SSourceControlHistory.cpp
+++ b/Engine/Source/Editor/SourceControlWindows/Private/SSourceControlHistory.cpp
@@ -1189,9 +1189,9 @@ private:
 	 */
 	void OnRowDragEnter(FDragDropEvent const& DragDropEvent, TSharedPtr<FHistoryTreeItem> const HoveredItem) const
 	{
-		if (DragDrop::IsTypeMatch<FSourceControlHistoryRowDragDropOp>(DragDropEvent.GetOperation()))
+		TSharedPtr<FSourceControlHistoryRowDragDropOp> DragRowOp = DragDropEvent.GetOperationAs<FSourceControlHistoryRowDragDropOp>();
+		if (DragRowOp.IsValid())
 		{
-			TSharedPtr<FSourceControlHistoryRowDragDropOp> DragRowOp = StaticCastSharedPtr<FSourceControlHistoryRowDragDropOp>(DragDropEvent.GetOperation());
 			DragRowOp->PendingDropAction = FSourceControlHistoryRowDragDropOp::EDropAction::None;
 
 			TArray< TSharedPtr<FHistoryTreeItem> > DiffingItems = DragRowOp->SelectedItems;
@@ -1252,10 +1252,9 @@ private:
 	 */
 	void OnRowDragLeave(FDragDropEvent const& DragDropEvent) const
 	{
-		if (DragDrop::IsTypeMatch<FSourceControlHistoryRowDragDropOp>(DragDropEvent.GetOperation()))
+		TSharedPtr<FSourceControlHistoryRowDragDropOp> DragRowOp = DragDropEvent.GetOperationAs<FSourceControlHistoryRowDragDropOp>();
+		if (DragRowOp.IsValid())
 		{
-			TSharedPtr<FSourceControlHistoryRowDragDropOp> DragRowOp = StaticCastSharedPtr<FSourceControlHistoryRowDragDropOp>(DragDropEvent.GetOperation());
-
 			DragRowOp->HoverText = FText::GetEmpty();
 			DragRowOp->PendingDropAction = FSourceControlHistoryRowDragDropOp::EDropAction::None;
 		}
@@ -1272,9 +1271,9 @@ private:
 	 */
 	FReply OnRowDrop(FDragDropEvent const& DragDropEvent, TSharedPtr<FHistoryTreeItem> const HoveredItem) const
 	{
-		if (DragDrop::IsTypeMatch<FSourceControlHistoryRowDragDropOp>(DragDropEvent.GetOperation()))
+		TSharedPtr<FSourceControlHistoryRowDragDropOp> DragRowOp = DragDropEvent.GetOperationAs<FSourceControlHistoryRowDragDropOp>();
+		if (DragRowOp.IsValid())
 		{
-			TSharedPtr<FSourceControlHistoryRowDragDropOp> DragRowOp = StaticCastSharedPtr<FSourceControlHistoryRowDragDropOp>(DragDropEvent.GetOperation());
 			if (DragRowOp->PendingDropAction == FSourceControlHistoryRowDragDropOp::EDropAction::Diff)
 			{
 				check(DragRowOp->SelectedItems.Num() > 0);
diff --git a/Engine/Source/Editor/UnrealEd/Private/AssetSelection.cpp b/Engine/Source/Editor/UnrealEd/Private/AssetSelection.cpp
--- a/Engine/Source/Editor/UnrealEd/Private/AssetSelection.cpp
+++ b/Engine/Source/Editor/UnrealEd/Private/AssetSelection.cpp
@@ -432,9 +432,16 @@ namespace AssetUtil
 	TArray<FAssetData> ExtractAssetDataFromDrag( const FDragDropEvent &DragDropEvent )
 	{
 		TArray<FAssetData> DroppedAssetData;
-		if ( DragDrop::IsTypeMatch<FExternalDragOperation>(DragDropEvent.GetOperation()) )
+
+		TSharedPtr<FDragDropOperation> Operation = DragDropEvent.GetOperation();
+		if (!Operation.IsValid())
+		{
+			return DroppedAssetData;
+		}
+
+		if (Operation->IsOfType<FExternalDragOperation>())
 		{
-			TSharedPtr<FExternalDragOperation> DragDropOp = StaticCastSharedPtr<FExternalDragOperation>( DragDropEvent.GetOperation() );
+			TSharedPtr<FExternalDragOperation> DragDropOp = StaticCastSharedPtr<FExternalDragOperation>(Operation);
 			if ( DragDropOp->HasText() )
 			{
 				TArray<FString> DroppedAssetStrings;
@@ -454,9 +461,9 @@ namespace AssetUtil
 				}
 			}
 		}
-		else if ( DragDrop::IsTypeMatch<FAssetDragDropOp>(DragDropEvent.GetOperation()) )
+		else if (Operation->IsOfType<FAssetDragDropOp>())
 		{
-			TSharedPtr<FAssetDragDropOp> DragDropOp = StaticCastSharedPtr<FAssetDragDropOp>( DragDropEvent.GetOperation() );
+			TSharedPtr<FAssetDragDropOp> DragDropOp = StaticCastSharedPtr<FAssetDragDropOp>( Operation );
 			DroppedAssetData.Append( DragDropOp->AssetData );
 		}
 
diff --git a/Engine/Source/Editor/UnrealEd/Private/LevelEditorViewport.cpp b/Engine/Source/Editor/UnrealEd/Private/LevelEditorViewport.cpp
--- a/Engine/Source/Editor/UnrealEd/Private/LevelEditorViewport.cpp
+++ b/Engine/Source/Editor/UnrealEd/Private/LevelEditorViewport.cpp
@@ -4153,7 +4153,7 @@ void FLevelEditorViewportClient::SetupViewForRendering( FSceneViewFamily& ViewFa
 	}
 
 	TSharedPtr<FDragDropOperation> DragOperation = FSlateApplication::Get().GetDragDroppingContent();
-	if( !(DragOperation.IsValid() && DragDrop::IsTypeMatch<FBrushBuilderDragDropOp>(DragOperation)) )
+	if (!(DragOperation.IsValid() && DragOperation->IsOfType<FBrushBuilderDragDropOp>()))
 	{
 		// Hide the builder brush when not in geometry mode
 		ViewFamily.EngineShowFlags.BuilderBrush = 0;
diff --git a/Engine/Source/Programs/SymbolDebugger/Private/SSymbolDebugger.cpp b/Engine/Source/Programs/SymbolDebugger/Private/SSymbolDebugger.cpp
--- a/Engine/Source/Programs/SymbolDebugger/Private/SSymbolDebugger.cpp
+++ b/Engine/Source/Programs/SymbolDebugger/Private/SSymbolDebugger.cpp
@@ -436,10 +436,9 @@ FReply SSymbolDebugger::OnDragOver(const FGeometry& MyGeometry, const FDragDropE
 {
 	if (GetSymbolDebuggerAction() == DebugAction_None)
 	{
-		if (DragDrop::IsTypeMatch<FExternalDragOperation>(DragDropEvent.GetOperation()))
+		TSharedPtr<FExternalDragOperation> DragDropOp = DragDropEvent.GetOperationAs<FExternalDragOperation>();
+		if (DragDropOp.IsValid())
 		{
-			TSharedPtr<FExternalDragOperation> DragDropOp = StaticCastSharedPtr<FExternalDragOperation>(DragDropEvent.GetOperation());
-
 			if (DragDropOp->HasFiles())
 			{
 				if (DragDropOp->GetFiles().Num() == 1)
@@ -458,9 +457,9 @@ FReply SSymbolDebugger::OnDrop(const FGeometry& MyGeometry, const FDragDropEvent
 	if (GetSymbolDebuggerAction() == DebugAction_None)
 	{
 		// Handle drag drop for import
-		if (DragDrop::IsTypeMatch<FExternalDragOperation>(DragDropEvent.GetOperation()))
+		TSharedPtr<FExternalDragOperation> DragDropOp = DragDropEvent.GetOperationAs<FExternalDragOperation>();
+		if (DragDropOp.IsValid())
 		{
-			TSharedPtr<FExternalDragOperation> DragDropOp = StaticCastSharedPtr<FExternalDragOperation>(DragDropEvent.GetOperation());
 			if (DragDropOp->HasFiles())
 			{
 				// Set the current method to CrashDump as that is the only thing we support dropping for
diff --git a/Engine/Source/Runtime/Slate/Private/DockingFramework/SDockingArea.cpp b/Engine/Source/Runtime/Slate/Private/DockingFramework/SDockingArea.cpp
--- a/Engine/Source/Runtime/Slate/Private/DockingFramework/SDockingArea.cpp
+++ b/Engine/Source/Runtime/Slate/Private/DockingFramework/SDockingArea.cpp
@@ -93,9 +93,9 @@ void SDockingArea::Construct( const FArguments& InArgs, const TSharedRef<FTabMan
 
 void SDockingArea::OnDragEnter( const FGeometry& MyGeometry, const FDragDropEvent& DragDropEvent )
 {
-	if ( DragDrop::IsTypeMatch<FDockingDragOperation>(DragDropEvent.GetOperation()) )
+	TSharedPtr<FDockingDragOperation> DragDropOperation = DragDropEvent.GetOperationAs<FDockingDragOperation>();
+	if ( DragDropOperation.IsValid() )
 	{
-		TSharedPtr<FDockingDragOperation> DragDropOperation = StaticCastSharedPtr<FDockingDragOperation>(DragDropEvent.GetOperation());
 		if ( DragDropOperation->GetTabBeingDragged()->CanDockInNode(SharedThis(this), SDockTab::DockingViaTarget) )
 		{
 			ShowCross();
@@ -106,7 +106,7 @@ void SDockingArea::OnDragEnter( const FGeometry& MyGeometry, const FDragDropEven
 
 void SDockingArea::OnDragLeave( const FDragDropEvent& DragDropEvent )
 {
-	if ( DragDrop::IsTypeMatch<FDockingDragOperation>(DragDropEvent.GetOperation()) )
+	if ( DragDropEvent.GetOperationAs<FDockingDragOperation>().IsValid() )
 	{
 		HideCross();
 	}
@@ -114,7 +114,7 @@ void SDockingArea::OnDragLeave( const FDragDropEvent& DragDropEvent )
 
 FReply SDockingArea::OnDrop( const FGeometry& MyGeometry, const FDragDropEvent& DragDropEvent )
 {
-	if ( DragDrop::IsTypeMatch<FDockingDragOperation>(DragDropEvent.GetOperation()) )
+	if ( DragDropEvent.GetOperationAs<FDockingDragOperation>().IsValid() )
 	{
 		HideCross();
 	}
@@ -124,13 +124,13 @@ FReply SDockingArea::OnDrop( const FGeometry& MyGeometry, const FDragDropEvent&
 
 FReply SDockingArea::OnUserAttemptingDock( SDockingNode::RelativeDirection Direction, const FDragDropEvent& DragDropEvent )
 {
-	if ( DragDrop::IsTypeMatch<FDockingDragOperation>(DragDropEvent.GetOperation()) )
+	TSharedPtr<FDockingDragOperation> DragDropOperation = DragDropEvent.GetOperationAs<FDockingDragOperation>();
+	if ( DragDropOperation.IsValid() )
 	{
 		if (Direction == Center)
 		{
 			//check(Children.Num() <= 1);
-			TSharedPtr<FDockingDragOperation> DragDropOperation = StaticCastSharedPtr<FDockingDragOperation>(DragDropEvent.GetOperation());
-			
+
 			TSharedRef<SDockingTabStack> NewStack = SNew(SDockingTabStack, FTabManager::NewStack());
 			AddChildNode( NewStack );
 			NewStack->OpenTab( DragDropOperation->GetTabBeingDragged().ToSharedRef() );
diff --git a/Engine/Source/Runtime/Slate/Private/DockingFramework/SDockingCross.cpp b/Engine/Source/Runtime/Slate/Private/DockingFramework/SDockingCross.cpp
--- a/Engine/Source/Runtime/Slate/Private/DockingFramework/SDockingCross.cpp
+++ b/Engine/Source/Runtime/Slate/Private/DockingFramework/SDockingCross.cpp
@@ -232,19 +232,18 @@ static FDockingDragOperation::FDockTarget GetDropTarget( const TSharedPtr<SDocki
 
 void SDockingCross::OnDragLeave( const FDragDropEvent& DragDropEvent )
 {
-	if ( DragDrop::IsTypeMatch<FDockingDragOperation>(DragDropEvent.GetOperation()) )
+	TSharedPtr<FDockingDragOperation> DragDropOperation = DragDropEvent.GetOperationAs<FDockingDragOperation>();
+	if ( DragDropOperation.IsValid() )
 	{
-		TSharedPtr<FDockingDragOperation> DragDropOperation = StaticCastSharedPtr<FDockingDragOperation>(DragDropEvent.GetOperation());
 		DragDropOperation->SetHoveredTarget(FDockingDragOperation::FDockTarget(), DragDropEvent);
 	}
 }
 
 FReply SDockingCross::OnDragOver( const FGeometry& MyGeometry, const FDragDropEvent& DragDropEvent )
 {
-	if ( DragDrop::IsTypeMatch<FDockingDragOperation>(DragDropEvent.GetOperation()) )
+	TSharedPtr<FDockingDragOperation> DragDropOperation = DragDropEvent.GetOperationAs<FDockingDragOperation>();
+	if ( DragDropOperation.IsValid() )
 	{
-		TSharedPtr<FDockingDragOperation> DragDropOperation = StaticCastSharedPtr<FDockingDragOperation>(DragDropEvent.GetOperation());
-
 		const FDockingDragOperation::FDockTarget DropTarget = GetDropTarget(OwnerNode.Pin(), MyGeometry, DragDropEvent);
 		DragDropOperation->SetHoveredTarget(DropTarget, DragDropEvent);
 
@@ -258,7 +257,7 @@ FReply SDockingCross::OnDragOver( const FGeometry& MyGeometry, const FDragDropEv
 
 FReply SDockingCross::OnDrop( const FGeometry& MyGeometry, const FDragDropEvent& DragDropEvent )
 {
-	if ( DragDrop::IsTypeMatch<FDockingDragOperation>(DragDropEvent.GetOperation()) )
+	if ( DragDropEvent.GetOperationAs<FDockingDragOperation>().IsValid() )
 	{
 		const FDockingDragOperation::FDockTarget DropTarget = GetDropTarget(OwnerNode.Pin(), MyGeometry, DragDropEvent);
 		TSharedPtr<SDockingNode> DropTargetNode = DropTarget.TargetNode.Pin();
diff --git a/Engine/Source/Runtime/Slate/Private/DockingFramework/SDockingTabStack.cpp b/Engine/Source/Runtime/Slate/Private/DockingFramework/SDockingTabStack.cpp
--- a/Engine/Source/Runtime/Slate/Private/DockingFramework/SDockingTabStack.cpp
+++ b/Engine/Source/Runtime/Slate/Private/DockingFramework/SDockingTabStack.cpp
@@ -322,10 +322,9 @@ void SDockingTabStack::SetNodeContent( const TSharedRef<SWidget>& InContent, con
 
 FReply SDockingTabStack::OnDragOver( const FGeometry& MyGeometry, const FDragDropEvent& DragDropEvent )
 {
-	if ( DragDrop::IsTypeMatch<FDockingDragOperation>(DragDropEvent.GetOperation()) )
-	{
-		TSharedPtr<FDockingDragOperation> DragDropOperation = StaticCastSharedPtr<FDockingDragOperation>(DragDropEvent.GetOperation());
-		
+	TSharedPtr<FDockingDragOperation> DragDropOperation = DragDropEvent.GetOperationAs<FDockingDragOperation>();
+	if ( DragDropOperation.IsValid() )
+	{	
 		if (DragDropOperation->GetTabBeingDragged()->CanDockInNode(SharedThis(this), SDockTab::DockingViaTarget))
 		{
 			FGeometry OverlayGeometry = this->FindChildGeometry( MyGeometry, OverlayManagement.ContentAreaOverlay.ToSharedRef() );
@@ -348,18 +347,16 @@ FReply SDockingTabStack::OnDragOver( const FGeometry& MyGeometry, const FDragDro
 
 void SDockingTabStack::OnDragLeave( const FDragDropEvent& DragDropEvent )
 {
-	if ( DragDrop::IsTypeMatch<FDockingDragOperation>(DragDropEvent.GetOperation()) )
+	if (DragDropEvent.GetOperationAs<FDockingDragOperation>().IsValid())
 	{
-		TSharedPtr<FDockingDragOperation> DragDropOperation = StaticCastSharedPtr<FDockingDragOperation>(DragDropEvent.GetOperation());
 		HideCross();
 	}
 }
 
 FReply SDockingTabStack::OnDrop( const FGeometry& MyGeometry, const FDragDropEvent& DragDropEvent )
 {
-	if ( DragDrop::IsTypeMatch<FDockingDragOperation>(DragDropEvent.GetOperation()) )
+	if (DragDropEvent.GetOperationAs<FDockingDragOperation>().IsValid())
 	{
-		TSharedPtr<FDockingDragOperation> DragDropOperation = StaticCastSharedPtr<FDockingDragOperation>(DragDropEvent.GetOperation());
 		HideCross();
 	}
 
@@ -399,10 +396,9 @@ FReply SDockingTabStack::OnMouseButtonDown( const FGeometry& MyGeometry, const F
 
 FReply SDockingTabStack::OnUserAttemptingDock( SDockingNode::RelativeDirection Direction, const FDragDropEvent& DragDropEvent )
 {
-	if ( DragDrop::IsTypeMatch<FDockingDragOperation>(DragDropEvent.GetOperation()) )
+	TSharedPtr<FDockingDragOperation> DragDropOperation = DragDropEvent.GetOperationAs<FDockingDragOperation>();
+	if ( DragDropOperation.IsValid() )
 	{
-		TSharedPtr<FDockingDragOperation> DragDropOperation = StaticCastSharedPtr<FDockingDragOperation>(DragDropEvent.GetOperation());
-
 		// We want to replace this placeholder  with whatever is being dragged.				
 		CreateNewTabStackBySplitting( Direction )->OpenTab( DragDropOperation->GetTabBeingDragged().ToSharedRef() );
 
diff --git a/Engine/Source/Runtime/Slate/Private/DockingFramework/SDockingTabWell.cpp b/Engine/Source/Runtime/Slate/Private/DockingFramework/SDockingTabWell.cpp
--- a/Engine/Source/Runtime/Slate/Private/DockingFramework/SDockingTabWell.cpp
+++ b/Engine/Source/Runtime/Slate/Private/DockingFramework/SDockingTabWell.cpp
@@ -263,10 +263,9 @@ FReply SDockingTabWell::StartDraggingTab( TSharedRef<SDockTab> TabToStartDraggin
 
 void SDockingTabWell::OnDragEnter( const FGeometry& MyGeometry, const FDragDropEvent& DragDropEvent )
 {
-	if ( DragDrop::IsTypeMatch<FDockingDragOperation>(DragDropEvent.GetOperation()) )
+	TSharedPtr<FDockingDragOperation> DragDropOperation = DragDropEvent.GetOperationAs<FDockingDragOperation>();
+	if ( DragDropOperation.IsValid() )
 	{
-		TSharedPtr<FDockingDragOperation> DragDropOperation = StaticCastSharedPtr<FDockingDragOperation>(DragDropEvent.GetOperation());
-
 		if (DragDropOperation->GetTabBeingDragged()->CanDockInNode(ParentTabStackPtr.Pin().ToSharedRef(), SDockTab::DockingViaTabWell))
 		{
 			// The user dragged a tab into this TabWell.
@@ -286,9 +285,9 @@ void SDockingTabWell::OnDragEnter( const FGeometry& MyGeometry, const FDragDropE
 
 void SDockingTabWell::OnDragLeave( const FDragDropEvent& DragDropEvent )
 {
-	if ( DragDrop::IsTypeMatch<FDockingDragOperation>(DragDropEvent.GetOperation()) )
+	TSharedPtr<FDockingDragOperation> DragDropOperation = DragDropEvent.GetOperationAs<FDockingDragOperation>();
+	if ( DragDropOperation.IsValid() )
 	{
-		TSharedPtr<FDockingDragOperation> DragDropOperation = StaticCastSharedPtr<FDockingDragOperation>( DragDropEvent.GetOperation() );
 		// Check for TabBeingDraggedPtr validity as it may no longer be valid when dragging tabs in game
 		if (this->TabBeingDraggedPtr.IsValid() && DragDropOperation->GetTabBeingDragged()->CanDockInNode(ParentTabStackPtr.Pin().ToSharedRef(), SDockTab::DockingViaTabWell))
 		{
@@ -325,9 +324,9 @@ void SDockingTabWell::OnDragLeave( const FDragDropEvent& DragDropEvent )
 
 FReply SDockingTabWell::OnDragOver( const FGeometry& MyGeometry, const FDragDropEvent& DragDropEvent )
 {
-	if ( DragDrop::IsTypeMatch<FDockingDragOperation>(DragDropEvent.GetOperation()) )
+	TSharedPtr<FDockingDragOperation> DragDropOperation = DragDropEvent.GetOperationAs<FDockingDragOperation>();
+	if ( DragDropOperation.IsValid() )
 	{
-		TSharedPtr<FDockingDragOperation> DragDropOperation = StaticCastSharedPtr<FDockingDragOperation>( DragDropEvent.GetOperation() );
 		if (DragDropOperation->GetTabBeingDragged()->CanDockInNode(ParentTabStackPtr.Pin().ToSharedRef(), SDockTab::DockingViaTabWell))
 		{
 			// We are dragging the tab through a TabWell.
@@ -341,9 +340,9 @@ FReply SDockingTabWell::OnDragOver( const FGeometry& MyGeometry, const FDragDrop
 
 FReply SDockingTabWell::OnDrop( const FGeometry& MyGeometry, const FDragDropEvent& DragDropEvent )
 {
-	if ( DragDrop::IsTypeMatch<FDockingDragOperation>( DragDropEvent.GetOperation().ToSharedRef() ) )
+	TSharedPtr<FDockingDragOperation> DragDropOperation = DragDropEvent.GetOperationAs<FDockingDragOperation>();
+	if (DragDropOperation.IsValid())
 	{
-		TSharedPtr<FDockingDragOperation> DragDropOperation = StaticCastSharedPtr<FDockingDragOperation>( DragDropEvent.GetOperation() );
 		if (DragDropOperation->GetTabBeingDragged()->CanDockInNode(ParentTabStackPtr.Pin().ToSharedRef(), SDockTab::DockingViaTabWell))
 		{
 			if ( ensure( TabBeingDraggedPtr.IsValid() ) )
diff --git a/Engine/Source/Runtime/Slate/Private/DockingFramework/SDockingTarget.cpp b/Engine/Source/Runtime/Slate/Private/DockingFramework/SDockingTarget.cpp
--- a/Engine/Source/Runtime/Slate/Private/DockingFramework/SDockingTarget.cpp
+++ b/Engine/Source/Runtime/Slate/Private/DockingFramework/SDockingTarget.cpp
@@ -75,38 +75,36 @@ void SDockingTarget::Construct( const FArguments& InArgs )
 
 void SDockingTarget::OnDragEnter( const FGeometry& MyGeometry, const FDragDropEvent& DragDropEvent )
 {
-	if ( DragDrop::IsTypeMatch<FDockingDragOperation>(DragDropEvent.GetOperation()) )
+	TSharedPtr<FDockingDragOperation> DragDropOperation = DragDropEvent.GetOperationAs<FDockingDragOperation>();
+	if ( DragDropOperation.IsValid() )
 	{
 		bIsDragHovered = true;
 		// Provide feedback that this target is hovered
 		ColorAndOpacity = FCoreStyle::Get().GetColor(TEXT("Docking.Cross.HoveredTint"));
 
 		// Let the drag and drop operation know that we are currently hovering over this node
-		TSharedPtr<FDockingDragOperation> DragDropOperation = StaticCastSharedPtr<FDockingDragOperation>(DragDropEvent.GetOperation());
 		DragDropOperation->SetHoveredTarget( FDockingDragOperation::FDockTarget(OwnerNode.Pin(), GetDockDirection()), DragDropEvent );
 	}
 }
 	
 void SDockingTarget::OnDragLeave( const FDragDropEvent& DragDropEvent )
 {
-	if ( DragDrop::IsTypeMatch<FDockingDragOperation>(DragDropEvent.GetOperation()) )
+	TSharedPtr<FDockingDragOperation> DragDropOperation = DragDropEvent.GetOperationAs<FDockingDragOperation>();
+	if ( DragDropOperation.IsValid() )
 	{
 		bIsDragHovered = false;
 		// We are no longer hovered, remove the hover feedback
 		ColorAndOpacity = FCoreStyle::Get().GetColor(TEXT("Docking.Cross.Tint"));
 		
 		// Let the drag and drop operation know that we are no longer hovering over any nodes
-		TSharedPtr<FDockingDragOperation> DragDropOperation = StaticCastSharedPtr<FDockingDragOperation>(DragDropEvent.GetOperation());
 		DragDropOperation->SetHoveredTarget( FDockingDragOperation::FDockTarget(), DragDropEvent );
 	}
 }
 	
 FReply SDockingTarget::OnDrop( const FGeometry& MyGeometry, const FDragDropEvent& DragDropEvent )
 {
-	if ( DragDrop::IsTypeMatch<FDockingDragOperation>(DragDropEvent.GetOperation()) )
+	if (DragDropEvent.GetOperationAs<FDockingDragOperation>().IsValid())
 	{
-		TSharedPtr<FDockingDragOperation> DragDropOperation = StaticCastSharedPtr<FDockingDragOperation>(DragDropEvent.GetOperation());
-
 		TSharedPtr<SDockingNode> PinnedOwnerNode = OwnerNode.Pin();
 
 		// We are a direction node, so re-arrange layout as desired
diff --git a/Engine/Source/Runtime/Slate/Private/MultiBox/MultiBox.cpp b/Engine/Source/Runtime/Slate/Private/MultiBox/MultiBox.cpp
--- a/Engine/Source/Runtime/Slate/Private/MultiBox/MultiBox.cpp
+++ b/Engine/Source/Runtime/Slate/Private/MultiBox/MultiBox.cpp
@@ -56,7 +56,7 @@ const FMultiBoxCustomization FMultiBoxCustomization::None( NAME_None );
 
 void SMultiBlockBaseWidget::OnDragEnter( const FGeometry& MyGeometry, const FDragDropEvent& DragDropEvent )
 {
-	if( DragDrop::IsTypeMatch<FUICommandDragDropOp>( DragDropEvent.GetOperation() ) )
+	if ( DragDropEvent.GetOperationAs<FUICommandDragDropOp>().IsValid() )
 	{
 		OwnerMultiBoxWidget.Pin()->OnCustomCommandDragEnter( MultiBlock.ToSharedRef(), MyGeometry, DragDropEvent );
 	}
@@ -64,7 +64,7 @@ void SMultiBlockBaseWidget::OnDragEnter( const FGeometry& MyGeometry, const FDra
 
 FReply SMultiBlockBaseWidget::OnDragOver( const FGeometry& MyGeometry, const FDragDropEvent& DragDropEvent )
 {
-	if( DragDrop::IsTypeMatch<FUICommandDragDropOp>( DragDropEvent.GetOperation() ) )
+	if ( DragDropEvent.GetOperationAs<FUICommandDragDropOp>().IsValid() )
 	{
 		OwnerMultiBoxWidget.Pin()->OnCustomCommandDragged( MultiBlock.ToSharedRef(), MyGeometry, DragDropEvent );
 		return FReply::Handled();
@@ -75,7 +75,7 @@ FReply SMultiBlockBaseWidget::OnDragOver( const FGeometry& MyGeometry, const FDr
 
 FReply SMultiBlockBaseWidget::OnDrop( const FGeometry& MyGeometry, const FDragDropEvent& DragDropEvent )
 {
-	if( DragDrop::IsTypeMatch<FUICommandDragDropOp>( DragDropEvent.GetOperation() )  )
+	if ( DragDropEvent.GetOperationAs<FUICommandDragDropOp>().IsValid() )
 	{
 		OwnerMultiBoxWidget.Pin()->OnCustomCommandDropped();
 		return FReply::Handled();
@@ -965,7 +965,7 @@ void SMultiBoxWidget::OnDropExternal()
 
 FReply SMultiBoxWidget::OnDragOver( const FGeometry& MyGeometry, const FDragDropEvent& DragDropEvent )
 {
-	if( DragDrop::IsTypeMatch<FUICommandDragDropOp>( DragDropEvent.GetOperation() ) && MultiBox->IsInEditMode() )
+	if ( DragDropEvent.GetOperationAs<FUICommandDragDropOp>().IsValid() && MultiBox->IsInEditMode() )
 	{
 		return FReply::Handled();
 	}
@@ -975,7 +975,7 @@ FReply SMultiBoxWidget::OnDragOver( const FGeometry& MyGeometry, const FDragDrop
 
 FReply SMultiBoxWidget::OnDrop( const FGeometry& MyGeometry, const FDragDropEvent& DragDropEvent )
 {
-	if( DragDrop::IsTypeMatch<FUICommandDragDropOp>( DragDropEvent.GetOperation() ) )
+	if ( DragDropEvent.GetOperationAs<FUICommandDragDropOp>().IsValid() )
 	{
 		OnCustomCommandDropped();
 		return FReply::Handled();
diff --git a/Engine/Source/Runtime/Slate/Private/MultiBox/MultiBoxCustomization.cpp b/Engine/Source/Runtime/Slate/Private/MultiBox/MultiBoxCustomization.cpp
--- a/Engine/Source/Runtime/Slate/Private/MultiBox/MultiBoxCustomization.cpp
+++ b/Engine/Source/Runtime/Slate/Private/MultiBox/MultiBoxCustomization.cpp
@@ -45,7 +45,7 @@ FReply SMultiBlockDragHandle::OnMouseButtonDown( const FGeometry& MyGeometry, co
 
 void SMultiBlockDragHandle::OnDragEnter( const FGeometry& MyGeometry, const FDragDropEvent& DragDropEvent )
 {
-	if( DragDrop::IsTypeMatch<FUICommandDragDropOp>( DragDropEvent.GetOperation() ) )
+	if ( DragDropEvent.GetOperationAs<FUICommandDragDropOp>().IsValid() )
 	{
 		BaseWidget.Pin()->OnCustomCommandDragEnter( Block.ToSharedRef(), MyGeometry, DragDropEvent );
 	}
@@ -53,7 +53,7 @@ void SMultiBlockDragHandle::OnDragEnter( const FGeometry& MyGeometry, const FDra
 
 FReply SMultiBlockDragHandle::OnDragOver( const FGeometry& MyGeometry, const FDragDropEvent& DragDropEvent )
 {
-	if( DragDrop::IsTypeMatch<FUICommandDragDropOp>( DragDropEvent.GetOperation() ) )
+	if ( DragDropEvent.GetOperationAs<FUICommandDragDropOp>().IsValid() )
 	{
 		BaseWidget.Pin()->OnCustomCommandDragged( Block.ToSharedRef(), MyGeometry, DragDropEvent );
 		return FReply::Handled();
@@ -64,7 +64,7 @@ FReply SMultiBlockDragHandle::OnDragOver( const FGeometry& MyGeometry, const FDr
 
 FReply SMultiBlockDragHandle::OnDrop( const FGeometry& MyGeometry, const FDragDropEvent& DragDropEvent )
 {
-	if( DragDrop::IsTypeMatch<FUICommandDragDropOp>( DragDropEvent.GetOperation() )  )
+	if ( DragDropEvent.GetOperationAs<FUICommandDragDropOp>().IsValid() )
 	{
 		BaseWidget.Pin()->OnCustomCommandDropped();
 		return FReply::Handled();
diff --git a/Engine/Source/Runtime/Slate/Private/SColorThemes.cpp b/Engine/Source/Runtime/Slate/Private/SColorThemes.cpp
--- a/Engine/Source/Runtime/Slate/Private/SColorThemes.cpp
+++ b/Engine/Source/Runtime/Slate/Private/SColorThemes.cpp
@@ -97,7 +97,7 @@ void SColorTrash::Construct( const FArguments& InArgs )
 */
 void SColorTrash::OnDragEnter( const FGeometry& MyGeometry, const FDragDropEvent& DragDropEvent )
 {
-	if ( DragDrop::IsTypeMatch<FColorDragDrop>(DragDropEvent.GetOperation()) )
+	if ( DragDropEvent.GetOperationAs<FColorDragDrop>().IsValid() )
 	{
 		bBorderActivated = true;
 	}
@@ -112,7 +112,7 @@ void SColorTrash::OnDragEnter( const FGeometry& MyGeometry, const FDragDropEvent
 */
 void SColorTrash::OnDragLeave( const FDragDropEvent& DragDropEvent )
 {
-	if ( DragDrop::IsTypeMatch<FColorDragDrop>(DragDropEvent.GetOperation()) )
+	if ( DragDropEvent.GetOperationAs<FColorDragDrop>().IsValid() )
 	{
 		bBorderActivated = false;
 	}
@@ -128,10 +128,9 @@ void SColorTrash::OnDragLeave( const FDragDropEvent& DragDropEvent )
 */
 FReply SColorTrash::OnDrop( const FGeometry& MyGeometry, const FDragDropEvent& DragDropEvent )
 {
-	if ( DragDrop::IsTypeMatch<FColorDragDrop>( DragDropEvent.GetOperation().ToSharedRef() ) )
+	TSharedPtr<FColorDragDrop> DragDropContent = DragDropEvent.GetOperationAs<FColorDragDrop>();
+	if ( DragDropContent.IsValid() )
 	{
-		TSharedPtr<FColorDragDrop> DragDropContent = StaticCastSharedPtr<FColorDragDrop>( DragDropEvent.GetOperation() );
-
 		DragDropContent->bSetForDeletion = true;
 
 		bBorderActivated = false;
@@ -437,10 +436,9 @@ FChildren* SThemeColorBlocksBar::GetChildren()
 
 void SThemeColorBlocksBar::OnDragEnter( const FGeometry& MyGeometry, const FDragDropEvent& DragDropEvent )
 {
-	if ( DragDrop::IsTypeMatch<FColorDragDrop>(DragDropEvent.GetOperation()) )
+	TSharedPtr<FColorDragDrop> DragDropOperation = DragDropEvent.GetOperationAs<FColorDragDrop>();
+	if ( DragDropOperation.IsValid() )
 	{
-		TSharedPtr<FColorDragDrop> DragDropOperation = StaticCastSharedPtr<FColorDragDrop>(DragDropEvent.GetOperation());
-
 		NewColorPlaceholder = MakeShareable(new FLinearColor(DragDropOperation->Color));
 		NewColorBlockPlaceholder =
 			SNew(SThemeColorBlock) .Color(NewColorPlaceholder) .UseSRGB(bUseSRGB) .UseAlpha(bUseAlpha);
@@ -449,9 +447,9 @@ void SThemeColorBlocksBar::OnDragEnter( const FGeometry& MyGeometry, const FDrag
 
 void SThemeColorBlocksBar::OnDragLeave( const FDragDropEvent& DragDropEvent )
 {
-	if ( DragDrop::IsTypeMatch<FColorDragDrop>(DragDropEvent.GetOperation()) )
+	TSharedPtr<FColorDragDrop> DragDropOperation = DragDropEvent.GetOperationAs<FColorDragDrop>();
+	if ( DragDropOperation.IsValid() )
 	{
-		TSharedPtr<FColorDragDrop> DragDropOperation = StaticCastSharedPtr<FColorDragDrop>( DragDropEvent.GetOperation() );
 		DragDropOperation->SetDecoratorVisibility(true);
 
 		DestroyPlaceholders();
@@ -468,10 +466,9 @@ void SThemeColorBlocksBar::OnDragLeave( const FDragDropEvent& DragDropEvent )
 */
 FReply SThemeColorBlocksBar::OnDragOver( const FGeometry& MyGeometry, const FDragDropEvent& DragDropEvent )
 {
-	if ( DragDrop::IsTypeMatch<FColorDragDrop>(DragDropEvent.GetOperation()) )
+	TSharedPtr<FColorDragDrop> DragDropOperation = DragDropEvent.GetOperationAs<FColorDragDrop>();
+	if ( DragDropOperation.IsValid() )
 	{
-		TSharedPtr<FColorDragDrop> DragDropOperation = StaticCastSharedPtr<FColorDragDrop>( DragDropEvent.GetOperation() );
-		
 		const float ChildSizeX = MyGeometry.Size.X / (ColorBlocks.Num() + 1);
 		const float GrabOffsetX = PlaceholderInitialGrabOffset.X;
 		PlaceholderBlockOffset = MyGeometry.AbsoluteToLocal(DragDropEvent.GetScreenSpacePosition()).X -
@@ -497,10 +494,9 @@ FReply SThemeColorBlocksBar::OnDragOver( const FGeometry& MyGeometry, const FDra
 */
 FReply SThemeColorBlocksBar::OnDrop( const FGeometry& MyGeometry, const FDragDropEvent& DragDropEvent )
 {
-	if ( DragDrop::IsTypeMatch<FColorDragDrop>( DragDropEvent.GetOperation().ToSharedRef() ) )
+	TSharedPtr<FColorDragDrop> DragDropOperation = DragDropEvent.GetOperationAs<FColorDragDrop>();
+	if ( DragDropOperation.IsValid() )
 	{
-		TSharedPtr<FColorDragDrop> DragDropOperation = StaticCastSharedPtr<FColorDragDrop>( DragDropEvent.GetOperation() );
-		
 		const float ChildSizeX = MyGeometry.Size.X / (ColorBlocks.Num() + 1);
 		const float CurrentDragCenter = PlaceholderBlockOffset + ChildSizeX * 0.5f;
 		const int32 LocID = FMath::Clamp(int32(CurrentDragCenter / ChildSizeX), 0, ColorBlocks.Num());
diff --git a/Engine/Source/Runtime/Slate/Private/SEditableText.cpp b/Engine/Source/Runtime/Slate/Private/SEditableText.cpp
--- a/Engine/Source/Runtime/Slate/Private/SEditableText.cpp
+++ b/Engine/Source/Runtime/Slate/Private/SEditableText.cpp
@@ -1316,9 +1316,9 @@ int32 SEditableText::OnPaint( const FGeometry& AllottedGeometry, const FSlateRec
 
 FReply SEditableText::OnDragOver( const FGeometry& MyGeometry, const FDragDropEvent& DragDropEvent )
 {
-	if ( DragDrop::IsTypeMatch<FExternalDragOperation>( DragDropEvent.GetOperation() ) )
+	TSharedPtr<FExternalDragOperation> DragDropOp = DragDropEvent.GetOperationAs<FExternalDragOperation>();
+	if ( DragDropOp.IsValid() )
 	{
-		TSharedPtr<FExternalDragOperation> DragDropOp = StaticCastSharedPtr<FExternalDragOperation>(DragDropEvent.GetOperation());
 		if ( DragDropOp->HasText() )
 		{
 			return FReply::Handled();
@@ -1331,9 +1331,9 @@ FReply SEditableText::OnDragOver( const FGeometry& MyGeometry, const FDragDropEv
 
 FReply SEditableText::OnDrop( const FGeometry& MyGeometry, const FDragDropEvent& DragDropEvent )
 {
-	if ( DragDrop::IsTypeMatch<FExternalDragOperation>( DragDropEvent.GetOperation() ) )
+	TSharedPtr<FExternalDragOperation> DragDropOp = DragDropEvent.GetOperationAs<FExternalDragOperation>();
+	if ( DragDropOp.IsValid() )
 	{
-		TSharedPtr<FExternalDragOperation> DragDropOp = StaticCastSharedPtr<FExternalDragOperation>(DragDropEvent.GetOperation());
 		if ( DragDropOp->HasText() )
 		{
 			this->Text = FText::FromString(DragDropOp->GetText());
diff --git a/Engine/Source/Runtime/Slate/Public/DragAndDrop.h b/Engine/Source/Runtime/Slate/Public/DragAndDrop.h
--- a/Engine/Source/Runtime/Slate/Public/DragAndDrop.h
+++ b/Engine/Source/Runtime/Slate/Public/DragAndDrop.h
@@ -18,9 +18,10 @@ public:
 
 	virtual ~FDragDropOperation();
 
-	virtual bool IsOfType(const FString& Type)
+	/** Check if this drag and drop operation can cast safely to the specified template type */
+	template<class TType> bool IsOfType() const 
 	{
-		return false;
+		return IsOfTypeImpl(TType::GetTypeId());
 	}
 
 	/**
@@ -61,6 +62,12 @@ protected:
 	/** Constructs the window and widget if applicable */
 	virtual void Construct();
 
+	/** Check if this drag and drop operation can cast safely to the specified type */
+	virtual bool IsOfTypeImpl(const FString& Type) const
+	{
+		return false;
+	}
+
 protected:
 
 	/** Destroy the cursor decorator */
@@ -147,15 +154,7 @@ DECLARE_DELEGATE_OneParam( FOnDragDropUpdate,
 
 #define DRAG_DROP_OPERATOR_TYPE(TYPE, BASE) \
 	static const FString& GetTypeId() { static FString Type = TEXT(#TYPE); return Type; } \
-	virtual bool IsOfType(const FString& Type) OVERRIDE { return GetTypeId() == Type || BASE::IsOfType(Type); }
-
-namespace DragDrop
-{
-	/** See if this dragdrop operation matches another dragdrop operation */
-	template <typename OperatorType>
-	bool IsTypeMatch(const TSharedPtr<FDragDropOperation> Operation);
-}
-
+	virtual bool IsOfTypeImpl(const FString& Type) const OVERRIDE { return GetTypeId() == Type || BASE::IsOfTypeImpl(Type); }
 
 
 /**
diff --git a/Engine/Source/Runtime/Slate/Public/DragAndDrop.inl b/Engine/Source/Runtime/Slate/Public/DragAndDrop.inl
--- a/Engine/Source/Runtime/Slate/Public/DragAndDrop.inl
+++ b/Engine/Source/Runtime/Slate/Public/DragAndDrop.inl
@@ -6,7 +6,7 @@
 template<typename OperationType>
 TSharedPtr<OperationType> FDragDropEvent::GetOperationAs() const
 {
-	if (DragDrop::IsTypeMatch<OperationType>(Content))
+	if (Content.IsValid() && Content->IsOfType<OperationType>())
 	{
 		return StaticCastSharedPtr<OperationType>(Content);
 	}
@@ -14,16 +14,4 @@ TSharedPtr<OperationType> FDragDropEvent::GetOperationAs() const
 	{
 		return TSharedPtr<OperationType>();
 	}
-}
-
-// DragAndDrop functions that depend on classes defined after DragAndDrop.h is included in Slate.h
-
-namespace DragDrop
-{
-	/** See if this dragdrop operation matches another dragdrop operation */
-	template <typename OperatorType>
-	bool IsTypeMatch(const TSharedPtr<FDragDropOperation> Operation)
-	{
-		return Operation.IsValid() && Operation->IsOfType(OperatorType::GetTypeId());
-	}
-}
+}
\ No newline at end of file
