#UE4 - Fix crash reloading the level editor

diff --git a/Engine/Source/Editor/LevelEditor/Private/SLevelViewport.cpp b/Engine/Source/Editor/LevelEditor/Private/SLevelViewport.cpp
--- a/Engine/Source/Editor/LevelEditor/Private/SLevelViewport.cpp
+++ b/Engine/Source/Editor/LevelEditor/Private/SLevelViewport.cpp
@@ -79,6 +79,8 @@ SLevelViewport::~SLevelViewport()
 	LevelEditor.OnActorSelectionChanged().RemoveAll( this );
 	LevelEditor.OnMapChanged().RemoveAll( this );
 	LevelEditor.OnRedrawLevelEditingViewports().RemoveAll( this );
+	GEngine->OnLevelActorDeleted().RemoveAll( this );
+
 	GetMutableDefault<ULevelEditorViewportSettings>()->OnSettingChanged().RemoveAll( this );
 
 	// If this viewport has a high res screenshot window attached to it, close it
