UE4: Added EWorldType::Inactive to represent worlds that are loaded but not currently being edited in the level editor.

diff --git a/Engine/Source/Editor/LevelEditor/Private/LightmapResRatioAdjust.cpp b/Engine/Source/Editor/LevelEditor/Private/LightmapResRatioAdjust.cpp
--- a/Engine/Source/Editor/LevelEditor/Private/LightmapResRatioAdjust.cpp
+++ b/Engine/Source/Editor/LevelEditor/Private/LightmapResRatioAdjust.cpp
@@ -68,7 +68,7 @@ bool FLightmapResRatioAdjustSettings::ApplyRatioAdjustment()
 				// Store the required levels relating to this object based on the level option settings
 				UWorld* MeshWorld = StaticMeshComp->GetWorld();
 				// Exclude preview worlds				
-				if( MeshWorld->WorldType != EWorldType::Preview)
+				if( MeshWorld->WorldType != EWorldType::Preview && MeshWorld->WorldType != EWorldType::Inactive )
 				{
 					AddRequiredLevels( Settings.LevelOptions, MeshWorld , MeshLevels );
 					Worlds.AddUnique( MeshWorld );
@@ -85,7 +85,7 @@ bool FLightmapResRatioAdjustSettings::ApplyRatioAdjustment()
 			{
 				UWorld* BrushWorld = Brush->GetWorld();
 				// Exclude preview worlds
-				if( BrushWorld->WorldType != EWorldType::Preview)
+				if( BrushWorld->WorldType != EWorldType::Preview && BrushWorld->WorldType != EWorldType::Inactive )
 				{
 					AddRequiredLevels( Settings.LevelOptions, BrushWorld, BrushLevels );				
 					Worlds.AddUnique( BrushWorld );
@@ -107,7 +107,7 @@ bool FLightmapResRatioAdjustSettings::ApplyRatioAdjustment()
 						ABrush* AttachedBrush = CastChecked<ABrush>( AttachedBrushes[BrushIndex] );
 						UWorld* BrushWorld = AttachedBrush->GetWorld();
 						// Exclude preview worlds
-						if( BrushWorld->WorldType != EWorldType::Preview)
+						if( BrushWorld->WorldType != EWorldType::Preview && BrushWorld->WorldType != EWorldType::Inactive )
 						{	
 							AddRequiredLevels( Settings.LevelOptions, BrushWorld, BrushLevels );
 							Worlds.AddUnique( BrushWorld );
@@ -161,9 +161,9 @@ bool FLightmapResRatioAdjustSettings::ApplyRatioAdjustment()
 				{
 					bConvertIt = false;
 				}
-				else if( SMComp->GetWorld() && SMComp->GetWorld()->WorldType == EWorldType::Preview )
+				else if( SMComp->GetWorld() && SMComp->GetWorld()->WorldType == EWorldType::Preview && SMComp->GetWorld()->WorldType != EWorldType::Inactive )
 				{
-					// Dont do objects with a preview world
+					// Don't do objects with a preview or inactive world
 					bConvertIt = false;
 				}
 				if (bConvertIt)
diff --git a/Engine/Source/Editor/UnrealEd/Private/Editor.cpp b/Engine/Source/Editor/UnrealEd/Private/Editor.cpp
--- a/Engine/Source/Editor/UnrealEd/Private/Editor.cpp
+++ b/Engine/Source/Editor/UnrealEd/Private/Editor.cpp
@@ -6523,7 +6523,7 @@ void UEditorEngine::VerifyLoadMapWorldCleanup()
 	for( TObjectIterator<UWorld> It; It; ++It )
 	{
 		UWorld* World = *It;
-		if (World->WorldType != EWorldType::Preview && World->WorldType != EWorldType::Editor)
+		if (World->WorldType != EWorldType::Preview && World->WorldType != EWorldType::Editor && World->WorldType != EWorldType::Inactive)
 		{
 			TArray<UWorld*> OtherEditorWorlds;
 			EditorLevelUtils::GetWorlds(EditorWorld, OtherEditorWorlds, true, false);
diff --git a/Engine/Source/Editor/UnrealEd/Private/EditorServer.cpp b/Engine/Source/Editor/UnrealEd/Private/EditorServer.cpp
--- a/Engine/Source/Editor/UnrealEd/Private/EditorServer.cpp
+++ b/Engine/Source/Editor/UnrealEd/Private/EditorServer.cpp
@@ -1596,7 +1596,7 @@ void UEditorEngine::EditorDestroyWorld( FWorldContext & Context, const FText& Cl
 		WorldPackage = NULL;
 	}
 
-	if (Context.World()->WorldType != EWorldType::Preview)
+	if (Context.World()->WorldType != EWorldType::Preview && Context.World()->WorldType != EWorldType::Inactive)
 	{
 		// Go away, come again never!
 		Context.World()->ClearFlags(RF_Standalone | RF_RootSet | RF_Transactional);
@@ -1628,7 +1628,8 @@ void UEditorEngine::EditorDestroyWorld( FWorldContext & Context, const FText& Cl
 	for( TObjectIterator<UWorld> It; It; ++It )
 	{
 		UWorld* RemainingWorld = *It;
-		if( RemainingWorld->WorldType != EWorldType::Preview && !WorldHasValidContext(RemainingWorld))
+		const bool bIsPersistantWorldType = (RemainingWorld->WorldType == EWorldType::Inactive) || (RemainingWorld->WorldType == EWorldType::Preview);
+		if (!bIsPersistantWorldType && !WorldHasValidContext(RemainingWorld))
 		{
 			StaticExec(RemainingWorld, *FString::Printf(TEXT("OBJ REFS CLASS=WORLD NAME=%s"), *RemainingWorld->GetPathName()));
 
diff --git a/Engine/Source/Editor/UnrealEd/Private/Kismet2/BlueprintEditorUtils.cpp b/Engine/Source/Editor/UnrealEd/Private/Kismet2/BlueprintEditorUtils.cpp
--- a/Engine/Source/Editor/UnrealEd/Private/Kismet2/BlueprintEditorUtils.cpp
+++ b/Engine/Source/Editor/UnrealEd/Private/Kismet2/BlueprintEditorUtils.cpp
@@ -5159,8 +5159,8 @@ void FBlueprintEditorUtils::FindAndSetDebuggableBlueprintInstances()
 				for (int32 iSelected = 0; iSelected < Selected->Num() ; iSelected++)
 				{
 					AActor* ObjectAsActor = Cast<AActor>( Selected->GetSelectedObject( iSelected ) );
-
-					if ((ObjectAsActor != NULL) && (ObjectAsActor->GetWorld() != NULL) && (ObjectAsActor->GetWorld()->WorldType != EWorldType::Preview))
+					UWorld* ActorWorld = ObjectAsActor ? ObjectAsActor->GetWorld() : NULL;
+					if ((ActorWorld != NULL) && (ActorWorld->WorldType != EWorldType::Preview) && (ActorWorld->WorldType != EWorldType::Inactive))
 					{
 						if( IsObjectADebugCandidate(ObjectAsActor, EachBlueprint, true/*bInDisallowDerivedBlueprints*/ ) == true )
 						{
@@ -5186,7 +5186,8 @@ void FBlueprintEditorUtils::FindAndSetDebuggableBlueprintInstances()
 				for (TObjectIterator<UObject> It; It; ++It)
 				{
 					AActor* ObjectAsActor = Cast<AActor>( *It );
-					if( ObjectAsActor && ( ObjectAsActor->GetWorld() ) && ( ObjectAsActor->GetWorld()->WorldType != EWorldType::Preview) )
+					UWorld* ActorWorld = ObjectAsActor ? ObjectAsActor->GetWorld() : NULL;
+					if( ActorWorld && ( ActorWorld->WorldType != EWorldType::Preview) && ActorWorld->WorldType != EWorldType::Inactive )
 					{
 						if( IsObjectADebugCandidate(ObjectAsActor, EachBlueprint, true/*bInDisallowDerivedBlueprints*/ ) == true )
 						{
diff --git a/Engine/Source/Editor/UnrealEd/Private/Kismet2/KismetDebugUtilities.cpp b/Engine/Source/Editor/UnrealEd/Private/Kismet2/KismetDebugUtilities.cpp
--- a/Engine/Source/Editor/UnrealEd/Private/Kismet2/KismetDebugUtilities.cpp
+++ b/Engine/Source/Editor/UnrealEd/Private/Kismet2/KismetDebugUtilities.cpp
@@ -56,7 +56,7 @@ void FKismetDebugUtilities::OnScriptException(const UObject* ActiveObject, const
 	checkSlow(ActiveObject != NULL);
 
 	// Ignore script exceptions for preview actors
-	if(FActorEditorUtils::IsAPreviewActor(Cast<const AActor>(ActiveObject)))
+	if(FActorEditorUtils::IsAPreviewOrInactiveActor(Cast<const AActor>(ActiveObject)))
 	{
 		return;
 	}
diff --git a/Engine/Source/Runtime/Engine/Classes/Engine/EngineTypes.h b/Engine/Source/Runtime/Engine/Classes/Engine/EngineTypes.h
--- a/Engine/Source/Runtime/Engine/Classes/Engine/EngineTypes.h
+++ b/Engine/Source/Runtime/Engine/Classes/Engine/EngineTypes.h
@@ -400,7 +400,8 @@ namespace EWorldType
 		Game,		// The game world
 		Editor,		// A world being edited in the editor
 		PIE,		// A Play In Editor world
-		Preview		// A preview world for an editor tool
+		Preview,	// A preview world for an editor tool
+		Inactive	// An editor world that was loaded but not currently being edited in the level editor
 	};
 }
 
diff --git a/Engine/Source/Runtime/Engine/Private/ActorEditorUtils.cpp b/Engine/Source/Runtime/Engine/Private/ActorEditorUtils.cpp
--- a/Engine/Source/Runtime/Engine/Private/ActorEditorUtils.cpp
+++ b/Engine/Source/Runtime/Engine/Private/ActorEditorUtils.cpp
@@ -16,9 +16,10 @@ namespace FActorEditorUtils
 		return false;
 	}
 
-	bool IsAPreviewActor( const AActor* InActor )
+	bool IsAPreviewOrInactiveActor( const AActor* InActor )
 	{
-		return InActor && InActor->GetWorld() && InActor->GetWorld()->WorldType == EWorldType::Preview;
+		UWorld* World = InActor ? InActor->GetWorld() : NULL;
+		return World && (World->WorldType == EWorldType::Preview || World->WorldType == EWorldType::Inactive);
 	}
 
 	void GetEditableComponents( const AActor* InActor, TArray<UActorComponent*>& OutEditableComponents )
diff --git a/Engine/Source/Runtime/Engine/Private/UnrealEngine.cpp b/Engine/Source/Runtime/Engine/Private/UnrealEngine.cpp
--- a/Engine/Source/Runtime/Engine/Private/UnrealEngine.cpp
+++ b/Engine/Source/Runtime/Engine/Private/UnrealEngine.cpp
@@ -10462,7 +10462,8 @@ void UEngine::VerifyLoadMapWorldCleanup()
 	for( TObjectIterator<UWorld> It; It; ++It )
 	{
 		UWorld* World = *It;
-		if (World->WorldType != EWorldType::Preview && !WorldHasValidContext(World))
+		const bool bIsPersistantWorldType = (World->WorldType == EWorldType::Inactive) || (World->WorldType == EWorldType::Preview);
+		if (!bIsPersistantWorldType && !WorldHasValidContext(World))
 		{
 			// Print some debug information...
 			UE_LOG(LogLoad, Log, TEXT("%s not cleaned up by garbage collection! "), *World->GetFullName());
diff --git a/Engine/Source/Runtime/Engine/Public/ActorEditorUtils.h b/Engine/Source/Runtime/Engine/Public/ActorEditorUtils.h
--- a/Engine/Source/Runtime/Engine/Public/ActorEditorUtils.h
+++ b/Engine/Source/Runtime/Engine/Public/ActorEditorUtils.h
@@ -21,7 +21,7 @@ namespace FActorEditorUtils
 	*
 	* @return true if this actor is a preview actor, false otherwise
 	*/
-	ENGINE_API bool IsAPreviewActor( const AActor* InActor );
+	ENGINE_API bool IsAPreviewOrInactiveActor( const AActor* InActor );
 
 	/**
 	 * Gets a list of components, from an actor, which can be externally modified.  Editable components have properties visible when viewing their owner actor in a details view
diff --git a/Engine/Source/Runtime/Renderer/Private/SceneRendering.cpp b/Engine/Source/Runtime/Renderer/Private/SceneRendering.cpp
--- a/Engine/Source/Runtime/Renderer/Private/SceneRendering.cpp
+++ b/Engine/Source/Runtime/Renderer/Private/SceneRendering.cpp
@@ -797,7 +797,7 @@ void FSceneRenderer::RenderCustomDepthPass()
 	// do we have primitives in this pass?
 	bool bPrimitives = false;
 
-	if(!Scene->World || Scene->World->WorldType != EWorldType::Preview)
+	if(!Scene->World || (Scene->World->WorldType != EWorldType::Preview && Scene->World->WorldType != EWorldType::Inactive))
 	{
 		for(int32 ViewIndex = 0; ViewIndex < Views.Num(); ++ViewIndex)
 		{
