#ttp 331013 - Regression: EDITOR: CRASH: When using the viewport options to "create camera here", a crash can occur due to the undo behavior
#branch UE4
#proj Editor.LevelEditor
#summary Undo behaviour corrected in "Create camera here" action - now works as an undoable transaction.
#add Added an undo transaction in SLevelViewport::OnCreateCameraActor.

diff --git a/Engine/Source/Editor/LevelEditor/Private/SLevelViewport.cpp b/Engine/Source/Editor/LevelEditor/Private/SLevelViewport.cpp
--- a/Engine/Source/Editor/LevelEditor/Private/SLevelViewport.cpp
+++ b/Engine/Source/Editor/LevelEditor/Private/SLevelViewport.cpp
@@ -1444,6 +1444,8 @@ void SLevelViewport::OnCreateCameraActor()
 		return;
 	}
 
+	const FScopedTransaction Transaction(NSLOCTEXT("LevelViewport", "CreateCameraHere", "Create Camera Here"));
+
 	// Set new camera to match viewport
 	ACameraActor* pNewCamera = ViewportClient->GetWorld()->SpawnActor<ACameraActor>();
 	pNewCamera->SetActorLocation( ViewportClient->GetViewLocation(), false );
