#ttp 323891 - EDITOR: Viewport Layouts: Please add a 1x1 viewport preset layout option

#branch UE4
#proj Editor.LevelEditor
#summary Added 1x1 viewport layout and tidied the margins in the layout menu

#add Added FLevelViewportLayout::bIsMaximizeSupported and public getter IsMaximizeSupported(). This is set on the new one-pane layout to block the maximize command/button.
#add LevelViewportConfigurationNames::OnePane - new named layout config.
#add FLevelViewportCommands::ViewportConfig_OnePane - new layout command.
#change FLevelViewportLayout::InitCommonLayoutFromString() - will never start a layout maximised if bIsMaximizeSupported is false.
#change FLevelViewportLayout::SaveCommonLayoutString() - will never save a layout maximised if bIsMaximizeSupported is false.
#change FLevelViewportLayout::MaximizeViewport() - checks that bIsMaximizeSupported is true as it should never be called on a layout that doesn't support it.
#add SlateEditorStyle - added brushes for one pane option in the viewport layout menu.
#add SLevelViewportToolBar::GenerateViewportConfigsMenu() - new section for one pane layouts with the one pane button in it.
#change  SLevelViewportToolBar::GenerateViewportConfigsMenu() - fixed the layout margins by setting the label visibility to collapsed instead of passing an empty string as the label.
#add Binding new ViewportConfig_OnePane command SLevelViewport::BindViewCommands to OnSetViewportConfiguration().
#change SLevelViewport::GetMaximizeToggleVisibility() & OnToggleMaximize - check IsMaximizeSupported() on the parent layout and always collapses the maximize button when it isn't supported.
#add FLevelViewportTabContent::ConstructViewportLayoutByTypeName() - added LevelViewportConfigurationNames::OnePane.
#add New icon image for one pane button in the menu - ViewportLayout_OnePane.png
#add New layout class FLevelViewportLayoutOnePane - based on other layouts but simpler because it only contains one pane and therefore no splitters.

diff --git a/Engine/Source/Editor/EditorStyle/Private/SlateEditorStyle.cpp b/Engine/Source/Editor/EditorStyle/Private/SlateEditorStyle.cpp
--- a/Engine/Source/Editor/EditorStyle/Private/SlateEditorStyle.cpp
+++ b/Engine/Source/Editor/EditorStyle/Private/SlateEditorStyle.cpp
@@ -3392,9 +3392,11 @@ void FSlateEditorStyle::FStyle::SetupLevelEditorStyle()
 
 	// Level viewport layout command icons
 	{
-		const FVector2D IconLayoutSize(49.0f, 37.0f);
-		const FVector2D IconLayoutSizeSmall(49.0f, 37.0f);		// small version set to same size as these are in their own menu and don't clutter the UI
+		const FVector2D IconLayoutSize(47.0f, 37.0f);
+		const FVector2D IconLayoutSizeSmall(47.0f, 37.0f);		// small version set to same size as these are in their own menu and don't clutter the UI
 
+		Set( "LevelViewport.ViewportConfig_OnePane", new IMAGE_BRUSH("Icons/ViewportLayout_OnePane", IconLayoutSize) );
+		Set( "LevelViewport.ViewportConfig_OnePane.Small", new IMAGE_BRUSH("Icons/ViewportLayout_OnePane", IconLayoutSizeSmall) );
 		Set( "LevelViewport.ViewportConfig_TwoPanesH", new IMAGE_BRUSH( "Icons/ViewportLayout_TwoPanesHoriz", IconLayoutSize ) );
 		Set( "LevelViewport.ViewportConfig_TwoPanesH.Small", new IMAGE_BRUSH( "Icons/ViewportLayout_TwoPanesHoriz", IconLayoutSizeSmall ) );
 		Set( "LevelViewport.ViewportConfig_TwoPanesV", new IMAGE_BRUSH( "Icons/ViewportLayout_TwoPanesVert", IconLayoutSize ) );
diff --git a/Engine/Source/Editor/LevelEditor/Private/LevelViewportActions.cpp b/Engine/Source/Editor/LevelEditor/Private/LevelViewportActions.cpp
--- a/Engine/Source/Editor/LevelEditor/Private/LevelViewportActions.cpp
+++ b/Engine/Source/Editor/LevelEditor/Private/LevelViewportActions.cpp
@@ -25,6 +25,7 @@ void FLevelViewportCommands::RegisterCommands()
 	UI_COMMAND( ActorLockSelected, "Lock Selected Actor", "Lock the viewport's position and orientation to the selected actor.", EUserInterfaceActionType::Button, FInputGesture() );
 	UI_COMMAND( ActorUnlockSelected, "Unlock Selected Actor", "Unlock the viewport's position and orientation from the selected actor.", EUserInterfaceActionType::Check, FInputGesture() );
 
+	UI_COMMAND( ViewportConfig_OnePane, "Layout One Pane", "Changes the viewport arrangement to one pane", EUserInterfaceActionType::ToggleButton, FInputGesture() );
 	UI_COMMAND( ViewportConfig_TwoPanesH, "Layout Two Panes (horizontal)", "Changes the viewport arrangement to two panes, side-by-side", EUserInterfaceActionType::ToggleButton, FInputGesture() );
 	UI_COMMAND( ViewportConfig_TwoPanesV, "Layout Two Panes (vertical)", "Changes the viewport arrangement to two panes, one above the other", EUserInterfaceActionType::ToggleButton, FInputGesture() );
 	UI_COMMAND( ViewportConfig_ThreePanesLeft, "Layout Three Panes (one left, two right)", "Changes the viewport arrangement to three panes, one on the left, two on the right", EUserInterfaceActionType::ToggleButton, FInputGesture() );
diff --git a/Engine/Source/Editor/LevelEditor/Private/LevelViewportLayout.cpp b/Engine/Source/Editor/LevelEditor/Private/LevelViewportLayout.cpp
--- a/Engine/Source/Editor/LevelEditor/Private/LevelViewportLayout.cpp
+++ b/Engine/Source/Editor/LevelEditor/Private/LevelViewportLayout.cpp
@@ -128,6 +128,7 @@ FLevelViewportLayout::FLevelViewportLayout()
 	: bIsTransitioning( false ),
 	  bIsReplacement( false ),
 	  bIsQueryingLayoutMetrics( false ),
+	  bIsMaximizeSupported( true ),
 	  bIsMaximized( false ),
 	  bIsImmersive( false ),
 	  bWasMaximized( false ),
@@ -216,7 +217,7 @@ void  FLevelViewportLayout::EndThrottleForAnimatedResize()
 
 void FLevelViewportLayout::InitCommonLayoutFromString( const FString& SpecificLayoutString )
 {
-	bool bShouldBeMaximized = ViewportLayoutDefs::bDefaultShouldBeMaximized;
+	bool bShouldBeMaximized = bIsMaximizeSupported && ViewportLayoutDefs::bDefaultShouldBeMaximized;
 	bool bShouldBeImmersive = ViewportLayoutDefs::bDefaultShouldBeImmersive;
 	int32 MaximizedViewportID = ViewportLayoutDefs::DefaultMaximizedViewportID;
 	if (!SpecificLayoutString.IsEmpty())
@@ -228,7 +229,7 @@ void FLevelViewportLayout::InitCommonLayoutFromString( const FString& SpecificLa
 		GConfig->GetInt(*IniSection, *(SpecificLayoutString + TEXT(".MaximizedViewportID")), MaximizedViewportID, GEditorUserSettingsIni);
 	}
 	// Replacement layouts (those selected by the user via a command) don't start maximized so the layout can be seen clearly.
-	if (!bIsReplacement && MaximizedViewportID >= 0 && MaximizedViewportID < Viewports.Num() && (bShouldBeMaximized || bShouldBeImmersive))
+	if (!bIsReplacement && bIsMaximizeSupported && MaximizedViewportID >= 0 && MaximizedViewportID < Viewports.Num() && (bShouldBeMaximized || bShouldBeImmersive))
 	{
 		TSharedRef<SLevelViewport> LevelViewport = Viewports[MaximizedViewportID].ToSharedRef();
 		// we are not toggling maximize or immersive state but setting it directly
@@ -255,7 +256,7 @@ void FLevelViewportLayout::SaveCommonLayoutString( const FString& SpecificLayout
 		}
 	}
 
-	GConfig->SetBool(*IniSection, *(SpecificLayoutString + TEXT(".bIsMaximized")), bIsMaximized, GEditorUserSettingsIni);
+	GConfig->SetBool(*IniSection, *(SpecificLayoutString + TEXT(".bIsMaximized")), bIsMaximizeSupported && bIsMaximized, GEditorUserSettingsIni);
 	GConfig->SetBool(*IniSection, *(SpecificLayoutString + TEXT(".bIsImmersive")), bIsImmersive, GEditorUserSettingsIni);
 	GConfig->SetInt(*IniSection, *(SpecificLayoutString + TEXT(".MaximizedViewportID")), MaximizedViewportID, GEditorUserSettingsIni);
 }
@@ -284,6 +285,7 @@ void FLevelViewportLayout::MaximizeViewport( TSharedRef<SLevelViewport> Viewport
 {
 	// Should never get into a situation where the viewport is being maximized and there is already a maximized viewport. 
 	// I.E Maximized viewport is NULL which means this is a new maximize or MaximizeViewport is equal to the passed in one which means this is a restore of the current maximized viewport
+	check( bIsMaximizeSupported );
 	check( Viewports.Contains( ViewportToMaximize ) );
 	check( !MaximizedViewport.IsValid() || MaximizedViewport == ViewportToMaximize );
 
diff --git a/Engine/Source/Editor/LevelEditor/Private/LevelViewportLayout.h b/Engine/Source/Editor/LevelEditor/Private/LevelViewportLayout.h
--- a/Engine/Source/Editor/LevelEditor/Private/LevelViewportLayout.h
+++ b/Engine/Source/Editor/LevelEditor/Private/LevelViewportLayout.h
@@ -60,6 +60,13 @@ public:
 	 */
 	bool IsLevelViewportVisible( const class SLevelViewport& InViewport ) const;
 
+	/**
+	* Checks to see if the specified level viewport supports maximizing one pane
+	*
+	* @return true if the viewport supports maximizing
+	*/
+	bool IsMaximizeSupported() const { return bIsMaximizeSupported; }
+
 	/** 
 	 * Checks to see if the specified level viewport is maximized
 	 *
@@ -204,6 +211,9 @@ protected:
 	/** Temporarily set to true while we are querying layout metrics and want all widgets to be visible */
 	bool bIsQueryingLayoutMetrics;
 
+	/** True if the layout supports maximizing one viewport, false if the feature is disabled  */
+	bool bIsMaximizeSupported;
+
 	/** True if we're currently maximized */
 	bool bIsMaximized;
 
diff --git a/Engine/Source/Editor/LevelEditor/Private/LevelViewportLayoutOnePane.cpp b/Engine/Source/Editor/LevelEditor/Private/LevelViewportLayoutOnePane.cpp
new file mode 100644
--- /dev/null
+++ b/Engine/Source/Editor/LevelEditor/Private/LevelViewportLayoutOnePane.cpp
@@ -0,0 +1,73 @@
+// Copyright 1998-2014 Epic Games, Inc. All Rights Reserved.
+
+#include "LevelEditor.h"
+#include "LevelViewportLayoutOnePane.h"
+#include "SLevelViewport.h"
+
+// FLevelViewportLayoutOnePane /////////////////////////////
+
+void FLevelViewportLayoutOnePane::SaveLayoutString(const FString& LayoutString) const
+{
+	if (!bIsTransitioning)
+	{
+		FString SpecificLayoutString = GetTypeSpecificLayoutString(LayoutString);
+
+		SaveCommonLayoutString(SpecificLayoutString);
+	}
+}
+
+TSharedRef<SWidget> FLevelViewportLayoutOnePane::MakeViewportLayout(const FString& LayoutString)
+{
+	// single viewport layout blocks maximize feature as it doesn't make sense
+	bIsMaximizeSupported = false;
+
+	FString SpecificLayoutString = GetTypeSpecificLayoutString(LayoutString);
+
+	FEngineShowFlags OrthoShowFlags(ESFIM_Editor);
+	ApplyViewMode(VMI_BrushWireframe, false, OrthoShowFlags);
+
+	FEngineShowFlags PerspectiveShowFlags(ESFIM_Editor);
+	ApplyViewMode(VMI_Lit, true, PerspectiveShowFlags);
+
+	FString ViewportKey;
+	if (!SpecificLayoutString.IsEmpty())
+	{
+		ViewportKey = SpecificLayoutString + TEXT(".Viewport0");
+	}
+
+	TSharedPtr<SLevelViewport> ViewportWidget;
+
+	ViewportBox =
+		SNew(SHorizontalBox)
+		+ SHorizontalBox::Slot()
+		[
+			SAssignNew(ViewportWidget, SLevelViewport)
+			.ViewportType(LVT_Perspective)
+			.Realtime(true)
+			.ParentLayout(AsShared())
+			.ParentLevelEditor(ParentLevelEditor)
+			.ConfigKey(ViewportKey)
+		];
+
+	Viewports.Add(ViewportWidget);
+
+	// Make newly-created perspective viewports active by default
+	GCurrentLevelEditingViewportClient = &ViewportWidget->GetLevelViewportClient();
+
+	InitCommonLayoutFromString(SpecificLayoutString);
+
+	return StaticCastSharedRef<SWidget>(ViewportBox.ToSharedRef());
+}
+
+void FLevelViewportLayoutOnePane::ReplaceWidget(TSharedRef<SWidget> Source, TSharedRef<SWidget> Replacement)
+{
+	check(ViewportBox->GetChildren()->Num() == 1)
+	TSharedRef<SWidget> ViewportWidget = ViewportBox->GetChildren()->GetChildAt(0);
+
+	check(ViewportWidget == Source);
+	ViewportBox->RemoveSlot(Source);
+	ViewportBox->AddSlot()
+		[
+			Replacement
+		];
+}
diff --git a/Engine/Source/Editor/LevelEditor/Private/LevelViewportLayoutOnePane.h b/Engine/Source/Editor/LevelEditor/Private/LevelViewportLayoutOnePane.h
new file mode 100644
--- /dev/null
+++ b/Engine/Source/Editor/LevelEditor/Private/LevelViewportLayoutOnePane.h
@@ -0,0 +1,34 @@
+// Copyright 1998-2014 Epic Games, Inc. All Rights Reserved.
+
+
+#ifndef __LevelViewportLayoutOnePane_h__
+#define __LevelViewportLayoutOnePane_h__
+
+#pragma once
+
+#include "LevelViewportLayout.h"
+
+class FLevelViewportLayoutOnePane : public FLevelViewportLayout
+{
+public:
+	/**
+	* Saves viewport layout information between editor sessions
+	*/
+	virtual void SaveLayoutString(const FString& LayoutString) const OVERRIDE;
+
+	virtual const FName& GetLayoutTypeName() const OVERRIDE{ return LevelViewportConfigurationNames::OnePane; }
+protected:
+	/**
+	* Creates the viewport for the single pane
+	*/
+	virtual TSharedRef<SWidget> MakeViewportLayout(const FString& LayoutString) OVERRIDE;
+
+	/** Overridden from FLevelViewportLayout */
+	virtual void ReplaceWidget(TSharedRef< SWidget > Source, TSharedRef< SWidget > Replacement) OVERRIDE;
+
+protected:
+	/** The viewport widget parent box */
+	TSharedPtr< SHorizontalBox > ViewportBox;
+};
+
+#endif
\ No newline at end of file
diff --git a/Engine/Source/Editor/LevelEditor/Private/LevelViewportTabContent.cpp b/Engine/Source/Editor/LevelEditor/Private/LevelViewportTabContent.cpp
--- a/Engine/Source/Editor/LevelEditor/Private/LevelViewportTabContent.cpp
+++ b/Engine/Source/Editor/LevelEditor/Private/LevelViewportTabContent.cpp
@@ -27,6 +27,7 @@ TSharedPtr< class FLevelViewportLayout > FLevelViewportTabContent::ConstructView
 	else if (TypeName == LevelViewportConfigurationNames::FourPanesRight) ViewportLayout = MakeShareable(new FLevelViewportLayoutFourPanesRight);
 	else if (TypeName == LevelViewportConfigurationNames::FourPanesBottom) ViewportLayout = MakeShareable(new FLevelViewportLayoutFourPanesBottom);
 	else if (TypeName == LevelViewportConfigurationNames::FourPanesTop) ViewportLayout = MakeShareable(new FLevelViewportLayoutFourPanesTop);
+	else if (TypeName == LevelViewportConfigurationNames::OnePane) ViewportLayout = MakeShareable(new FLevelViewportLayoutOnePane);
 
 	check (ViewportLayout.IsValid());
 	ViewportLayout->SetIsReplacement(bSwitchingLayouts);
diff --git a/Engine/Source/Editor/LevelEditor/Private/SLevelViewport.cpp b/Engine/Source/Editor/LevelEditor/Private/SLevelViewport.cpp
--- a/Engine/Source/Editor/LevelEditor/Private/SLevelViewport.cpp
+++ b/Engine/Source/Editor/LevelEditor/Private/SLevelViewport.cpp
@@ -1049,6 +1049,12 @@ void SLevelViewport::BindViewCommands( FUICommandList& CommandList )
 		FIsActionChecked::CreateSP( this, &SLevelViewport::IsSelectedActorLocked )
 		);
 
+	CommandList.MapAction(
+		ViewportActions.ViewportConfig_OnePane,
+		FExecuteAction::CreateSP(this, &SLevelViewport::OnSetViewportConfiguration, LevelViewportConfigurationNames::OnePane),
+		FCanExecuteAction(),
+		FIsActionChecked::CreateSP(this, &SLevelViewport::IsViewportConfigurationSet, LevelViewportConfigurationNames::OnePane));
+
 	CommandList.MapAction(
 		ViewportActions.ViewportConfig_TwoPanesH,
 		FExecuteAction::CreateSP( this, &SLevelViewport::OnSetViewportConfiguration, LevelViewportConfigurationNames::TwoPanesHoriz ),
@@ -1296,8 +1302,15 @@ EVisibility SLevelViewport::GetToolBarVisibility() const
 
 EVisibility SLevelViewport::GetMaximizeToggleVisibility() const
 {
+	bool bIsMaximizeSupported = false;
+	TSharedPtr<FLevelViewportLayout> LayoutPinned = ParentLayout.Pin();
+	if (LayoutPinned.IsValid())
+	{
+		bIsMaximizeSupported = LayoutPinned->IsMaximizeSupported();
+	}
+
 	// Do not show the maximize/minimize toggle when in immersive mode
-	return IsImmersive() ? EVisibility::Collapsed : EVisibility::Visible;
+	return (!bIsMaximizeSupported || IsImmersive()) ? EVisibility::Collapsed : EVisibility::Visible;
 }
 
 EVisibility SLevelViewport::GetCloseImmersiveButtonVisibility() const
@@ -2042,7 +2055,8 @@ void  SLevelViewport::OnToggleMaximizeMode()
 
 FReply SLevelViewport::OnToggleMaximize()
 {
-	if( ParentLayout.IsValid() ) 
+	TSharedPtr<FLevelViewportLayout> ParentLayoutPinned = ParentLayout.Pin();
+	if (ParentLayoutPinned.IsValid() && ParentLayoutPinned->IsMaximizeSupported())
 	{
 		bool bWantImmersive = IsImmersive();
 		bool bWantMaximize = IsMaximized();
@@ -2060,7 +2074,7 @@ FReply SLevelViewport::OnToggleMaximize()
 		// We always want to animate in response to user-interactive toggling of maximized state
 		const bool bAllowAnimation = true;
 
-		ParentLayout.Pin()->RequestMaximizeViewport( SharedThis( this ), bWantMaximize, bWantImmersive, bAllowAnimation );
+		ParentLayoutPinned->RequestMaximizeViewport(SharedThis(this), bWantMaximize, bWantImmersive, bAllowAnimation);
 	}
 	return FReply::Handled();
 }
diff --git a/Engine/Source/Editor/LevelEditor/Private/SLevelViewportToolBar.cpp b/Engine/Source/Editor/LevelEditor/Private/SLevelViewportToolBar.cpp
--- a/Engine/Source/Editor/LevelEditor/Private/SLevelViewportToolBar.cpp
+++ b/Engine/Source/Editor/LevelEditor/Private/SLevelViewportToolBar.cpp
@@ -711,9 +711,35 @@ void SLevelViewportToolBar::GenerateViewportConfigsMenu(FMenuBuilder& MenuBuilde
 	check (Viewport.IsValid());
 	TSharedPtr<FUICommandList> CommandList = Viewport.Pin()->GetCommandList();
 
+	MenuBuilder.BeginSection("LevelViewportOnePaneConfigs", LOCTEXT("OnePaneConfigHeader", "One Pane"));
+	{
+		FToolBarBuilder OnePaneButton(CommandList, FMultiBoxCustomization::None);
+		OnePaneButton.SetLabelVisibility(EVisibility::Collapsed);
+		OnePaneButton.SetStyle(&FEditorStyle::Get(), "ViewportLayoutToolbar");
+
+		OnePaneButton.AddToolBarButton(FLevelViewportCommands::Get().ViewportConfig_OnePane);
+
+		MenuBuilder.AddWidget(
+			SNew(SHorizontalBox)
+			+ SHorizontalBox::Slot()
+			.AutoWidth()
+			[
+				OnePaneButton.MakeWidget()
+			]
+			+ SHorizontalBox::Slot()
+			.FillWidth(1)
+			[
+				SNullWidget::NullWidget
+			],
+			FText::GetEmpty(), true
+			);
+	}
+	MenuBuilder.EndSection();
+
 	MenuBuilder.BeginSection("LevelViewportTwoPaneConfigs", LOCTEXT("TwoPaneConfigHeader", "Two Panes") );
 	{
 		FToolBarBuilder TwoPaneButtons(CommandList, FMultiBoxCustomization::None);
+		TwoPaneButtons.SetLabelVisibility(EVisibility::Collapsed);
 		TwoPaneButtons.SetStyle(&FEditorStyle::Get(), "ViewportLayoutToolbar");
 
 		TwoPaneButtons.AddToolBarButton(FLevelViewportCommands::Get().ViewportConfig_TwoPanesH, NAME_None, FText());
@@ -739,6 +765,7 @@ void SLevelViewportToolBar::GenerateViewportConfigsMenu(FMenuBuilder& MenuBuilde
 	MenuBuilder.BeginSection("LevelViewportThreePaneConfigs", LOCTEXT("ThreePaneConfigHeader", "Three Panes") );
 	{
 		FToolBarBuilder ThreePaneButtons(CommandList, FMultiBoxCustomization::None);
+		ThreePaneButtons.SetLabelVisibility(EVisibility::Collapsed);
 		ThreePaneButtons.SetStyle(&FEditorStyle::Get(), "ViewportLayoutToolbar");
 
 		ThreePaneButtons.AddToolBarButton(FLevelViewportCommands::Get().ViewportConfig_ThreePanesLeft, NAME_None, FText());
@@ -766,6 +793,7 @@ void SLevelViewportToolBar::GenerateViewportConfigsMenu(FMenuBuilder& MenuBuilde
 	MenuBuilder.BeginSection("LevelViewportFourPaneConfigs", LOCTEXT("FourPaneConfigHeader", "Four Panes") );
 	{
 		FToolBarBuilder FourPaneButtons(CommandList, FMultiBoxCustomization::None);
+		FourPaneButtons.SetLabelVisibility(EVisibility::Collapsed);
 		FourPaneButtons.SetStyle(&FEditorStyle::Get(), "ViewportLayoutToolbar");
 
 		FourPaneButtons.AddToolBarButton(FLevelViewportCommands::Get().ViewportConfig_FourPanes2x2, NAME_None, FText());
diff --git a/Engine/Source/Editor/LevelEditor/Public/LevelViewportActions.h b/Engine/Source/Editor/LevelEditor/Public/LevelViewportActions.h
--- a/Engine/Source/Editor/LevelEditor/Public/LevelViewportActions.h
+++ b/Engine/Source/Editor/LevelEditor/Public/LevelViewportActions.h
@@ -21,6 +21,7 @@ namespace LevelViewportConfigurationNames
 	static FName FourPanesTop("FourPanesTop");
 	static FName FourPanesBottom("FourPanesBottom");
 	static FName FourPanes2x2("FourPanes2x2");
+	static FName OnePane("OnePane");
 }
 
 /**
@@ -140,6 +141,7 @@ public:
 	TSharedPtr< FUICommandInfo > ActorUnlockSelected;
 
 	/** Viewport pane configurations */
+	TSharedPtr< FUICommandInfo > ViewportConfig_OnePane;
 	TSharedPtr< FUICommandInfo > ViewportConfig_TwoPanesH;
 	TSharedPtr< FUICommandInfo > ViewportConfig_TwoPanesV;
 	TSharedPtr< FUICommandInfo > ViewportConfig_ThreePanesLeft;
