Fix non-unity build

diff --git a/Engine/Source/Editor/LevelEditor/Private/LevelViewportTabContent.cpp b/Engine/Source/Editor/LevelEditor/Private/LevelViewportTabContent.cpp
--- a/Engine/Source/Editor/LevelEditor/Private/LevelViewportTabContent.cpp
+++ b/Engine/Source/Editor/LevelEditor/Private/LevelViewportTabContent.cpp
@@ -3,6 +3,7 @@
 #include "LevelEditor.h"
 #include "LevelViewportTabContent.h"
 #include "LevelViewportLayout2x2.h"
+#include "LevelViewportLayoutOnePane.h"
 #include "LevelViewportLayoutTwoPanes.h"
 #include "LevelViewportLayoutThreePanes.h"
 #include "LevelViewportLayoutFourPanes.h"
