Correct a copy and paste mistake that caused the translation for "45" to be overridden with "90"
#ttp 333401 EDITOR: 45 degree rotation button has incorrect label
#proj Editor
#branch UE4

diff --git a/Engine/Source/Editor/LevelEditor/Private/SSurfaceProperties.cpp b/Engine/Source/Editor/LevelEditor/Private/SSurfaceProperties.cpp
--- a/Engine/Source/Editor/LevelEditor/Private/SSurfaceProperties.cpp
+++ b/Engine/Source/Editor/LevelEditor/Private/SSurfaceProperties.cpp
@@ -276,7 +276,7 @@ TSharedRef<SWidget> SSurfaceProperties::ConstructTextureRotate()
 	];
 
 	// Rotation button fields:
-	const static FText ButtonFields[] = { LOCTEXT("Rotate45Degrees", "45"), LOCTEXT("Rotate45Degrees", "90"), LOCTEXT("RotateCustom", "---") };
+	const static FText ButtonFields[] = { LOCTEXT("Rotate45Degrees", "45"), LOCTEXT("Rotate90Degrees", "90"), LOCTEXT("RotateCustom", "---") };
 	const static int32 RotationValues[] = {45, 90, -1};
 	const static RotationAction RotationActions[] = {Rotate,Rotate,RotateCustom};
 	
