Fixed placed classes appearing in the level editor menu as "Class Class"

The actor class type name is now preferred over the factory name in the case of classes.

diff --git a/Engine/Source/Editor/LevelEditor/Private/LevelEditorCreateActorMenu.cpp b/Engine/Source/Editor/LevelEditor/Private/LevelEditorCreateActorMenu.cpp
--- a/Engine/Source/Editor/LevelEditor/Private/LevelEditorCreateActorMenu.cpp
+++ b/Engine/Source/Editor/LevelEditor/Private/LevelEditorCreateActorMenu.cpp
@@ -86,7 +86,12 @@ class SAssetMenuEntry : public SCompoundWidget
 			if ( MenuItem.FactoryToUse != NULL )
 			{
 				DefaultActor = MenuItem.FactoryToUse->GetDefaultActor( MenuItem.AssetData );
-				ActorTypeDisplayName = MenuItem.FactoryToUse->GetDisplayName();
+
+				// Prefer the class type name set above over the factory's display name
+				if (ActorTypeDisplayName.IsEmpty())
+				{
+					ActorTypeDisplayName = MenuItem.FactoryToUse->GetDisplayName();
+				}
 
 				FName BrushName = *FString::Printf( TEXT("ClassIcon.%s"), *MenuItem.FactoryToUse->GetClass()->GetName() );
 				IconBrush = FEditorStyle::GetOptionalBrush(BrushName, nullptr, nullptr);
