Quick fix for unusual CamOrthoZoom bug, where the editor crashes upon level load due to 0.0 zoom value.

diff --git a/Engine/Source/Editor/LevelEditor/Private/SLevelViewport.cpp b/Engine/Source/Editor/LevelEditor/Private/SLevelViewport.cpp
--- a/Engine/Source/Editor/LevelEditor/Private/SLevelViewport.cpp
+++ b/Engine/Source/Editor/LevelEditor/Private/SLevelViewport.cpp
@@ -29,7 +29,7 @@
 #include "SceneOutlinerTreeItems.h"
 
 static const FName LevelEditorName("LevelEditor");
-
+#pragma optimize("", off)
 
 #define LOCTEXT_NAMESPACE "LevelViewport"
 
@@ -829,6 +829,11 @@ void SLevelViewport::OnMapChanged( UWorld* World, EMapChangeType::Type MapChange
 	{
 		if( MapChangeType == EMapChangeType::LoadMap )
 		{
+			if (World->EditorViews[LevelViewportClient->ViewportType].CamOrthoZoom == 0.0f)
+			{
+				World->EditorViews[LevelViewportClient->ViewportType].CamOrthoZoom = DEFAULT_ORTHOZOOM;
+			}
+
 			LevelViewportClient->SetInitialViewTransform(
 				World->EditorViews[LevelViewportClient->ViewportType].CamPosition,
 				World->EditorViews[LevelViewportClient->ViewportType].CamRotation,
