* added ... to tell user where a dialog opens

diff --git a/Engine/Source/Editor/LevelEditor/Private/LevelViewportActions.cpp b/Engine/Source/Editor/LevelEditor/Private/LevelViewportActions.cpp
--- a/Engine/Source/Editor/LevelEditor/Private/LevelViewportActions.cpp
+++ b/Engine/Source/Editor/LevelEditor/Private/LevelViewportActions.cpp
@@ -28,7 +28,7 @@ void FLevelViewportCommands::RegisterCommands()
 	UI_COMMAND( ToggleImmersive, "Immersive Mode", "Switches this viewport between immersive mode and regular mode", EUserInterfaceActionType::ToggleButton, FInputGesture( EKeys::F11 ) );
 	
 	UI_COMMAND( CreateCamera, "Create Camera Here", "Creates a new camera actor at the current location of this viewport's camera", EUserInterfaceActionType::Button, FInputGesture() );
-	UI_COMMAND( HighResScreenshot, "High Resolution Screenshot", "Opens the control panel for high resolution screenshots", EUserInterfaceActionType::Button, FInputGesture() );
+	UI_COMMAND( HighResScreenshot, "High Resolution Screenshot...", "Opens the control panel for high resolution screenshots", EUserInterfaceActionType::Button, FInputGesture() );
 	
 	UI_COMMAND( UseDefaultShowFlags, "Use Defaults", "Resets all show flags to default", EUserInterfaceActionType::Button, FInputGesture() );
 
