The move/save pivot options for brushes are now only available for ortho viewports

TTP# 333271 - UE4 Editor - "Move Pivot" brush functionality not working

They don't work for perspective viewports, however you need to be able to reset the pivot point in any viewport type.

diff --git a/Engine/Source/Editor/LevelEditor/Private/LevelEditorContextMenu.cpp b/Engine/Source/Editor/LevelEditor/Private/LevelEditorContextMenu.cpp
--- a/Engine/Source/Editor/LevelEditor/Private/LevelEditorContextMenu.cpp
+++ b/Engine/Source/Editor/LevelEditor/Private/LevelEditorContextMenu.cpp
@@ -325,11 +325,19 @@ void FLevelEditorContextMenu::FillMenu( FMenuBuilder& MenuBuilder, TWeakPtr<SLev
 		// @todo UE4: The current pivot options only work for brushes
 		if( SelectionInfo.bHaveBrush )
 		{
-			// Add a sub-menu for "Pivot"
-			MenuBuilder.AddSubMenu( 
-				LOCTEXT("PivotSubMenu", "Pivot"), 
-				LOCTEXT("PivotSubMenu_ToolTip", "Actor pivoting utils"),
-				FNewMenuDelegate::CreateStatic( &FLevelEditorContextMenuImpl::FillPivotMenu ) );
+			// You can only move the pivot in ortho viewports, but you can reset it in any viewport
+			if( GCurrentLevelEditingViewportClient->ViewportType != LVT_Perspective )
+			{
+				// Add a sub-menu for "Pivot"
+				MenuBuilder.AddSubMenu( 
+					LOCTEXT("PivotSubMenu", "Pivot"), 
+					LOCTEXT("PivotSubMenu_ToolTip", "Actor pivoting utils"),
+					FNewMenuDelegate::CreateStatic( &FLevelEditorContextMenuImpl::FillPivotMenu ) );
+			}
+			else
+			{
+				MenuBuilder.AddMenuEntry( FLevelEditorCommands::Get().ResetPivot );
+			}
 		}
 
 		if (GetDefault<UEditorExperimentalSettings>()->bActorMerging && 
