Fixed crash after recompiling level editor.
#ttp 335187 - CRITICAL: EDITOR: CRASH: Recompiling level editor
#branch UE4
#proj Editor.LevelEditor
#add Added calls in FLevelEditorModule::ShutdownModule to unregister the FEditorViewportCommands and FLevelViewportCommands so that they are re-registered correctly upon ModuleStart again.

diff --git a/Engine/Source/Editor/LevelEditor/Private/LevelEditor.cpp b/Engine/Source/Editor/LevelEditor/Private/LevelEditor.cpp
--- a/Engine/Source/Editor/LevelEditor/Private/LevelEditor.cpp
+++ b/Engine/Source/Editor/LevelEditor/Private/LevelEditor.cpp
@@ -194,6 +194,8 @@ void FLevelEditorModule::ShutdownModule()
 
 	FLevelEditorCommands::Unregister();
 	FLevelEditorModesCommands::Unregister();
+	FEditorViewportCommands::Unregister();
+	FLevelViewportCommands::Unregister();
 }
 
 void FLevelEditorModule::PreUnloadCallback()
