UE4: Added open level asset pickers to replace the existing open level dialog. The existing dialog is still accessible via File->Open Level->Open Other Level...

diff --git a/Engine/Source/Editor/LevelEditor/Private/LevelEditor.cpp b/Engine/Source/Editor/LevelEditor/Private/LevelEditor.cpp
--- a/Engine/Source/Editor/LevelEditor/Private/LevelEditor.cpp
+++ b/Engine/Source/Editor/LevelEditor/Private/LevelEditor.cpp
@@ -492,7 +492,15 @@ void FLevelEditorModule::BindGlobalLevelEditorCommands()
 	ActionList.MapAction( Commands.BrowseAPIReference, FExecuteAction::CreateStatic( &FLevelEditorActionCallbacks::BrowseAPIReference ) );
 	ActionList.MapAction( Commands.BrowseViewportControls, FExecuteAction::CreateStatic( &FLevelEditorActionCallbacks::BrowseViewportControls ) );
 	ActionList.MapAction( Commands.NewLevel, FExecuteAction::CreateStatic( &FLevelEditorActionCallbacks::NewLevel ), FCanExecuteAction::CreateStatic( &FLevelEditorActionCallbacks::NewLevel_CanExecute ) );
-	ActionList.MapAction( Commands.OpenLevel, FExecuteAction::CreateStatic( &FLevelEditorActionCallbacks::OpenLevel ), FCanExecuteAction::CreateStatic( &FLevelEditorActionCallbacks::OpenLevel_CanExecute ) );
+	if (FParse::Param(FCommandLine::Get(), TEXT("WorldAssets")))
+	{
+		ActionList.MapAction( Commands.OpenLevel, FExecuteAction::CreateStatic( &FLevelEditorActionCallbacks::OpenLevelPickingDialog ), FCanExecuteAction::CreateStatic( &FLevelEditorActionCallbacks::OpenLevel_CanExecute ) );
+		ActionList.MapAction( Commands.LegacyOpenLevel, FExecuteAction::CreateStatic( &FLevelEditorActionCallbacks::OpenLevel ), FCanExecuteAction::CreateStatic( &FLevelEditorActionCallbacks::OpenLevel_CanExecute ) );
+	}
+	else
+	{
+		ActionList.MapAction( Commands.OpenLevel, FExecuteAction::CreateStatic( &FLevelEditorActionCallbacks::OpenLevel ), FCanExecuteAction::CreateStatic( &FLevelEditorActionCallbacks::OpenLevel_CanExecute ) );
+	}
 	ActionList.MapAction( Commands.Save, FExecuteAction::CreateStatic( &FLevelEditorActionCallbacks::Save ) );
 	ActionList.MapAction( Commands.SaveAs, FExecuteAction::CreateStatic( &FLevelEditorActionCallbacks::SaveAs ), DefaultExecuteAction );
 	ActionList.MapAction( Commands.SaveAllLevels, FExecuteAction::CreateStatic( &FLevelEditorActionCallbacks::SaveAllLevels ) );
diff --git a/Engine/Source/Editor/LevelEditor/Private/LevelEditorActions.cpp b/Engine/Source/Editor/LevelEditor/Private/LevelEditorActions.cpp
--- a/Engine/Source/Editor/LevelEditor/Private/LevelEditorActions.cpp
+++ b/Engine/Source/Editor/LevelEditor/Private/LevelEditorActions.cpp
@@ -180,6 +180,91 @@ bool FLevelEditorActionCallbacks::OpenLevel_CanExecute()
 	return FSlateApplication::Get().IsNormalExecution() && !GEditorModeTools().IsTracking();
 }
 
+FAssetPickerConfig FLevelEditorActionCallbacks::CreateLevelAssetPickerConfig()
+{
+	FAssetPickerConfig AssetPickerConfig;
+	AssetPickerConfig.Filter.ClassNames.Add(UWorld::StaticClass()->GetFName());
+	AssetPickerConfig.OnAssetsActivated = FOnAssetsActivated::CreateStatic(&FLevelEditorActionCallbacks::OpenLevelFromAssetPicker);
+	AssetPickerConfig.bAllowDragging = false;
+	AssetPickerConfig.SelectionMode = ESelectionMode::Single;
+	AssetPickerConfig.InitialAssetViewType = EAssetViewType::List;
+	UWorld* CurrentWorld = GEditor->GetEditorWorldContext().World();
+	if (CurrentWorld)
+	{
+		AssetPickerConfig.InitialAssetSelection = FAssetData(CurrentWorld);
+	}
+
+	return AssetPickerConfig;
+}
+
+void FLevelEditorActionCallbacks::OpenLevelPickingDialog()
+{
+	const FVector2D AssetPickerSize(600.0f, 586.0f);
+
+	FMenuBuilder MenuBuilder(false, NULL);
+
+	// Create the contents of the popup
+	FContentBrowserModule& ContentBrowserModule = FModuleManager::Get().LoadModuleChecked<FContentBrowserModule>(TEXT("ContentBrowser"));
+	FAssetPickerConfig AssetPickerConfig = CreateLevelAssetPickerConfig();
+	TSharedRef<SWidget> ActualWidget = 
+		SNew(SBox)
+		.HeightOverride(AssetPickerSize.X)
+		.WidthOverride(AssetPickerSize.Y)
+		[
+			ContentBrowserModule.Get().CreateAssetPicker(AssetPickerConfig)
+		];
+
+	// Wrap the picker widget in a multibox-style menu body
+	MenuBuilder.BeginSection("AssetPickerOpenLevel", NSLOCTEXT("OpenLevelDialog", "WindowTitle", "Open Level"));
+	{
+		const bool bNoIndent = true;
+		MenuBuilder.AddWidget(ActualWidget, FText::GetEmpty(), bNoIndent);
+	}
+	MenuBuilder.EndSection();
+
+	TSharedRef<SWidget> WindowContents = MenuBuilder.MakeWidget();
+
+	// Determine where the pop-up should open
+	TSharedPtr<SWindow> ParentWindow = FSlateApplication::Get().GetActiveTopLevelWindow();
+	FVector2D WindowPosition = FSlateApplication::Get().GetCursorPos();
+	if (ParentWindow.IsValid())
+	{
+		FSlateRect ParentMonitorRect = ParentWindow->GetFullScreenInfo();
+		const FVector2D MonitorCenter((ParentMonitorRect.Right + ParentMonitorRect.Left) * 0.5f, (ParentMonitorRect.Top + ParentMonitorRect.Bottom) * 0.5f);
+		WindowPosition = MonitorCenter - AssetPickerSize * 0.5f;
+
+		// Open the pop-up
+		FPopupTransitionEffect TransitionEffect(FPopupTransitionEffect::None);
+		TSharedRef<SWindow> PopupWindow = FSlateApplication::Get().PushMenu(ParentWindow.ToSharedRef(), WindowContents, WindowPosition, TransitionEffect);
+	}
+}
+
+void FLevelEditorActionCallbacks::OpenLevelFromAssetPicker(const TArray<class FAssetData>& SelectedAssets, EAssetTypeActivationMethod::Type ActivationType)
+{
+	const bool bCorrectActivationMethod = (ActivationType == EAssetTypeActivationMethod::DoubleClicked || ActivationType == EAssetTypeActivationMethod::Opened);
+	if (SelectedAssets.Num() > 0 && bCorrectActivationMethod)
+	{
+		const FAssetData& AssetData = SelectedAssets[0];
+		if (AssetData.AssetClass == UWorld::StaticClass()->GetFName())
+		{
+			// Close the menu that we were picking from
+			FSlateApplication::Get().DismissAllMenus();
+
+			// If there are any unsaved changes to the current level, see if the user wants to save those first.
+			bool bPromptUserToSave = true;
+			bool bSaveMapPackages = true;
+			bool bSaveContentPackages = true;
+			if (FEditorFileUtils::SaveDirtyPackages(bPromptUserToSave, bSaveMapPackages, bSaveContentPackages))
+			{
+				const FString FileToOpen = FPackageName::LongPackageNameToFilename(AssetData.PackageName.ToString(), FPackageName::GetMapPackageExtension());
+				const bool bLoadAsTemplate = false;
+				const bool bShowProgress = true;
+				FEditorFileUtils::LoadMap(FileToOpen, bLoadAsTemplate, bShowProgress);
+			}
+		}
+	}
+}
+
 void FLevelEditorActionCallbacks::DeltaTransform()
 {
 	if (GUnrealEd->WarnIfLightingBuildIsCurrentlyRunning())
@@ -2485,7 +2570,15 @@ void FLevelEditorCommands::RegisterCommands()
 	UI_COMMAND( BrowseAPIReference, "API Reference...", "Opens the API reference documentation", EUserInterfaceActionType::Button, FInputGesture() );
 	UI_COMMAND( BrowseViewportControls, "Viewport Controls...", "Opens the viewport controls cheat sheet", EUserInterfaceActionType::Button, FInputGesture() );
 	UI_COMMAND( NewLevel, "New Level...", "Create a new level, or choose a level template to start from.", EUserInterfaceActionType::Button, FInputGesture( EModifierKey::Control, EKeys::N ) );
-	UI_COMMAND( OpenLevel, "Open Level...", "Loads an existing level", EUserInterfaceActionType::Button, FInputGesture( EModifierKey::Control, EKeys::O ) );
+	if (FParse::Param(FCommandLine::Get(), TEXT("WorldAssets")))
+	{
+		UI_COMMAND( OpenLevel, "Open Level...", "Loads an existing level", EUserInterfaceActionType::Button, FInputGesture( EModifierKey::Control, EKeys::O ) );
+		UI_COMMAND( LegacyOpenLevel, "Open Other Level...", "Loads an existing level using the file explorer", EUserInterfaceActionType::Button, FInputGesture() );
+	}
+	else
+	{
+		UI_COMMAND( OpenLevel, "Open Level...", "Loads an existing level", EUserInterfaceActionType::Button, FInputGesture( EModifierKey::Control, EKeys::O ) );
+	}
 	UI_COMMAND( Save, "Save", "Saves the current level to disk", EUserInterfaceActionType::Button, FInputGesture() );
 	UI_COMMAND( SaveAs, "Save As...", "Save the current level as...", EUserInterfaceActionType::Button, FInputGesture( EModifierKey::Shift|EModifierKey::Control, EKeys::S ) );
 	UI_COMMAND( SaveAllLevels, "Save All Levels", "Saves all unsaved levels to disk", EUserInterfaceActionType::Button, FInputGesture() );
diff --git a/Engine/Source/Editor/LevelEditor/Private/LevelEditorMenu.cpp b/Engine/Source/Editor/LevelEditor/Private/LevelEditorMenu.cpp
--- a/Engine/Source/Editor/LevelEditor/Private/LevelEditorMenu.cpp
+++ b/Engine/Source/Editor/LevelEditor/Private/LevelEditorMenu.cpp
@@ -8,6 +8,7 @@
 #include "Editor/MainFrame/Public/MainFrame.h"
 #include "MRUFavoritesList.h"
 #include "GlobalEditorCommonCommands.h"
+#include "Editor/ContentBrowser/Public/ContentBrowserModule.h"
 
 #define LOCTEXT_NAMESPACE "LevelEditorMenu"
 
@@ -129,8 +130,20 @@ TSharedRef< SWidget > FLevelEditorMenu::MakeLevelEditorMenu( const TSharedPtr<FU
 			// New Level
 			MenuBuilder.AddMenuEntry( FLevelEditorCommands::Get().NewLevel );
 
-			// Open Level
-			MenuBuilder.AddMenuEntry( FLevelEditorCommands::Get().OpenLevel );
+			static const bool bUsingWorldAssets = FParse::Param(FCommandLine::Get(), TEXT("WorldAssets"));
+			if ( bUsingWorldAssets )
+			{
+				// Open level menu.
+				MenuBuilder.AddSubMenu(
+					LOCTEXT("OpenLevelSubMenu", "Open Level"),
+					LOCTEXT("OpenLevelSubMenu_ToolTip", "Select a level to load"),
+					FNewMenuDelegate::CreateStatic(&Local::MakeOpenLevelMenu), false, FSlateIcon(FEditorStyle::GetStyleSetName(), "LevelEditor.OpenLevel"));
+			}
+			else
+			{
+				// Open Level
+				MenuBuilder.AddMenuEntry( FLevelEditorCommands::Get().OpenLevel );
+			}
 
 			// Open Asset
 			//@TODO: Doesn't work when summoned from here: MenuBuilder.AddMenuEntry( FGlobalEditorCommonCommands::Get().SummonOpenAssetDialog );
@@ -145,6 +158,33 @@ TSharedRef< SWidget > FLevelEditorMenu::MakeLevelEditorMenu( const TSharedPtr<FU
 			MenuBuilder.AddMenuEntry( FLevelEditorCommands::Get().SaveAllLevels );
 		}
 
+		static void MakeOpenLevelMenu( FMenuBuilder& MenuBuilder )
+		{
+			// Asset picker for a level
+			{
+				FContentBrowserModule& ContentBrowserModule = FModuleManager::Get().LoadModuleChecked<FContentBrowserModule>(TEXT("ContentBrowser"));
+				FAssetPickerConfig AssetPickerConfig = FLevelEditorActionCallbacks::CreateLevelAssetPickerConfig();
+				AssetPickerConfig.ThumbnailScale = 0;
+
+				TSharedRef<SWidget> MenuContent =
+					SNew(SBox)
+					.HeightOverride(600)
+					.WidthOverride(300)
+					[
+						ContentBrowserModule.Get().CreateAssetPicker(AssetPickerConfig)
+					];
+
+				const bool bNoIndent = true;
+				MenuBuilder.AddWidget(MenuContent, FText::GetEmpty(), bNoIndent);
+			}
+
+			MenuBuilder.BeginSection("LegacyOpenLevel", LOCTEXT("LegacyOpenLevelHeading", "Other Levels"));
+			{
+				// Legacy Open Level
+				MenuBuilder.AddMenuEntry( FLevelEditorCommands::Get().LegacyOpenLevel );
+			}
+			MenuBuilder.EndSection();
+		}
 
 		static void FillFavoriteLevelSubMenu( FMenuBuilder& MenuBuilder, const int32 CurFavoriteIndex )
 		{
diff --git a/Engine/Source/Editor/LevelEditor/Public/LevelEditorActions.h b/Engine/Source/Editor/LevelEditor/Public/LevelEditorActions.h
--- a/Engine/Source/Editor/LevelEditor/Public/LevelEditorActions.h
+++ b/Engine/Source/Editor/LevelEditor/Public/LevelEditorActions.h
@@ -6,6 +6,7 @@
 
 #include "IToolkit.h"
 #include "LightmapResRatioAdjust.h"
+#include "Developer/AssetTools/Public/IAssetTypeActions.h"
 
 /**
  * Unreal level editor actions
@@ -39,6 +40,7 @@ public:
 	/** Level file commands */
 	TSharedPtr< FUICommandInfo > NewLevel;
 	TSharedPtr< FUICommandInfo > OpenLevel;
+	TSharedPtr< FUICommandInfo > LegacyOpenLevel;
 	TSharedPtr< FUICommandInfo > Save;
 	TSharedPtr< FUICommandInfo > SaveAs;
 	TSharedPtr< FUICommandInfo > SaveAllLevels;
@@ -589,6 +591,10 @@ public:
 	/** Opens an existing level */
 	static void OpenLevel();
 	static bool OpenLevel_CanExecute();
+	static struct FAssetPickerConfig CreateLevelAssetPickerConfig();
+	static void OpenLevelPickingDialog();
+	static void OpenLevelFromAssetPicker(const TArray<class FAssetData>& SelectedAssets, EAssetTypeActivationMethod::Type ActivationType);
+	
 
 	/** Opens delta transform */
 	static void DeltaTransform();
