remove optimize-off pragmas

diff --git a/Engine/Source/Editor/LevelEditor/Private/SLevelViewport.cpp b/Engine/Source/Editor/LevelEditor/Private/SLevelViewport.cpp
--- a/Engine/Source/Editor/LevelEditor/Private/SLevelViewport.cpp
+++ b/Engine/Source/Editor/LevelEditor/Private/SLevelViewport.cpp
@@ -29,7 +29,6 @@
 #include "SceneOutlinerTreeItems.h"
 
 static const FName LevelEditorName("LevelEditor");
-#pragma optimize("", off)
 
 #define LOCTEXT_NAMESPACE "LevelViewport"
 
