Viewport does not need to support maximise if it is switching to immersive mode.
#ttp 337318 EDITOR: CRASH: Changing viewport layout to one pane and entering immersive mode causes a crash
reviewer Thomas.Sarkanen

diff --git a/Engine/Source/Editor/LevelEditor/Private/LevelViewportLayout.cpp b/Engine/Source/Editor/LevelEditor/Private/LevelViewportLayout.cpp
--- a/Engine/Source/Editor/LevelEditor/Private/LevelViewportLayout.cpp
+++ b/Engine/Source/Editor/LevelEditor/Private/LevelViewportLayout.cpp
@@ -285,7 +285,6 @@ void FLevelViewportLayout::MaximizeViewport( TSharedRef<SLevelViewport> Viewport
 {
 	// Should never get into a situation where the viewport is being maximized and there is already a maximized viewport. 
 	// I.E Maximized viewport is NULL which means this is a new maximize or MaximizeViewport is equal to the passed in one which means this is a restore of the current maximized viewport
-	check( bIsMaximizeSupported );
 	check( Viewports.Contains( ViewportToMaximize ) );
 	check( !MaximizedViewport.IsValid() || MaximizedViewport == ViewportToMaximize );
 
