Disable the 'recompile level editor' button for promoted builds. It should only be enabled if the build is compiled locally.

diff --git a/Engine/Source/Editor/LevelEditor/Private/LevelEditor.cpp b/Engine/Source/Editor/LevelEditor/Private/LevelEditor.cpp
--- a/Engine/Source/Editor/LevelEditor/Private/LevelEditor.cpp
+++ b/Engine/Source/Editor/LevelEditor/Private/LevelEditor.cpp
@@ -153,7 +153,7 @@ void FLevelEditorModule::StartupModule()
 
 	// Figure out if we recompile the level editor.
 	FString SourcePath = FPaths::Combine(*FPaths::EngineDir(), TEXT("Source/Editor/LevelEditor/Private"));
-	bCanBeRecompiled = IFileManager::Get().DirectoryExists(*SourcePath);
+	bCanBeRecompiled = IFileManager::Get().DirectoryExists(*SourcePath) && !GEngineVersion.IsPromotedBuild();
 }
 
 /**
