Removed assumption from GetBrush that second actor is always a brush and renamed it to GetDefaultBrush. Replaced GetBrush with GetDefaultBrush and made GetBrush deprecated. Revised comments appropriatley.

Fixes TTP 337889 - STRATEGYGAME: VEHICLEGAME: CRASH: Console command 'Show Collision" crashes Standalone window

diff --git a/Engine/Source/Editor/GeometryMode/Private/GeometryModifiers.cpp b/Engine/Source/Editor/GeometryMode/Private/GeometryModifiers.cpp
--- a/Engine/Source/Editor/GeometryMode/Private/GeometryModifiers.cpp
+++ b/Engine/Source/Editor/GeometryMode/Private/GeometryModifiers.cpp
@@ -122,7 +122,7 @@ void UGeomModifier::DrawHUD(FLevelEditorViewportClient* ViewportClient,FViewport
 void UGeomModifier::CacheBrushState()
 {
 	FEdModeGeometry* GeomMode = (FEdModeGeometry*)GEditorModeTools().GetActiveMode(FBuiltinEditorModes::EM_Geometry);
-	ABrush* BuilderBrush = GeomMode->GetWorld()->GetBrush();
+	ABrush* BuilderBrush = GeomMode->GetWorld()->GetDefaultBrush();
 	if( !CachedPolys )
 	{
 		//Create the list of polys
@@ -152,7 +152,7 @@ void UGeomModifier::CacheBrushState()
 void UGeomModifier::RestoreBrushState()
 {
 	FEdModeGeometry* GeomMode = (FEdModeGeometry*)GEditorModeTools().GetActiveMode(FBuiltinEditorModes::EM_Geometry);
-	ABrush* BuilderBrush = GeomMode->GetWorld()->GetBrush();
+	ABrush* BuilderBrush = GeomMode->GetWorld()->GetDefaultBrush();
 
 	//Remove all of the current polys
 	BuilderBrush->Brush->Polys->Element.Empty();
@@ -926,7 +926,7 @@ void UGeomModifier_Lathe::Apply( int32 InTotalSegments, int32 InSegments, EAxis:
 
 	// We will be replacing the builder brush, so get it prepped.
 
-	ABrush* BuilderBrush = GeomMode->GetWorld()->GetBrush();
+	ABrush* BuilderBrush = GeomMode->GetWorld()->GetDefaultBrush();
 
 	BuilderBrush->SetActorLocation(GeomMode->GetWidgetLocation(), false);
 	BuilderBrush->SetPrePivot(FVector::ZeroVector);
@@ -1200,8 +1200,8 @@ void UGeomModifier_Pen::Apply()
 	if( ShapeVertices.Num() > 2 )
 	{
 		FEdModeGeometry* GeomMode = (FEdModeGeometry*)GEditorModeTools().GetActiveMode(FBuiltinEditorModes::EM_Geometry);
-		ABrush* ResultingBrush = GeomMode->GetWorld()->GetBrush();
-		ABrush* BuilderBrush = GeomMode->GetWorld()->GetBrush();
+		ABrush* ResultingBrush = GeomMode->GetWorld()->GetDefaultBrush();
+		ABrush* BuilderBrush = GeomMode->GetWorld()->GetDefaultBrush();
 
 		// Move all the vertices that the user placed to the same "height" as the builder brush, based on
 		// viewport orientation.  This is preferable to always creating the new builder brush at height zero.
@@ -1786,7 +1786,7 @@ static ABrush* ClipBrushAgainstPlane( const FPlane& InPlane, ABrush* InBrush)
 	}
 
 	// Move the new brush to where the new brush was to preserve brush ordering.
-	ABrush* BuilderBrush = World->GetBrush();
+	ABrush* BuilderBrush = World->GetDefaultBrush();
 	if( InBrush == BuilderBrush )
 	{
 		// Special-case behavior for the builder brush.
diff --git a/Engine/Source/Editor/LevelEditor/Private/LevelEditorActions.cpp b/Engine/Source/Editor/LevelEditor/Private/LevelEditorActions.cpp
--- a/Engine/Source/Editor/LevelEditor/Private/LevelEditorActions.cpp
+++ b/Engine/Source/Editor/LevelEditor/Private/LevelEditorActions.cpp
@@ -2523,63 +2523,66 @@ void FLevelEditorActionCallbacks::OnSaveBrushAsCollision()
 	}
 	
 	check(World)
-	
-	// Now get the builder brush.
-	UModel* builderModel = World->GetBrush()->Brush;
-
-	// Need the transform between builder brush space and static mesh actor space.
-	const FMatrix BrushL2W = World->GetBrush()->ActorToWorld().ToMatrixWithScale();
-	const FMatrix MeshW2L = MeshToWorld.Inverse();
-	const FMatrix SMToBB = BrushL2W * MeshW2L;
-	const FMatrix SMToBB_AT = SMToBB.TransposeAdjoint();
+	ABrush* BuildBrush = World->GetDefaultBrush();
+	if(BuildBrush != nullptr)
+	{
+		// Now get the builder brush.
+		UModel* BuilderModel = BuildBrush->Brush;
 
-	// Copy the current builder brush into a temp model.
-	// We keep no reference to this, so it will be GC'd at some point.
-	UModel* TempModel = new UModel(FPostConstructInitializeProperties(),NULL,1);
-	TempModel->Polys->Element.AssignButKeepOwner(builderModel->Polys->Element);
+		// Need the transform between builder brush space and static mesh actor space.
+		const FMatrix BrushL2W = BuildBrush->ActorToWorld().ToMatrixWithScale();
+		const FMatrix MeshW2L = MeshToWorld.Inverse();
+		const FMatrix SMToBB = BrushL2W * MeshW2L;
+		const FMatrix SMToBB_AT = SMToBB.TransposeAdjoint();
 
-	// Now transform each poly into local space for the selected static mesh.
-	for(int32 i=0; i<TempModel->Polys->Element.Num(); i++)
-	{
-		FPoly* Poly = &TempModel->Polys->Element[i];
+		// Copy the current builder brush into a temp model.
+		// We keep no reference to this, so it will be GC'd at some point.
+		UModel* TempModel = new UModel(FPostConstructInitializeProperties(), NULL, 1);
+		TempModel->Polys->Element.AssignButKeepOwner(BuilderModel->Polys->Element);
 
-		for(int32 j=0; j<Poly->Vertices.Num(); j++ )
+		// Now transform each poly into local space for the selected static mesh.
+		for (int32 i = 0; i < TempModel->Polys->Element.Num(); i++)
 		{
-			Poly->Vertices[j]  = SMToBB.TransformPosition(Poly->Vertices[j]);
-		}
+			FPoly* Poly = &TempModel->Polys->Element[i];
 
-		Poly->Normal = SMToBB_AT.TransformVector(Poly->Normal);
-		Poly->Normal.Normalize(); // SmToBB might have scaling in it.
-	}
+			for (int32 j = 0; j < Poly->Vertices.Num(); j++)
+			{
+				Poly->Vertices[j] = SMToBB.TransformPosition(Poly->Vertices[j]);
+			}
 
-	// Build bounding box.
-	TempModel->BuildBound();
+			Poly->Normal = SMToBB_AT.TransformVector(Poly->Normal);
+			Poly->Normal.Normalize(); // SmToBB might have scaling in it.
+		}
 
-	// Build BSP for the brush.
-	FBSPOps::bspBuild(TempModel,FBSPOps::BSP_Good,15,70,1,0);
-	FBSPOps::bspRefresh(TempModel,1);
-	FBSPOps::bspBuildBounds(TempModel);
+		// Build bounding box.
+		TempModel->BuildBound();
 
-	// Now - use this as the Rigid Body collision for this static mesh as well.
+		// Build BSP for the brush.
+		FBSPOps::bspBuild(TempModel, FBSPOps::BSP_Good, 15, 70, 1, 0);
+		FBSPOps::bspRefresh(TempModel, 1);
+		FBSPOps::bspBuildBounds(TempModel);
 
-	// Make sure rendering is done - so we are not changing data being used by collision drawing.
-	FlushRenderingCommands();
 
-	// If we already have a BodySetup - clear it.
-	if( StaticMesh->BodySetup )
-	{
-		StaticMesh->BodySetup->RemoveSimpleCollision();
-	}
-	// If we don't already have physics props, construct them here.
-	else
-	{
-		StaticMesh->CreateBodySetup();
-	}
+		// Now - use this as the Rigid Body collision for this static mesh as well.
 
-	// Convert collision model into a collection of convex hulls.
-	// NB: This removes any convex hulls that were already part of the collision data.
-	StaticMesh->BodySetup->CreateFromModel(TempModel, true);
+		// Make sure rendering is done - so we are not changing data being used by collision drawing.
+		FlushRenderingCommands();
 
+		// If we already have a BodySetup - clear it.
+		if (StaticMesh->BodySetup)
+		{
+			StaticMesh->BodySetup->RemoveSimpleCollision();
+		}
+		// If we don't already have physics props, construct them here.
+		else
+		{
+			StaticMesh->CreateBodySetup();
+		}
+
+		// Convert collision model into a collection of convex hulls.
+		// NB: This removes any convex hulls that were already part of the collision data.
+		StaticMesh->BodySetup->CreateFromModel(TempModel, true);
+	}
 	// refresh collision change back to staticmesh components
 	RefreshCollisionChange(StaticMesh);
 
diff --git a/Engine/Source/Editor/UnrealEd/Private/BSPOps.cpp b/Engine/Source/Editor/UnrealEd/Private/BSPOps.cpp
--- a/Engine/Source/Editor/UnrealEd/Private/BSPOps.cpp
+++ b/Engine/Source/Editor/UnrealEd/Private/BSPOps.cpp
@@ -971,6 +971,7 @@ void FBSPOps::bspBuildBounds( UModel* Model )
  */
 void FBSPOps::bspValidateBrush( UModel* Brush, bool ForceValidate, bool DoStatusUpdate )
 {
+	check(Brush != nullptr);
 	Brush->Modify();
 	if( ForceValidate || !Brush->Linked )
 	{
diff --git a/Engine/Source/Editor/UnrealEd/Private/EditorActor.cpp b/Engine/Source/Editor/UnrealEd/Private/EditorActor.cpp
--- a/Engine/Source/Editor/UnrealEd/Private/EditorActor.cpp
+++ b/Engine/Source/Editor/UnrealEd/Private/EditorActor.cpp
@@ -750,7 +750,7 @@ bool UUnrealEdEngine::ShouldAbortActorDeletion() const
 void UUnrealEdEngine::edactReplaceSelectedBrush( UWorld* InWorld )
 {
 	// Make a list of brush actors to replace.
-	ABrush* DefaultBrush = InWorld->GetBrush();
+	ABrush* DefaultBrush = InWorld->GetDefaultBrush();
 
 	TArray<ABrush*> BrushesToReplace;
 	for ( FSelectionIterator It( GetSelectedActorIterator() ) ; It ; ++It )
diff --git a/Engine/Source/Editor/UnrealEd/Private/EditorBrushBuilder.cpp b/Engine/Source/Editor/UnrealEd/Private/EditorBrushBuilder.cpp
--- a/Engine/Source/Editor/UnrealEd/Private/EditorBrushBuilder.cpp
+++ b/Engine/Source/Editor/UnrealEd/Private/EditorBrushBuilder.cpp
@@ -34,7 +34,7 @@ bool UEditorBrushBuilder::EndBrush( UWorld* InWorld, ABrush* InBrush )
 {
 	//!!validate
 	check( InWorld );
-	ABrush* BuilderBrush = (InBrush!=NULL) ? InBrush : InWorld->GetBrush();
+	ABrush* BuilderBrush = (InBrush != NULL) ? InBrush : InWorld->GetDefaultBrush();
 
 	// Ensure the builder brush is unhidden.
 	BuilderBrush->bHidden = false;
diff --git a/Engine/Source/Editor/UnrealEd/Private/EditorCsg.cpp b/Engine/Source/Editor/UnrealEd/Private/EditorCsg.cpp
--- a/Engine/Source/Editor/UnrealEd/Private/EditorCsg.cpp
+++ b/Engine/Source/Editor/UnrealEd/Private/EditorCsg.cpp
@@ -1257,7 +1257,7 @@ void UEditorEngine::MapBrushGet(UWorld* InWorld)
 		if( BrushActor && !FActorEditorUtils::IsABuilderBrush(Actor) )
 		{
 			check( BrushActor->GetWorld() );			
-			ABrush* WorldBrush = BrushActor->GetWorld()->GetBrush();
+			ABrush* WorldBrush = BrushActor->GetWorld()->GetDefaultBrush();
 			check( WorldBrush );
 			WorldBrush->Modify();
 			WorldBrush->Brush->Polys->Element.AssignButKeepOwner(BrushActor->Brush->Polys->Element);
@@ -1269,7 +1269,7 @@ void UEditorEngine::MapBrushGet(UWorld* InWorld)
 	}
 
 	GEditor->SelectNone( false, true );
-	GEditor->SelectActor( InWorld->GetBrush(), true, true );
+	GEditor->SelectActor(InWorld->GetDefaultBrush(), true, true);
 }
 
 
@@ -1287,7 +1287,7 @@ void UEditorEngine::mapBrushPut()
 		if( BrushActor && !FActorEditorUtils::IsABuilderBrush(Actor) )
 		{
 			check( BrushActor->GetWorld() );
-			ABrush* WorldBrush = BrushActor->GetWorld()->GetBrush();
+			ABrush* WorldBrush = BrushActor->GetWorld()->GetDefaultBrush();
 			check( WorldBrush );
 
 			BrushActor->Modify();
diff --git a/Engine/Source/Editor/UnrealEd/Private/EditorSelectUtils.cpp b/Engine/Source/Editor/UnrealEd/Private/EditorSelectUtils.cpp
--- a/Engine/Source/Editor/UnrealEd/Private/EditorSelectUtils.cpp
+++ b/Engine/Source/Editor/UnrealEd/Private/EditorSelectUtils.cpp
@@ -53,7 +53,7 @@ void UUnrealEdEngine::NoteActorMovement()
 		{
 			USelection* SelectedActors = GetSelectedActors();
 			SelectedActors->Modify();
-			SelectActor( GWorld->GetBrush(), true, true );
+			SelectActor( GWorld->GetDefaultBrush(), true, true );
 		}
 
 		// Look for an actor that requires snapping.
diff --git a/Engine/Source/Editor/UnrealEd/Private/EditorServer.cpp b/Engine/Source/Editor/UnrealEd/Private/EditorServer.cpp
--- a/Engine/Source/Editor/UnrealEd/Private/EditorServer.cpp
+++ b/Engine/Source/Editor/UnrealEd/Private/EditorServer.cpp
@@ -440,7 +440,7 @@ bool UEditorEngine::Exec_StaticMesh( UWorld* InWorld, const TCHAR* Str, FOutputD
 	bool bResult = false;
 #if !UE_BUILD_SHIPPING
 	// Not supported on shipped builds because PC cooking strips raw mesh data.
-	ABrush* WorldBrush = InWorld->GetBrush();
+	ABrush* WorldBrush = InWorld->GetDefaultBrush();
 	if(FParse::Command(&Str,TEXT("TO")))
 	{
 		if(FParse::Command(&Str,TEXT("BRUSH")))
@@ -526,7 +526,7 @@ bool UEditorEngine::Exec_Brush( UWorld* InWorld, const TCHAR* Str, FOutputDevice
 {
 	// Keep a pointer to the beginning of the string to use for message displaying purposes
 	const TCHAR* const FullStr = Str;
-	ABrush* WorldBrush = InWorld->GetBrush();
+	ABrush* WorldBrush = InWorld->GetDefaultBrush();
 	if( FParse::Command(&Str,TEXT("APPLYTRANSFORM")) )
 	{
 		CommandIsDeprecated( TEXT("APPLYTRANSFORM"), Ar );
@@ -536,12 +536,20 @@ bool UEditorEngine::Exec_Brush( UWorld* InWorld, const TCHAR* Str, FOutputDevice
 	{
 		{
 			const FScopedTransaction Transaction( NSLOCTEXT("UnrealEd", "BrushSet", "Brush Set") );
-			InWorld->GetBrush()->Brush->Modify();
-			FRotator Temp(0.0f,0.0f,0.0f);
-			FVector SnapLocation = InWorld->GetBrush()->GetActorLocation();
+			FRotator Temp(0.0f, 0.0f, 0.0f);
+			FVector SnapLocation(0.0f, 0.0f, 0.0f);
+			FVector PrePivot(0.0f, 0.0f, 0.0f);
+			ABrush* DefaultBrush = InWorld->GetDefaultBrush();
+			if (DefaultBrush != NULL)
+			{
+				DefaultBrush->Brush->Modify();
+				SnapLocation = DefaultBrush->GetActorLocation();
+				PrePivot = DefaultBrush->GetPrePivot();
+			}
+			
 			FSnappingUtils::SnapToBSPVertex( SnapLocation, FVector::ZeroVector, Temp );
 
-			WorldBrush->SetActorLocation(SnapLocation - InWorld->GetBrush()->GetPrePivot(), false);
+			WorldBrush->SetActorLocation(SnapLocation - PrePivot, false);
 			WorldBrush->SetPrePivot( FVector::ZeroVector );
 			WorldBrush->Brush->Polys->Element.Empty();
 			UPolysFactory* It = new UPolysFactory(FPostConstructInitializeProperties());
@@ -1585,10 +1593,13 @@ void UEditorEngine::BSPIntersectionHelper(UWorld* InWorld, ECsgOper Operation)
 	{
 		Mode->GeometrySelectNone(true, true);
 	}
-
-	InWorld->GetBrush()->Brush->Modify();
-	FinishAllSnaps();
-	bspBrushCSG( InWorld->GetBrush(), InWorld->GetModel(), 0, Brush_MAX, Operation, false, true, true );
+	ABrush* DefaultBrush = InWorld->GetDefaultBrush();
+	if (DefaultBrush != NULL)
+	{
+		DefaultBrush->Brush->Modify();
+		FinishAllSnaps();
+		bspBrushCSG(DefaultBrush, InWorld->GetModel(), 0, Brush_MAX, Operation, false, true, true);
+	}
 }
 
 void UEditorEngine::EditorDestroyWorld( FWorldContext & Context, const FText& CleanseText, UWorld* NewWorld )
@@ -1811,7 +1822,10 @@ UWorld* UEditorEngine::NewMap()
 	}
 
 	// Move the brush to the origin.
-	Context.World()->GetBrush()->SetActorLocation(FVector::ZeroVector, false);
+	if (Context.World()->GetDefaultBrush() != NULL)
+	{
+		Context.World()->GetDefaultBrush()->SetActorLocation(FVector::ZeroVector, false);
+	}
 
 	// Make the builder brush a small 256x256x256 cube so its visible.
 	InitBuilderBrush( Context.World() );
@@ -2110,7 +2124,7 @@ bool UEditorEngine::Map_Load(const TCHAR* Str, FOutputDevice& Ar)
 				}
 
 				{
-					FBSPOps::bspValidateBrush( Context.World()->GetBrush()->Brush, 0, 1 );
+					FBSPOps::bspValidateBrush(Context.World()->GetDefaultBrush()->Brush, 0, 1);
 
 					FString MapFileName = FPaths::GetCleanFilename(TempFname);
 					GWarn->StatusUpdate( -1, -1, FText::Format( LOCTEXT( "LoadingMapStatus_Initializing", "Loading map: {0}... (Initializing world)" ), FText::FromString(MapFileName) ) );
diff --git a/Engine/Source/Editor/UnrealEd/Private/Factories/EditorFactories.cpp b/Engine/Source/Editor/UnrealEd/Private/Factories/EditorFactories.cpp
--- a/Engine/Source/Editor/UnrealEd/Private/Factories/EditorFactories.cpp
+++ b/Engine/Source/Editor/UnrealEd/Private/Factories/EditorFactories.cpp
@@ -838,7 +838,7 @@ UObject* ULevelFactory::FactoryCreateText
 					for( int32 LevelIndex=0; LevelIndex<World->GetNumLevels(); LevelIndex++ )
 					{
 						ULevel* Level = World->GetLevel(LevelIndex);
-						if( Level->GetBrush() == Brush )
+						if(Level->GetDefaultBrush() == Brush)
 						{
 							bIsDefaultBrush = true;
 							break;
diff --git a/Engine/Source/Editor/UnrealEd/Private/LevelEditorViewport.cpp b/Engine/Source/Editor/UnrealEd/Private/LevelEditorViewport.cpp
--- a/Engine/Source/Editor/UnrealEd/Private/LevelEditorViewport.cpp
+++ b/Engine/Source/Editor/UnrealEd/Private/LevelEditorViewport.cpp
@@ -224,12 +224,16 @@ static TArray<AActor*> AttemptDropObjAsActors( ULevel* InLevel, UObject* ObjToUs
 		UWorld* World = InLevel->OwningWorld;
 		BrushBuilder->Build(World);
 
-		FVector ActorLoc = GEditor->ClickLocation + GEditor->ClickPlane * ( FVector::BoxPushOut( GEditor->ClickPlane, World->GetBrush()->GetPlacementExtent() ) );
-		FSnappingUtils::SnapPointToGrid( ActorLoc, FVector::ZeroVector ) ;
+		ABrush* DefaultBrush = World->GetDefaultBrush();
+		if (DefaultBrush != NULL)
+		{
+			FVector ActorLoc = GEditor->ClickLocation + GEditor->ClickPlane * (FVector::BoxPushOut(GEditor->ClickPlane, DefaultBrush->GetPlacementExtent()));
+			FSnappingUtils::SnapPointToGrid(ActorLoc, FVector::ZeroVector);
 
-		World->GetBrush()->SetActorLocation( ActorLoc );
-		PlacedActor = World->GetBrush();
-		PlacedActors.Add(World->GetBrush());
+			DefaultBrush->SetActorLocation(ActorLoc);
+			PlacedActor = DefaultBrush;
+			PlacedActors.Add(DefaultBrush);
+		}
 	}
 	else if (NULL == PlacedActor)
 	{
@@ -1412,7 +1416,7 @@ void FLevelEditorViewportClient::DestroyDropPreviewActors()
 		for ( auto ActorIt = DropPreviewActors.CreateConstIterator(); ActorIt; ++ActorIt )
 		{
 			AActor* PreviewActor = (*ActorIt).Get();
-			if ( PreviewActor && PreviewActor != GetWorld()->GetBrush())
+			if (PreviewActor && PreviewActor != GetWorld()->GetDefaultBrush())
 			{
 				GetWorld()->DestroyActor(PreviewActor);
 			}
diff --git a/Engine/Source/Editor/UnrealEd/Private/UnrealEdSrv.cpp b/Engine/Source/Editor/UnrealEd/Private/UnrealEdSrv.cpp
--- a/Engine/Source/Editor/UnrealEd/Private/UnrealEdSrv.cpp
+++ b/Engine/Source/Editor/UnrealEd/Private/UnrealEdSrv.cpp
@@ -1873,9 +1873,9 @@ void CreateBoundingBoxBuilderBrush( UWorld* InWorld, const TArray<FPoly*> Select
 	CubeBuilder->Z = extent.Z * 2;
 	CubeBuilder->Build( InWorld );
 
-	InWorld->GetBrush()->SetActorLocation( bbox.GetCenter(), false );
+	InWorld->GetDefaultBrush()->SetActorLocation(bbox.GetCenter(), false);
 
-	InWorld->GetBrush()->ReregisterAllComponents();
+	InWorld->GetDefaultBrush()->ReregisterAllComponents();
 }
 
 /**
@@ -2028,7 +2028,7 @@ bool UUnrealEdEngine::Exec_Actor( UWorld* InWorld, const TCHAR* Str, FOutputDevi
 	else if( FParse::Command(&Str,TEXT("CREATE_BV_BOUNDINGBOX")) )
 	{
 		const FScopedTransaction Transaction( NSLOCTEXT("UnrealEd", "CreateBoundingBoxBlockingVolume", "Create Bounding Box Blocking Volume") );
-		InWorld->GetBrush()->Modify();
+		InWorld->GetDefaultBrush()->Modify();
 
 		bool bSnapToGrid=0;
 		FParse::Bool( Str, TEXT("SNAPTOGRID="), bSnapToGrid );
@@ -2059,7 +2059,7 @@ bool UUnrealEdEngine::Exec_Actor( UWorld* InWorld, const TCHAR* Str, FOutputDevi
 	else if( FParse::Command(&Str,TEXT("CREATE_BV_CONVEXVOLUME")) )
 	{
 		const FScopedTransaction Transaction( NSLOCTEXT("UnrealEd", "CreateConvexBlockingVolume", "Create Convex Blocking Volume") );
-		InWorld->GetBrush()->Modify();
+		InWorld->GetDefaultBrush()->Modify();
 
 		bool bSnapToGrid=0;
 		FParse::Bool( Str, TEXT("SNAPTOGRID="), bSnapToGrid );
@@ -2086,7 +2086,7 @@ bool UUnrealEdEngine::Exec_Actor( UWorld* InWorld, const TCHAR* Str, FOutputDevi
 		// Get a list of the polygons that make up the builder brush
 
 		FPoly* poly;
-		TArray<FPoly>* BuilderBrushPolys = new TArray<FPoly>( InWorld->GetBrush()->Brush->Polys->Element );
+		TArray<FPoly>* BuilderBrushPolys = new TArray<FPoly>( InWorld->GetDefaultBrush()->Brush->Polys->Element );
 
 		// Create a list of valid splitting planes
 
@@ -2125,7 +2125,7 @@ bool UUnrealEdEngine::Exec_Actor( UWorld* InWorld, const TCHAR* Str, FOutputDevi
 			{
 				// Move the plane into the same coordinate space as the builder brush
 
-				*SplittingPlane = SplittingPlane->TransformBy( InWorld->GetBrush()->ActorToWorld().ToMatrixWithScale().Inverse() );
+				*SplittingPlane = SplittingPlane->TransformBy(InWorld->GetDefaultBrush()->ActorToWorld().ToMatrixWithScale().Inverse());
 
 				// Before keeping this plane, make sure there aren't any existing planes that have a normal within the rejection tolerance.
 
@@ -2262,14 +2262,14 @@ bool UUnrealEdEngine::Exec_Actor( UWorld* InWorld, const TCHAR* Str, FOutputDevi
 
 		// Create a new builder brush from the freshly clipped polygons.
 
-		InWorld->GetBrush()->Brush->Polys->Element.Empty();
+		InWorld->GetDefaultBrush()->Brush->Polys->Element.Empty();
 
 		for( int x = 0 ; x < BuilderBrushPolys->Num() ; ++x )
 		{
-			InWorld->GetBrush()->Brush->Polys->Element.Add( (*BuilderBrushPolys)[x] );
+			InWorld->GetDefaultBrush()->Brush->Polys->Element.Add((*BuilderBrushPolys)[x]);
 		}
 
-		InWorld->GetBrush()->ReregisterAllComponents();
+		InWorld->GetDefaultBrush()->ReregisterAllComponents();
 
 		// Create the blocking volume
 
diff --git a/Engine/Source/Editor/UnrealEd/Private/VertexSnapping.cpp b/Engine/Source/Editor/UnrealEd/Private/VertexSnapping.cpp
--- a/Engine/Source/Editor/UnrealEd/Private/VertexSnapping.cpp
+++ b/Engine/Source/Editor/UnrealEd/Private/VertexSnapping.cpp
@@ -586,7 +586,7 @@ void FVertexSnappingImpl::GetActorsInsideBox( const FBox& Box, UWorld* World, TA
 	{
 		AActor* Actor = *It;
 		// Ignore the builder brush, hidden actors and forcefully ignored actors (actors being moved)
-		if( Actor != World->GetBrush() && It->IsHiddenEd() == false && !ActorsToIgnore.Contains( Actor ) )
+		if( Actor != World->GetDefaultBrush() && It->IsHiddenEd() == false && !ActorsToIgnore.Contains( Actor ) )
 		{
 			const bool bNonColliding = true;
 			FBox ActorBoundingBox = Actor->GetComponentsBoundingBox(true);
diff --git a/Engine/Source/Editor/UnrealEd/Public/DragAndDrop/BrushBuilderDragDropOp.h b/Engine/Source/Editor/UnrealEd/Public/DragAndDrop/BrushBuilderDragDropOp.h
--- a/Engine/Source/Editor/UnrealEd/Public/DragAndDrop/BrushBuilderDragDropOp.h
+++ b/Engine/Source/Editor/UnrealEd/Public/DragAndDrop/BrushBuilderDragDropOp.h
@@ -24,15 +24,15 @@ public:
 		if(World != nullptr)
 		{
 			// Deselect & hide the builder brush
-			World->GetBrush()->SetIsTemporarilyHiddenInEditor(true);
-			GEditor->SelectActor(World->GetBrush(), false, false);
+			World->GetDefaultBrush()->SetIsTemporarilyHiddenInEditor(true);
+			GEditor->SelectActor(World->GetDefaultBrush(), false, false);
 		}
 	}
 
 	virtual void OnDrop( bool bDropWasHandled, const FPointerEvent& MouseEvent ) OVERRIDE
 	{
 		UWorld* World = GEditor->GetEditorWorldContext().World();
-		if (World != nullptr)
+		if ( (World != nullptr) && (World->GetDefaultBrush() != nullptr))
 		{
 			ULevel* CurrentLevel = World->GetCurrentLevel();
 			if (CurrentLevel != nullptr && !FLevelUtils::IsLevelLocked(CurrentLevel))
@@ -40,14 +40,14 @@ public:
 				// Cut the BSP
 				if (bDropWasHandled)
 				{
-					World->GetBrush()->BrushBuilder = DuplicateObject<UBrushBuilder>(BrushBuilder.Get(), World->GetBrush()->GetOuter());
+					World->GetDefaultBrush()->BrushBuilder = DuplicateObject<UBrushBuilder>(BrushBuilder.Get(), World->GetDefaultBrush()->GetOuter());
 					const TCHAR* Command = bIsAdditive ? TEXT("BRUSH ADD SELECTNEWBRUSH") : TEXT("BRUSH SUBTRACT SELECTNEWBRUSH");
 					GEditor->Exec(World, Command);
 				}
 
 				// Deselect & hide the builder brush
-				World->GetBrush()->SetIsTemporarilyHiddenInEditor(true);
-				GEditor->SelectActor(World->GetBrush(), false, false);
+				World->GetDefaultBrush()->SetIsTemporarilyHiddenInEditor(true);
+				GEditor->SelectActor(World->GetDefaultBrush(), false, false);
 			}
 		}
 
@@ -59,11 +59,11 @@ public:
 		UWorld* World = GEditor->GetEditorWorldContext().World();
 		if(World != nullptr)
 		{
-			if(CursorDecoratorWindow->IsVisible() && !World->GetBrush()->IsTemporarilyHiddenInEditor())
+			if(CursorDecoratorWindow->IsVisible() && !World->GetDefaultBrush()->IsTemporarilyHiddenInEditor())
 			{
 				// Deselect & hide the builder brush if the decorator is visible
-				World->GetBrush()->SetIsTemporarilyHiddenInEditor(true);
-				GEditor->SelectActor(World->GetBrush(), false, false);
+				World->GetDefaultBrush()->SetIsTemporarilyHiddenInEditor(true);
+				GEditor->SelectActor(World->GetDefaultBrush(), false, false);
 			}
 		}
 
@@ -100,8 +100,8 @@ private:
 	{
 		// show & select the builder brush
 		UWorld* World = GEditor->GetEditorWorldContext().World();
-		World->GetBrush()->SetIsTemporarilyHiddenInEditor(false);
-		GEditor->SelectActor(World->GetBrush(), true, false);
+		World->GetDefaultBrush()->SetIsTemporarilyHiddenInEditor(false);
+		GEditor->SelectActor(World->GetDefaultBrush(), true, false);
 	}
 
 private:
diff --git a/Engine/Source/Runtime/Engine/Classes/Engine/Level.h b/Engine/Source/Runtime/Engine/Classes/Engine/Level.h
--- a/Engine/Source/Runtime/Engine/Classes/Engine/Level.h
+++ b/Engine/Source/Runtime/Engine/Classes/Engine/Level.h
@@ -593,12 +593,20 @@ public:
 	/** Retrieves the array of streamable texture isntances. */
 	ENGINE_API TArray<FStreamableTextureInstance>* GetStreamableTextureInstances(UTexture2D*& TargetTexture);
 
+	/**
+	* Deprecated. Returns the default brush for this level.
+	*
+	* @return		The default brush for this level.
+	*/
+	DEPRECATED(4.3, "GetBrush is deprecated use GetDefaultBrush instead.")
+	ENGINE_API ABrush* GetBrush() const;
+
 	/**
 	 * Returns the default brush for this level.
 	 *
 	 * @return		The default brush for this level.
 	 */
-	ENGINE_API class ABrush* GetBrush() const;
+	ENGINE_API class ABrush* GetDefaultBrush() const;
 
 	/**
 	 * Returns the world info for this level.
diff --git a/Engine/Source/Runtime/Engine/Classes/Engine/World.h b/Engine/Source/Runtime/Engine/Classes/Engine/World.h
--- a/Engine/Source/Runtime/Engine/Classes/Engine/World.h
+++ b/Engine/Source/Runtime/Engine/Classes/Engine/World.h
@@ -1437,8 +1437,13 @@ public:
 	
 	UGameViewportClient* GetGameViewport() const;
 
-	/** Returns the default brush. */
+	DEPRECATED(4.3, "GetBrush is deprecated use GetDefaultBrush instead.")
 	ABrush* GetBrush() const;
+	/* 
+	 * Returns the default brush for the persistent level.
+	 * This is usually the 'builder brush' for editor builds, undefined for non editor instances and may be NULL.
+	 */
+	ABrush* GetDefaultBrush() const;
 
 	/** Returns true if the actors have been initialized and are ready to start play */
 	bool AreActorsInitialized() const;
diff --git a/Engine/Source/Runtime/Engine/Private/ActorEditorUtils.cpp b/Engine/Source/Runtime/Engine/Private/ActorEditorUtils.cpp
--- a/Engine/Source/Runtime/Engine/Private/ActorEditorUtils.cpp
+++ b/Engine/Source/Runtime/Engine/Private/ActorEditorUtils.cpp
@@ -10,7 +10,7 @@ namespace FActorEditorUtils
 #if WITH_EDITOR
 		if ( InActor && InActor->GetWorld() && !InActor->HasAnyFlags(RF_ClassDefaultObject) )
 		{
-			return InActor->GetLevel()->GetBrush() == InActor;
+			return InActor->GetLevel()->GetDefaultBrush() == InActor;
 		}
 #endif
 		return false;
diff --git a/Engine/Source/Runtime/Engine/Private/Level.cpp b/Engine/Source/Runtime/Engine/Private/Level.cpp
--- a/Engine/Source/Runtime/Engine/Private/Level.cpp
+++ b/Engine/Source/Runtime/Engine/Private/Level.cpp
@@ -1501,14 +1501,25 @@ TArray<FStreamableTextureInstance>* ULevel::GetStreamableTextureInstances(UTextu
 	return NULL;
 }
 
-
 ABrush* ULevel::GetBrush() const
 {
-	checkf( Actors.Num() >= 2, *GetPathName() );
-	ABrush* DefaultBrush = Cast<ABrush>( Actors[1] );
-	checkf( DefaultBrush != NULL, *GetPathName() );
-	checkf( DefaultBrush->BrushComponent, *GetPathName() );
-	checkf( DefaultBrush->Brush != NULL, *GetPathName() );
+	return GetDefaultBrush();
+}
+
+ABrush* ULevel::GetDefaultBrush() const
+{
+	ABrush* DefaultBrush = nullptr;
+	if (Actors.Num() >= 2)
+	{
+		// If the builder brush exists then it will be the 2nd actor in the actors array.
+		DefaultBrush = Cast<ABrush>(Actors[1]);
+		// If the second actor is not a brush then it certainly cannot be the builder brush.
+		if (DefaultBrush != nullptr)
+		{
+			checkf(DefaultBrush->BrushComponent, *GetPathName());
+			checkf(DefaultBrush->Brush != nullptr, *GetPathName());
+		}
+	}
 	return DefaultBrush;
 }
 
diff --git a/Engine/Source/Runtime/Engine/Private/World.cpp b/Engine/Source/Runtime/Engine/Private/World.cpp
--- a/Engine/Source/Runtime/Engine/Private/World.cpp
+++ b/Engine/Source/Runtime/Engine/Private/World.cpp
@@ -3021,9 +3021,14 @@ void UWorld::RemoveOnActorSpawnedHandler( const FOnActorSpawned::FDelegate& InHa
 }
 
 ABrush* UWorld::GetBrush() const
+{
+	return GetDefaultBrush();
+}
+
+ABrush* UWorld::GetDefaultBrush() const
 {
 	check(PersistentLevel);
-	return PersistentLevel->GetBrush();
+	return PersistentLevel->GetDefaultBrush();
 }
 
 
