Prohibited minimizing/maximizing two viewports at the same time.
#ttp 336555 - CRASH: Minimizing/maximizing two viewports at the same time crashes the editor
#branch UE4
#proj Editor.LevelEditor
#add Added a FCanExecuteAction to ToggleMaximize, and hid the toggle maximize icon during a transition, so it's no longer possible to set off a new maximize animation while one is already occurring.

diff --git a/Engine/Source/Editor/LevelEditor/Private/LevelViewportLayout.h b/Engine/Source/Editor/LevelEditor/Private/LevelViewportLayout.h
--- a/Engine/Source/Editor/LevelEditor/Private/LevelViewportLayout.h
+++ b/Engine/Source/Editor/LevelEditor/Private/LevelViewportLayout.h
@@ -104,6 +104,9 @@ public:
 	/** Returns the parent tab content object */
 	TWeakPtr< class FLevelViewportTabContent > GetParentTabContent() const { return ParentTabContent; }
 
+	/** Returns whether a viewport animation is currently taking place */
+	bool IsTransitioning() const { return bIsTransitioning; }
+
 protected:
 
 	/**
diff --git a/Engine/Source/Editor/LevelEditor/Private/SLevelViewport.cpp b/Engine/Source/Editor/LevelEditor/Private/SLevelViewport.cpp
--- a/Engine/Source/Editor/LevelEditor/Private/SLevelViewport.cpp
+++ b/Engine/Source/Editor/LevelEditor/Private/SLevelViewport.cpp
@@ -989,7 +989,8 @@ void SLevelViewport::BindOptionCommands( FUICommandList& CommandList )
 
 	CommandList.MapAction( 
 		ViewportActions.ToggleMaximize,
-		FExecuteAction::CreateSP( this, &SLevelViewport::OnToggleMaximizeMode ) );
+		FExecuteAction::CreateSP( this, &SLevelViewport::OnToggleMaximizeMode ),
+		FCanExecuteAction::CreateSP( this, &SLevelViewport::CanToggleMaximizeMode ) );
 
 	
 	CommandList.MapAction(
@@ -1365,14 +1366,16 @@ EVisibility SLevelViewport::GetToolBarVisibility() const
 EVisibility SLevelViewport::GetMaximizeToggleVisibility() const
 {
 	bool bIsMaximizeSupported = false;
+	bool bShowMaximizeToggle = false;
 	TSharedPtr<FLevelViewportLayout> LayoutPinned = ParentLayout.Pin();
 	if (LayoutPinned.IsValid())
 	{
 		bIsMaximizeSupported = LayoutPinned->IsMaximizeSupported();
+		bShowMaximizeToggle = !LayoutPinned->IsTransitioning();
 	}
 
 	// Do not show the maximize/minimize toggle when in immersive mode
-	return (!bIsMaximizeSupported || IsImmersive()) ? EVisibility::Collapsed : EVisibility::Visible;
+	return (!bIsMaximizeSupported || IsImmersive()) ? EVisibility::Collapsed : (bShowMaximizeToggle ? EVisibility::Visible : EVisibility::Hidden);
 }
 
 EVisibility SLevelViewport::GetCloseImmersiveButtonVisibility() const
@@ -2145,6 +2148,12 @@ void SLevelViewport::RedrawViewport( bool bInvalidateHitProxies )
 	}
 }
 
+bool SLevelViewport::CanToggleMaximizeMode() const
+{
+	TSharedPtr<FLevelViewportLayout> ParentLayoutPinned = ParentLayout.Pin();
+	return (ParentLayoutPinned.IsValid() && ParentLayoutPinned->IsMaximizeSupported() && !ParentLayoutPinned->IsTransitioning());
+}
+
 void  SLevelViewport::OnToggleMaximizeMode()
 {
 	OnToggleMaximize();
diff --git a/Engine/Source/Editor/LevelEditor/Public/SLevelViewport.h b/Engine/Source/Editor/LevelEditor/Public/SLevelViewport.h
--- a/Engine/Source/Editor/LevelEditor/Public/SLevelViewport.h
+++ b/Engine/Source/Editor/LevelEditor/Public/SLevelViewport.h
@@ -283,6 +283,11 @@ private:
 	 */
 	void OnToggleImmersive();
 
+	/**
+	* Called to determine whether the maximize mode of current viewport can be toggled
+	*/
+	bool CanToggleMaximizeMode() const;
+
 	/**
 	* Called to toggle maximize mode of current viewport
 	*/
