Fixed crashes caused by shared pointer delegate bindings when reloading the level editor

Using AddSP to bind onto a multicast delegate that's removed in a class's destructor is dangerous because the weak pointer used to compare against the 'this' pointer passed to .RemoveAll() will be invalid by the time it is compared. This causes "HasSameObject" to fail as the object cannot be pinned. I've have changed these to be raw bindings so that the ownership is explicit and thus force the call to RemoveAll to succeed correctly without leaving behind a null delegate instance. Previously, with a null delegate instance still in the list when the module is reloaded, that memory becomes invalid causing a crash when the event attempts to purge invalid delegate instances from its invocation list.

Also fixed one place where OnTakeHighResScreenShots was being bound but never removed (this looked like a copy paste error)

This addresses TTP#337596 - CRASH: EDITOR: Reload after Recompile causes crash

diff --git a/Engine/Source/Editor/LevelEditor/Private/SLevelEditor.cpp b/Engine/Source/Editor/LevelEditor/Private/SLevelEditor.cpp
--- a/Engine/Source/Editor/LevelEditor/Private/SLevelEditor.cpp
+++ b/Engine/Source/Editor/LevelEditor/Private/SLevelEditor.cpp
@@ -149,20 +149,23 @@ void SLevelEditor::BindCommands()
 
 void SLevelEditor::Construct( const SLevelEditor::FArguments& InArgs)
 {
+	// Important: We use raw bindings here because we are releasing our binding in our destructor (where a weak pointer would be invalid)
+	// It's imperative that our delegate is removed in the destructor for the level editor module to play nicely with reloading.
+
 	FLevelEditorModule& LevelEditorModule = FModuleManager::GetModuleChecked< FLevelEditorModule >( LevelEditorModuleName );
-	LevelEditorModule.OnNotificationBarChanged().AddSP( this, &SLevelEditor::ConstructNotificationBar );
+	LevelEditorModule.OnNotificationBarChanged().AddRaw( this, &SLevelEditor::ConstructNotificationBar );
 
-	GetMutableDefault<UEditorExperimentalSettings>()->OnSettingChanged().AddSP(this, &SLevelEditor::HandleExperimentalSettingChanged);
+	GetMutableDefault<UEditorExperimentalSettings>()->OnSettingChanged().AddRaw(this, &SLevelEditor::HandleExperimentalSettingChanged);
 
 	BindCommands();
 
 	// We need to register when modes list changes so that we can refresh the auto generated commands.
-	GEditorModeTools().OnRegisteredModesChanged().AddSP(this, &SLevelEditor::RefreshEditorModeCommands);
+	GEditorModeTools().OnRegisteredModesChanged().AddRaw(this, &SLevelEditor::RefreshEditorModeCommands);
 
 	// @todo This is a hack to get this working for now. This won't work with multiple worlds
 	GEditor->GetEditorWorldContext(true).AddRef(World);
 
-	FEditorDelegates::MapChange.AddSP(this, &SLevelEditor::HandleEditorMapChange);
+	FEditorDelegates::MapChange.AddRaw(this, &SLevelEditor::HandleEditorMapChange);
 	HandleEditorMapChange(MapChangeEventFlags::NewMap);
 }
 
diff --git a/Engine/Source/Editor/LevelEditor/Private/SLevelEditorToolBox.cpp b/Engine/Source/Editor/LevelEditor/Private/SLevelEditorToolBox.cpp
--- a/Engine/Source/Editor/LevelEditor/Private/SLevelEditorToolBox.cpp
+++ b/Engine/Source/Editor/LevelEditor/Private/SLevelEditorToolBox.cpp
@@ -20,7 +20,10 @@ void SLevelEditorToolBox::Construct( const FArguments& InArgs, const TSharedRef<
 {
 	LevelEditor = OwningLevelEditor;
 
-	GEditor->AccessEditorUserSettings().OnUserSettingChanged().AddSP( this, &SLevelEditorToolBox::HandleUserSettingsChange );
+	// Important: We use a raw binding here because we are releasing our binding in our destructor (where a weak pointer would be invalid)
+	// It's imperative that our delegate is removed in the destructor for the level editor module to play nicely with reloading.
+
+	GEditor->AccessEditorUserSettings().OnUserSettingChanged().AddRaw( this, &SLevelEditorToolBox::HandleUserSettingsChange );
 
 	ChildSlot
 	[
diff --git a/Engine/Source/Editor/LevelEditor/Private/SLevelViewport.cpp b/Engine/Source/Editor/LevelEditor/Private/SLevelViewport.cpp
--- a/Engine/Source/Editor/LevelEditor/Private/SLevelViewport.cpp
+++ b/Engine/Source/Editor/LevelEditor/Private/SLevelViewport.cpp
@@ -76,9 +76,9 @@ SLevelViewport::~SLevelViewport()
 
 	FLevelEditorModule& LevelEditor = FModuleManager::GetModuleChecked<FLevelEditorModule>( LevelEditorName );
 	LevelEditor.OnRedrawLevelEditingViewports().RemoveAll( this );
+	LevelEditor.OnTakeHighResScreenShots().RemoveAll( this );
 	LevelEditor.OnActorSelectionChanged().RemoveAll( this );
 	LevelEditor.OnMapChanged().RemoveAll( this );
-	LevelEditor.OnRedrawLevelEditingViewports().RemoveAll( this );
 	GEngine->OnLevelActorDeleted().RemoveAll( this );
 
 	GetMutableDefault<ULevelEditorViewportSettings>()->OnSettingChanged().RemoveAll( this );
@@ -107,7 +107,7 @@ bool SLevelViewport::IsVisible() const
 
 void SLevelViewport::Construct(const FArguments& InArgs)
 {
-	GetMutableDefault<ULevelEditorViewportSettings>()->OnSettingChanged().AddSP(this, &SLevelViewport::HandleViewportSettingChanged);
+	GetMutableDefault<ULevelEditorViewportSettings>()->OnSettingChanged().AddRaw(this, &SLevelViewport::HandleViewportSettingChanged);
 
 	ParentLayout = InArgs._ParentLayout;
 	ParentLevelEditor = InArgs._ParentLevelEditor;
@@ -134,17 +134,20 @@ void SLevelViewport::Construct(const FArguments& InArgs)
 	// If a map has already been loaded, this will test for it and copy the correct camera location out
 	OnMapChanged( GWorld, EMapChangeType::LoadMap );
 
+	// Important: We use raw bindings here because we are releasing our binding in our destructor (where a weak pointer would be invalid)
+	// It's imperative that our delegate is removed in the destructor for the level editor module to play nicely with reloading.
+
 	FLevelEditorModule& LevelEditor = FModuleManager::GetModuleChecked<FLevelEditorModule>( LevelEditorName );
-	LevelEditor.OnRedrawLevelEditingViewports().AddSP( this, &SLevelViewport::RedrawViewport );
-	LevelEditor.OnTakeHighResScreenShots().AddSP( this, &SLevelViewport::TakeHighResScreenShot );
+	LevelEditor.OnRedrawLevelEditingViewports().AddRaw( this, &SLevelViewport::RedrawViewport );
+	LevelEditor.OnTakeHighResScreenShots().AddRaw( this, &SLevelViewport::TakeHighResScreenShot );
 
 	// Tell the level editor we want to be notified when selection changes
-	LevelEditor.OnActorSelectionChanged().AddSP( this, &SLevelViewport::OnActorSelectionChanged );
+	LevelEditor.OnActorSelectionChanged().AddRaw( this, &SLevelViewport::OnActorSelectionChanged );
 
 	// Tell the level editor we want to be notified when selection changes
-	LevelEditor.OnMapChanged().AddSP( this, &SLevelViewport::OnMapChanged );
+	LevelEditor.OnMapChanged().AddRaw( this, &SLevelViewport::OnMapChanged );
 
-	GEngine->OnLevelActorDeleted().AddSP( this, &SLevelViewport::OnLevelActorsRemoved );
+	GEngine->OnLevelActorDeleted().AddRaw( this, &SLevelViewport::OnLevelActorsRemoved );
 }
 
 void SLevelViewport::ConstructViewportOverlayContent()
