Moved Stats sub menu entry from viewport's Show menu to the \\/ menu

#ttp 337844 - Editor: Move new stats menu to first drop down arrow next to "Show Stats" rather than show flags

#branch UE4

#change Moved FillShowMenu and FillShowStatsSubMenus functions so they could be more accessible

diff --git a/Engine/Source/Editor/LevelEditor/Private/SLevelViewportToolBar.cpp b/Engine/Source/Editor/LevelEditor/Private/SLevelViewportToolBar.cpp
--- a/Engine/Source/Editor/LevelEditor/Private/SLevelViewportToolBar.cpp
+++ b/Engine/Source/Editor/LevelEditor/Private/SLevelViewportToolBar.cpp
@@ -36,6 +36,70 @@ public:
 	}
 };
 
+static void FillShowMenu(class FMenuBuilder& MenuBuilder, TArray< FLevelViewportCommands::FShowMenuCommand > MenuCommands, int32 EntryOffset)
+{
+	// Generate entries for the standard show flags
+	// Assumption: the first 'n' entries types like 'Show All' and 'Hide All' buttons, so insert a separator after them
+	for (int32 EntryIndex = 0; EntryIndex < MenuCommands.Num(); ++EntryIndex)
+	{
+		MenuBuilder.AddMenuEntry(MenuCommands[EntryIndex].ShowMenuItem, NAME_None, MenuCommands[EntryIndex].LabelOverride);
+		if (EntryIndex == EntryOffset - 1)
+		{
+			MenuBuilder.AddMenuSeparator();
+		}
+	}
+}
+
+static void FillShowStatsSubMenus(class FMenuBuilder& MenuBuilder, TArray< FLevelViewportCommands::FShowMenuCommand > MenuCommands, TMap< FString, TArray< FLevelViewportCommands::FShowMenuCommand > > StatCatCommands)
+{
+	FillShowMenu(MenuBuilder, MenuCommands, 1);
+
+	// Separate out stats into two list, those with and without submenus
+	TArray< FLevelViewportCommands::FShowMenuCommand > SingleStatCommands;
+	TMap< FString, TArray< FLevelViewportCommands::FShowMenuCommand > > SubbedStatCommands;
+	for (auto StatCatIt = StatCatCommands.CreateConstIterator(); StatCatIt; ++StatCatIt)
+	{
+		const TArray< FLevelViewportCommands::FShowMenuCommand >& ShowStatCommands = StatCatIt.Value();
+		const FString& CategoryName = StatCatIt.Key();
+
+		// If no category is specified, or there's only one category, don't use submenus
+		FString NoCategory = FStatConstants::NAME_NoCategory.ToString();
+		NoCategory.RemoveFromStart(TEXT("STATCAT_"));
+		if (CategoryName == NoCategory || StatCatCommands.Num() == 1)
+		{
+			for (int32 StatIndex = 0; StatIndex < ShowStatCommands.Num(); ++StatIndex)
+			{
+				const FLevelViewportCommands::FShowMenuCommand& StatCommand = ShowStatCommands[StatIndex];
+				SingleStatCommands.Add(StatCommand);
+			}
+		}
+		else
+		{
+			SubbedStatCommands.Add(CategoryName, ShowStatCommands);
+		}
+	}
+
+	// First add all the stats that don't have a sub menu
+	for (auto StatCatIt = SingleStatCommands.CreateConstIterator(); StatCatIt; ++StatCatIt)
+	{
+		const FLevelViewportCommands::FShowMenuCommand& StatCommand = *StatCatIt;
+		MenuBuilder.AddMenuEntry(StatCommand.ShowMenuItem, NAME_None, StatCommand.LabelOverride);
+	}
+
+	// Now add all the stats that have sub menus
+	for (auto StatCatIt = SubbedStatCommands.CreateConstIterator(); StatCatIt; ++StatCatIt)
+	{
+		const TArray< FLevelViewportCommands::FShowMenuCommand >& StatCommands = StatCatIt.Value();
+		const FText CategoryName = FText::FromString(StatCatIt.Key());
+
+		FFormatNamedArguments Args;
+		Args.Add(TEXT("StatCat"), CategoryName);
+		const FText CategoryDescription = FText::Format(NSLOCTEXT("UICommands", "StatShowCatName", "Show {StatCat} stats"), Args);
+
+		MenuBuilder.AddSubMenu(CategoryName, CategoryDescription,
+			FNewMenuDelegate::CreateStatic(&FillShowMenu, StatCommands, 0));
+	}
+}
 
 void SLevelViewportToolBar::Construct( const FArguments& InArgs )
 {
@@ -532,6 +596,17 @@ TSharedRef<SWidget> SLevelViewportToolBar::GenerateOptionsMenu() const
 			OptionsMenuBuilder.AddMenuEntry( FEditorViewportCommands::Get().ToggleStats );
 			OptionsMenuBuilder.AddMenuEntry( FEditorViewportCommands::Get().ToggleFPS );
 
+			FText HideAllLabel = LOCTEXT("HideAllLabel", "Hide All");
+			TArray< FLevelViewportCommands::FShowMenuCommand > HideStatsMenu;
+
+			// 'Hide All' button
+			HideStatsMenu.Add(FLevelViewportCommands::FShowMenuCommand(LevelViewportActions.HideAllStats, HideAllLabel));
+
+			OptionsMenuBuilder.AddSubMenu(
+				LOCTEXT("ShowStatsMenu", "Stat"), 
+				LOCTEXT("ShowStatsMenu_ToolTip", "Show Stat commands"),
+				FNewMenuDelegate::CreateStatic(&FillShowStatsSubMenus, HideStatsMenu, LevelViewportActions.ShowStatCatCommands));
+
 			if( bIsPerspective )
 			{
 				OptionsMenuBuilder.AddWidget( GenerateFOVMenu(), LOCTEXT("FOVAngle", "Field of View") );
@@ -874,94 +949,26 @@ TSharedRef<SWidget> SLevelViewportToolBar::GenerateShowMenu() const
 			ShowMenuBuilder.EndSection();
 		}
 
-		struct FLocal
-		{
-			static void FillShowMenu( class FMenuBuilder& MenuBuilder, TArray< FLevelViewportCommands::FShowMenuCommand > MenuCommands, int32 EntryOffset )
-			{
-				// Generate entries for the standard show flags
-				// Assumption: the first 'n' entries types like 'Show All' and 'Hide All' buttons, so insert a separator after them
-				for (int32 EntryIndex = 0; EntryIndex < MenuCommands.Num(); ++EntryIndex)
-				{
-					MenuBuilder.AddMenuEntry(MenuCommands[EntryIndex].ShowMenuItem, NAME_None, MenuCommands[EntryIndex].LabelOverride);
-					if (EntryIndex == EntryOffset-1)
-					{
-						MenuBuilder.AddMenuSeparator();
-					}
-				}
-			}
-
-			static void FillShowStatsSubMenus(class FMenuBuilder& MenuBuilder, TArray< FLevelViewportCommands::FShowMenuCommand > MenuCommands, TMap< FString, TArray< FLevelViewportCommands::FShowMenuCommand > > StatCatCommands)
-			{
-				FLocal::FillShowMenu(MenuBuilder, MenuCommands, 1);
-
-				// Separate out stats into two list, those with and without submenus
-				TArray< FLevelViewportCommands::FShowMenuCommand > SingleStatCommands;
-				TMap< FString, TArray< FLevelViewportCommands::FShowMenuCommand > > SubbedStatCommands;
-				for (auto StatCatIt = StatCatCommands.CreateConstIterator(); StatCatIt; ++StatCatIt)
-				{
-					const TArray< FLevelViewportCommands::FShowMenuCommand >& ShowStatCommands = StatCatIt.Value();
-					const FString& CategoryName = StatCatIt.Key();
-
-					// If no category is specified, or there's only one category, don't use submenus
-					FString NoCategory = FStatConstants::NAME_NoCategory.ToString();
-					NoCategory.RemoveFromStart(TEXT("STATCAT_"));
-					if (CategoryName == NoCategory || StatCatCommands.Num() == 1)
-					{
-						for (int32 StatIndex = 0; StatIndex < ShowStatCommands.Num(); ++StatIndex)
-						{
-							const FLevelViewportCommands::FShowMenuCommand& StatCommand = ShowStatCommands[StatIndex];
-							SingleStatCommands.Add(StatCommand);
-						}
-					}
-					else
-					{
-						SubbedStatCommands.Add(CategoryName, ShowStatCommands);
-					}
-				}
-
-				// First add all the stats that don't have a sub menu
-				for (auto StatCatIt = SingleStatCommands.CreateConstIterator(); StatCatIt; ++StatCatIt)
-				{
-					const FLevelViewportCommands::FShowMenuCommand& StatCommand = *StatCatIt;
-					MenuBuilder.AddMenuEntry(StatCommand.ShowMenuItem, NAME_None, StatCommand.LabelOverride);
-				}
-
-				// Now add all the stats that have sub menus
-				for (auto StatCatIt = SubbedStatCommands.CreateConstIterator(); StatCatIt; ++StatCatIt)
-				{
-					const TArray< FLevelViewportCommands::FShowMenuCommand >& StatCommands = StatCatIt.Value();
-					const FText CategoryName = FText::FromString(StatCatIt.Key());
-
-					FFormatNamedArguments Args;
-					Args.Add(TEXT("StatCat"), CategoryName);
-					const FText CategoryDescription = FText::Format(NSLOCTEXT("UICommands", "StatShowCatName", "Show {StatCat} stats"), Args);
-
-					MenuBuilder.AddSubMenu(CategoryName, CategoryDescription,
-						FNewMenuDelegate::CreateStatic(&FLocal::FillShowMenu, StatCommands, 0));
-				}
-			}
-		};
-
 		// Generate entries for the different show flags groups
 		ShowMenuBuilder.BeginSection("LevelViewportShowFlags");
 		{
 			ShowMenuBuilder.AddSubMenu( LOCTEXT("PostProcessShowFlagsMenu", "Post Processing"), LOCTEXT("PostProcessShowFlagsMenu_ToolTip", "Post process show flags"),
-				FNewMenuDelegate::CreateStatic( &FLocal::FillShowMenu, ShowMenu[SFG_PostProcess], 0 ) );
+				FNewMenuDelegate::CreateStatic( &FillShowMenu, ShowMenu[SFG_PostProcess], 0 ) );
 
 			ShowMenuBuilder.AddSubMenu( LOCTEXT("LightingComponentsShowFlagsMenu", "Lighting Components"), LOCTEXT("LightingComponentsShowFlagsMenu_ToolTip", "Lighting Components show flags"),
-				FNewMenuDelegate::CreateStatic( &FLocal::FillShowMenu, ShowMenu[SFG_LightingComponents], 0 ) );
+				FNewMenuDelegate::CreateStatic( &FillShowMenu, ShowMenu[SFG_LightingComponents], 0 ) );
 
 			ShowMenuBuilder.AddSubMenu( LOCTEXT("LightingFeaturesShowFlagsMenu", "Lighting Features"), LOCTEXT("LightingFeaturesShowFlagsMenu_ToolTip", "Lighting Features show flags"),
-				FNewMenuDelegate::CreateStatic( &FLocal::FillShowMenu, ShowMenu[SFG_LightingFeatures], 0 ) );
+				FNewMenuDelegate::CreateStatic( &FillShowMenu, ShowMenu[SFG_LightingFeatures], 0 ) );
 
 			ShowMenuBuilder.AddSubMenu( LOCTEXT("DeveloperShowFlagsMenu", "Developer"), LOCTEXT("DeveloperShowFlagsMenu_ToolTip", "Developer show flags"),
-				FNewMenuDelegate::CreateStatic( &FLocal::FillShowMenu, ShowMenu[SFG_Developer], 0 ) );
+				FNewMenuDelegate::CreateStatic( &FillShowMenu, ShowMenu[SFG_Developer], 0 ) );
 
 			ShowMenuBuilder.AddSubMenu( LOCTEXT("VisualizeShowFlagsMenu", "Visualize"), LOCTEXT("VisualizeShowFlagsMenu_ToolTip", "Visualize show flags"),
-				FNewMenuDelegate::CreateStatic( &FLocal::FillShowMenu, ShowMenu[SFG_Visualize], 0 ) );
+				FNewMenuDelegate::CreateStatic( &FillShowMenu, ShowMenu[SFG_Visualize], 0 ) );
 
 			ShowMenuBuilder.AddSubMenu( LOCTEXT("AdvancedShowFlagsMenu", "Advanced"), LOCTEXT("AdvancedShowFlagsMenu_ToolTip", "Advanced show flags"),
-				FNewMenuDelegate::CreateStatic( &FLocal::FillShowMenu, ShowMenu[SFG_Advanced], 0 ) );
+				FNewMenuDelegate::CreateStatic( &FillShowMenu, ShowMenu[SFG_Advanced], 0 ) );
 		}
 		ShowMenuBuilder.EndSection();
 
@@ -981,7 +988,7 @@ TSharedRef<SWidget> SLevelViewportToolBar::GenerateShowMenu() const
 			ShowVolumesMenu += Actions.ShowVolumeCommands;
 
 			ShowMenuBuilder.AddSubMenu( LOCTEXT("ShowVolumesMenu", "Volumes"), LOCTEXT("ShowVolumesMenu_ToolTip", "Show volumes flags"),
-				FNewMenuDelegate::CreateStatic( &FLocal::FillShowMenu, ShowVolumesMenu, 2 ) );
+				FNewMenuDelegate::CreateStatic( &FillShowMenu, ShowVolumesMenu, 2 ) );
 		}
 
 		// Show Layers sub-menu is dynamically generated when the user enters 'show' menu
@@ -1002,18 +1009,7 @@ TSharedRef<SWidget> SLevelViewportToolBar::GenerateShowMenu() const
 			ShowSpritesMenu += Actions.ShowSpriteCommands;
 
 			ShowMenuBuilder.AddSubMenu( LOCTEXT("ShowSpritesMenu", "Sprites"), LOCTEXT("ShowSpritesMenu_ToolTip", "Show sprites flags"),
-				FNewMenuDelegate::CreateStatic( &FLocal::FillShowMenu, ShowSpritesMenu, 2 ) );
-		}
-
-		// Show Stats sub-menu
-		{
-			TArray< FLevelViewportCommands::FShowMenuCommand > HideStatsMenu;
-
-			// 'Hide All' button
-			HideStatsMenu.Add(FLevelViewportCommands::FShowMenuCommand(Actions.HideAllStats, HideAllLabel));
-
-			ShowMenuBuilder.AddSubMenu(LOCTEXT("ShowStatsMenu", "Stat"), LOCTEXT("ShowStatsMenu_ToolTip", "Show Stat commands"),
-				FNewMenuDelegate::CreateStatic(&FLocal::FillShowStatsSubMenus, HideStatsMenu, Actions.ShowStatCatCommands));
+				FNewMenuDelegate::CreateStatic( &FillShowMenu, ShowSpritesMenu, 2 ) );
 		}
 	}
 
