Selecting a camera that has a pinned camera preview now highlights its preview

This addresses TTP#322790 - EDITOR: Viewport: Feature Request: When double clicking on a pinned camera actor in viewport, have its preview window flash

diff --git a/Engine/Source/Editor/EditorStyle/Private/SlateEditorStyle.cpp b/Engine/Source/Editor/EditorStyle/Private/SlateEditorStyle.cpp
--- a/Engine/Source/Editor/EditorStyle/Private/SlateEditorStyle.cpp
+++ b/Engine/Source/Editor/EditorStyle/Private/SlateEditorStyle.cpp
@@ -648,6 +648,7 @@ void FSlateEditorStyle::FStyle::SetupGeneralStyles()
 	Set( "DashedBorder", new BORDER_BRUSH( "Old/DashedBorder", FMargin(6.0f/32.0f) ) );
 	Set( "Checker", new IMAGE_BRUSH( "Old/Checker", Icon16x16, FLinearColor::White, ESlateBrushTileType::Both ) );
 	Set( "UniformShadow", new BORDER_BRUSH( "Common/UniformShadow", FMargin( 16.0f / 64.0f ) ) );
+	Set( "UniformShadow_Tint", new BORDER_BRUSH( "Common/UniformShadow_Tint", FMargin( 16.0f / 64.0f ) ) );
 
 	// Splitter
 #if WITH_EDITOR || IS_PROGRAM
diff --git a/Engine/Source/Editor/LevelEditor/Private/SLevelViewport.cpp b/Engine/Source/Editor/LevelEditor/Private/SLevelViewport.cpp
--- a/Engine/Source/Editor/LevelEditor/Private/SLevelViewport.cpp
+++ b/Engine/Source/Editor/LevelEditor/Private/SLevelViewport.cpp
@@ -2276,6 +2276,8 @@ class SActorPreview : public SCompoundWidget
 
 public:
 
+	~SActorPreview();
+
 	SLATE_BEGIN_ARGS( SActorPreview )
 		: _ViewportWidth( 240 ),
 		  _ViewportHeight( 180 ) {}
@@ -2309,9 +2311,14 @@ public:
 	virtual void OnMouseEnter( const FGeometry& MyGeometry, const FPointerEvent& MouseEvent ) override;
 	virtual void OnMouseLeave( const FPointerEvent& MouseEvent ) override;
 
+	/** Highlight this preview window by flashing the border. Will replay the curve sequence if it is already in the middle of a highlight. */
+	void Highlight();
 
 private:
 
+	/** Called when an actor in the world is selected */
+	void OnActorSelected(UObject* InActor);
+
 	/** @return Returns the color and opacity to use for this widget */
 	FLinearColor GetColorAndOpacity() const;
 
@@ -2351,17 +2358,43 @@ private:
 	/** Curve sequence for fading in and out */
 	FCurveSequence FadeSequence;
 
+	/** Curve sequence for flashing the border (highlighting) when a pinned preview is re-selected */
+	FCurveSequence HighlightSequence;
+
 	/** Padding around the preview actor name */
 	static const float PreviewTextPadding;
 };
 
 const float SActorPreview::PreviewTextPadding = 3.0f;
 
+SActorPreview::~SActorPreview()
+{
+	USelection::SelectObjectEvent.RemoveAll(this);
+}
+
 void SActorPreview::Construct( const FArguments& InArgs )
 {
+	const int32 HorizSpacingBetweenViewports = 18;
+	const int32 PaddingBeforeBorder = 6;
+
+	USelection::SelectObjectEvent.AddRaw(this, &SActorPreview::OnActorSelected);
+
+	// We don't want the border to be hit testable, since it would just get in the way of other
+	// widgets that are added to the viewport overlay.
+	this->SetVisibility(EVisibility::SelfHitTestInvisible);
+
 	this->ChildSlot.Widget =
-		SNew( SOverlay )
-			.Visibility( EVisibility::SelfHitTestInvisible )
+		SNew(SBorder)
+		.Padding(0)
+
+		.Visibility(EVisibility::SelfHitTestInvisible)
+
+		.BorderImage(FEditorStyle::GetBrush("NoBorder"))
+		.HAlign(HAlign_Right)
+		.VAlign(VAlign_Bottom)
+		.Padding(FMargin(0, 0, PaddingBeforeBorder, PaddingBeforeBorder))
+		[
+			SNew( SOverlay )
 			+SOverlay::Slot()
 			[
 				SNew( SBorder )
@@ -2369,7 +2402,7 @@ void SActorPreview::Construct( const FArguments& InArgs )
 					.Visibility( EVisibility::HitTestInvisible )
 
 					.Padding( 16.0f )
-					.BorderImage( FEditorStyle::GetBrush( "UniformShadow" ) )
+					.BorderImage( FEditorStyle::GetBrush( "UniformShadow_Tint" ) )
 
 					.BorderBackgroundColor( this, &SActorPreview::GetBorderColorAndOpacity )
 					.ColorAndOpacity( this, &SActorPreview::GetColorAndOpacity )
@@ -2380,7 +2413,6 @@ void SActorPreview::Construct( const FArguments& InArgs )
 							.HeightOverride(this, &SActorPreview::OnReadHeight )
 							[
 								SNew( SOverlay )
-								.Visibility( EVisibility::SelfHitTestInvisible )
 									+SOverlay::Slot()
 									[
 										SAssignNew( ViewportWidget, SViewport )
@@ -2388,7 +2420,6 @@ void SActorPreview::Construct( const FArguments& InArgs )
 											.IsEnabled( FSlateApplication::Get().GetNormalExecutionAttribute() )
 											.EnableGammaCorrection( false )		// Scene rendering handles gamma correction
 											.EnableBlending( true )
-											.Visibility( EVisibility::SelfHitTestInvisible )
 									]
 									
 									+SOverlay::Slot()
@@ -2400,7 +2431,6 @@ void SActorPreview::Construct( const FArguments& InArgs )
 											.Font( FSlateFontInfo( FPaths::EngineContentDir() / TEXT("Slate/Fonts/Roboto-Bold.ttf"), 10 ) )
 											.ShadowOffset( FVector2D::UnitVector )
 											.WrapTextAt( this, &SActorPreview::OnReadTextWidth )
-											.Visibility( EVisibility::HitTestInvisible )
 									]
 							]
 					]
@@ -2431,7 +2461,7 @@ void SActorPreview::Construct( const FArguments& InArgs )
 					// Pass along the block's tool-tip string
 					.ToolTipText( this, &SActorPreview::GetPinButtonToolTipText )
 			]
-		;
+		];
 
 	// Setup animation curve for fading in and out.  Note that we add a bit of lead-in time on the fade-in
 	// to avoid hysteresis as the user moves the mouse over the view
@@ -2448,6 +2478,8 @@ void SActorPreview::Construct( const FArguments& InArgs )
 		FadeSequence.Play( TimeBeforeFadingIn );	// Skip the initial time delay and just fade straight in
 	}
 
+	HighlightSequence = FCurveSequence(0.f, 0.5f, ECurveEaseFunction::Linear);
+
 	PreviewActorPtr = InArgs._PreviewActor;
 	ParentViewport = InArgs._ParentViewport;
 }
@@ -2571,11 +2603,38 @@ FLinearColor SActorPreview::GetColorAndOpacity() const
 	return Color;
 }
 
+void SActorPreview::OnActorSelected(UObject* InActor)
+{
+	if (InActor && InActor == PreviewActorPtr && InActor->IsSelected())
+	{
+		TSharedPtr<SLevelViewport> ParentViewportPtr = ParentViewport.Pin();
+		const bool bIsPreviewPinned = ParentViewportPtr.IsValid() && ParentViewportPtr->IsActorPreviewPinned(PreviewActorPtr);
+
+		if (bIsPreviewPinned)
+		{
+			Highlight();
+		}
+	}
+}
+
+void SActorPreview::Highlight()
+{
+	HighlightSequence.JumpToStart();
+	HighlightSequence.Play();
+}
 
 FSlateColor SActorPreview::GetBorderColorAndOpacity() const
 {
-	FLinearColor Color = GetColorAndOpacity();
-	Color.A *= 0.6f;		// Make the border more transparent
+	FLinearColor Color(0.f, 0.f, 0.f, 0.5f);
+
+	if (HighlightSequence.IsPlaying())
+	{
+		const FLinearColor SelectionColor = FEditorStyle::Get().GetSlateColor("SelectionColor").GetSpecifiedColor().CopyWithNewOpacity(0.5f);
+		
+		const float Interp = FMath::Sin(HighlightSequence.GetLerp()*6*PI) / 2 + 1;
+		Color = FMath::Lerp(SelectionColor, Color, Interp);
+	}
+	
 	return Color;
 }
 
@@ -2769,34 +2828,10 @@ void SLevelViewport::PreviewActors( const TArray< AActor* >& ActorsToPreview )
 				ActorPreviewLevelViewportClient->SetControllingActor( CurActor );
 				ActorPreviewLevelViewportClient->PushControllingActorDataToViewportClient();
 			}
-			
-			const int32 HorizSpacingBetweenViewports = 18;
-			const int32 PaddingBeforeBorder = 6;
-
-			// Use the size of our array of preview viewports as the viewport index.  We haven't yet added this new
-			// viewport to the array.
-			const int32 ViewportIndex = ActorPreviews.Num();
-
-			TSharedPtr< SActorPreview > ActorPreviewWidget;
-			TSharedRef< SWidget > RootWidget =
-				SNew( SBorder )
-					.Padding(0)
 
-					// We don't want the border to be hit testable, since it would just get in the way of other
-					// widgets that are added to the viewport overlay.
-					.Visibility( EVisibility::SelfHitTestInvisible )
-
-					.BorderImage( FEditorStyle::GetBrush( "NoBorder" ) )
-					.HAlign( HAlign_Right )
-					.VAlign( VAlign_Bottom )
-					.Padding( FMargin( 0, 0, PaddingBeforeBorder, PaddingBeforeBorder ) )
-
-					[
-						SAssignNew( ActorPreviewWidget, SActorPreview )
-							.PreviewActor( CurActor )
-							.ParentViewport( SharedThis( this ) )
-					]
-			;
+			TSharedPtr< SActorPreview > ActorPreviewWidget = SNew(SActorPreview)
+				.PreviewActor(CurActor)
+				.ParentViewport(SharedThis(this));
 
 			auto ActorPreviewViewportWidget = ActorPreviewWidget->GetViewportWidget();
 
@@ -2810,13 +2845,13 @@ void SLevelViewport::PreviewActors( const TArray< AActor* >& ActorsToPreview )
 			NewActorPreview.Actor = CurActor;
 			NewActorPreview.LevelViewportClient = ActorPreviewLevelViewportClient;
 			NewActorPreview.SceneViewport = ActorPreviewSceneViewport;
-			NewActorPreview.PreviewWidget = RootWidget;
+			NewActorPreview.PreviewWidget = ActorPreviewWidget;
 			NewActorPreview.bIsPinned = false;
 
 			// Add our new widget to our viewport's overlay
 			// @todo camerapip: Consider using a canvas instead of an overlay widget -- our viewports get SQUASHED when the view shrinks!
 			auto& HorizontalBoxSlot = ActorPreviewHorizontalBox->AddSlot().AutoWidth();
-			HorizontalBoxSlot.Widget = RootWidget;
+			HorizontalBoxSlot.Widget = ActorPreviewWidget.ToSharedRef();
 		}
 
 		// OK, at least one new preview viewport was added, so update settings for all views immediately.
diff --git a/Engine/Source/Editor/LevelEditor/Public/SLevelViewport.h b/Engine/Source/Editor/LevelEditor/Public/SLevelViewport.h
--- a/Engine/Source/Editor/LevelEditor/Public/SLevelViewport.h
+++ b/Engine/Source/Editor/LevelEditor/Public/SLevelViewport.h
@@ -5,6 +5,8 @@
 
 #include "ILevelViewport.h"
 #include "Editor/UnrealEd/Public/SEditorViewport.h"
+
+class SActorPreview;
 	
 /**
  * Encapsulates an SViewport and an SLevelViewportToolBar
@@ -745,7 +747,7 @@ private:
 		TSharedPtr< FSceneViewport > SceneViewport;
 
 		/** Slate widget that represents this preview in the viewport */
-		TSharedPtr< SWidget > PreviewWidget;
+		TSharedPtr< SActorPreview > PreviewWidget;
 
 		/** Whether or not this actor preview will remain on screen if the actor is deselected */
 		bool bIsPinned;		
