Fixed Hide Viewport UI checkbox so that it actually hides the Viewport UI.
#ttp 338945 - Editor: Quick Settings: Hide Viewport UI checkbox is not hiding Viewport UI
#branch UE4
#proj Editor.LevelEditor, Editor.UnrealEd
#add Added extra checks into SLevelViewport::GetCurrentLevelTextVisibility and SLevelViewport::GetToolBarVisibility so they take into account the state of the Hide Viewport UI flag.
#add FEditorViewportClient::Draw only draws the axes if the Hide Viewport UI flag is clear.

diff --git a/Engine/Source/Editor/LevelEditor/Private/SLevelViewport.cpp b/Engine/Source/Editor/LevelEditor/Private/SLevelViewport.cpp
--- a/Engine/Source/Editor/LevelEditor/Private/SLevelViewport.cpp
+++ b/Engine/Source/Editor/LevelEditor/Private/SLevelViewport.cpp
@@ -1359,7 +1359,7 @@ EVisibility SLevelViewport::OnGetViewportContentVisibility() const
 EVisibility SLevelViewport::GetToolBarVisibility() const
 {
 	// Do not show the toolbar if this viewport has a play in editor session
-	return IsPlayInEditorViewportActive() ? EVisibility::Collapsed : EVisibility::Visible;
+	return IsPlayInEditorViewportActive() ? EVisibility::Collapsed : OnGetViewportContentVisibility();
 }
 
 EVisibility SLevelViewport::GetMaximizeToggleVisibility() const
@@ -2994,7 +2994,12 @@ FText SLevelViewport::GetCurrentLevelText( bool bDrawOnlyLabel ) const
 
 EVisibility SLevelViewport::GetCurrentLevelTextVisibility() const
 {
-	return (&GetLevelViewportClient() == GCurrentLevelEditingViewportClient) ? EVisibility::SelfHitTestInvisible : EVisibility::Collapsed;
+	EVisibility ContentVisibility = OnGetViewportContentVisibility();
+	if (ContentVisibility == EVisibility::Visible)
+	{
+		ContentVisibility = EVisibility::SelfHitTestInvisible;
+	}
+	return (&GetLevelViewportClient() == GCurrentLevelEditingViewportClient) ? ContentVisibility : EVisibility::Collapsed;
 }
 
 void SLevelViewport::OnSetViewportConfiguration(FName ConfigurationName)
diff --git a/Engine/Source/Editor/UnrealEd/Private/EditorViewportClient.cpp b/Engine/Source/Editor/UnrealEd/Private/EditorViewportClient.cpp
--- a/Engine/Source/Editor/UnrealEd/Private/EditorViewportClient.cpp
+++ b/Engine/Source/Editor/UnrealEd/Private/EditorViewportClient.cpp
@@ -2537,7 +2537,7 @@ void FEditorViewportClient::Draw(FViewport* InViewport, FCanvas* Canvas)
 	}
 
 	// Axes indicators
-	if( bDrawAxes && !ViewFamily.EngineShowFlags.Game )
+	if (bDrawAxes && !ViewFamily.EngineShowFlags.Game && !GLevelEditorModeTools().IsViewportUIHidden())
 	{
 		switch (GetViewportType())
 		{
