Fix crash switching layouts when in geometry mode

diff --git a/Engine/Source/Editor/LevelEditor/Private/SLevelViewport.cpp b/Engine/Source/Editor/LevelEditor/Private/SLevelViewport.cpp
--- a/Engine/Source/Editor/LevelEditor/Private/SLevelViewport.cpp
+++ b/Engine/Source/Editor/LevelEditor/Private/SLevelViewport.cpp
@@ -3011,6 +3011,8 @@ void SLevelViewport::OnSetViewportConfiguration(FName ConfigurationName)
 		if (ViewportTabPinned.IsValid())
 		{
 			ViewportTabPinned->SetViewportConfiguration(ConfigurationName);
+			// Viewport clients are going away.  Any current one is invalid.
+			GCurrentLevelEditingViewportClient = nullptr;
 			FSlateApplication::Get().DismissAllMenus();
 		}
 	}
