Extensibility - Adding an event plugin authors can hook in the level editor module to capture when the tab manager is changed (OnTabManagerChanged).  After that event it's safe to call GetLevelEditorTabManager and register your tab spawner from a plugin.

diff --git a/Engine/Source/Editor/LevelEditor/Private/LevelEditor.cpp b/Engine/Source/Editor/LevelEditor/Private/LevelEditor.cpp
--- a/Engine/Source/Editor/LevelEditor/Private/LevelEditor.cpp
+++ b/Engine/Source/Editor/LevelEditor/Private/LevelEditor.cpp
@@ -390,6 +390,8 @@ void FLevelEditorModule::SetLevelEditorTabManager( const TSharedPtr<SDockTab>& O
 	{
 		LevelEditorTabManager = FGlobalTabmanager::Get()->NewTabManager(OwnerTab.ToSharedRef());
 		LevelEditorTabManager->SetOnPersistLayout(FTabManager::FOnPersistLayout::CreateRaw(this, &FLevelEditorModule::HandleTabManagerPersistLayout));
+
+		TabManagerChangedEvent.Broadcast();
 	}
 }
 
diff --git a/Engine/Source/Editor/LevelEditor/Private/SLevelEditor.cpp b/Engine/Source/Editor/LevelEditor/Private/SLevelEditor.cpp
--- a/Engine/Source/Editor/LevelEditor/Private/SLevelEditor.cpp
+++ b/Engine/Source/Editor/LevelEditor/Private/SLevelEditor.cpp
@@ -1296,17 +1296,6 @@ void SLevelEditor::RefreshEditorModeCommands()
 				FIsActionChecked::CreateStatic( &SLevelEditor::IsModeActive, Mode.ID ));
 		}
 
-		// Register the new tab spawner, and unregister the old one if it exists
-		//if ( LevelEditorTabManager.IsValid() )
-		//{
-		//	LevelEditorTabManager->UnregisterTabSpawner( EditorModeTabName );
-
-		//	LevelEditorTabManager->RegisterTabSpawner( EditorModeTabName, FOnSpawnTab::CreateSP<SLevelEditor, FEdMode*>(this, &SLevelEditor::SpawnLevelEditorModeTab, Mode) )
-		//		.SetDisplayName( Mode->GetName() ) 
-		//		.SetGroup( MenuStructure.GetLevelEditorModesCategory() )
-		//		.SetIcon( Mode->GetIcon() );
-		//}
-
 		commandIndex++;
 	}
 
@@ -1320,33 +1309,6 @@ void SLevelEditor::RefreshEditorModeCommands()
 	}
 }
 
-//TSharedRef<SDockTab> SLevelEditor::SpawnLevelEditorModeTab( const FSpawnTabArgs& Args, FEdMode* EditorMode )
-//{
-	//if (!GLevelEditorModeTools().IsModeActive(EditorMode->GetID()))
-	//{
-	//	GLevelEditorModeTools().ActivateMode(EditorMode->GetID());
-	//}
-
-	//TSharedPtr<SDockTab> NewEditorModeTab;
-	//TSharedPtr<SLevelEditorModeContent> NewToolBox;
-
-	//SAssignNew( NewEditorModeTab, SDockTab )
-	//	.Icon( EditorMode->GetIcon().GetSmallIcon() )
-	//	.Label( EditorMode->GetName() );
-
-	//NewEditorModeTab->SetContent(
-	//	SNew( STutorialWrapper, TEXT("ToolsPanel") )
-	//	[
-	//		SAssignNew( NewToolBox, SLevelEditorModeContent, SharedThis( this ), NewEditorModeTab.ToSharedRef(), EditorMode )
-	//		.IsEnabled( FSlateApplication::Get().GetNormalExecutionAttribute() )
-	//	]
-	//);
-
-	//ModesTabs.Add( NewToolBox );
-
-	//return NewEditorModeTab.ToSharedRef();
-//}
-
 FReply SLevelEditor::OnKeyDown( const FGeometry& MyGeometry, const FKeyboardEvent& InKeyboardEvent )
 {
 	// Check to see if any of the actions for the level editor can be processed by the current keyboard event
@@ -1433,7 +1395,6 @@ void SLevelEditor::SummonLevelViewportContextMenu()
 	FLevelEditorContextMenu::SummonMenu( SharedThis( this ), LevelEditorMenuContext::Viewport );
 }
 
-
 const TArray< TSharedPtr< IToolkit > >& SLevelEditor::GetHostedToolkits() const
 {
 	return HostedToolkits;
diff --git a/Engine/Source/Editor/LevelEditor/Public/LevelEditor.h b/Engine/Source/Editor/LevelEditor/Public/LevelEditor.h
--- a/Engine/Source/Editor/LevelEditor/Public/LevelEditor.h
+++ b/Engine/Source/Editor/LevelEditor/Public/LevelEditor.h
@@ -152,6 +152,10 @@ public:
 	/* Create a tab manager for the level editor based on the given tab (or clears the tab manager if OwnerTab is null) */
 	virtual void SetLevelEditorTabManager( const TSharedPtr<SDockTab>& OwnerTab );
 
+	/** Called when the tab manager is changed */
+	DECLARE_EVENT(FLevelEditorModule, FTabManagerChangedEvent);
+	virtual FTabManagerChangedEvent& OnTabManagerChanged() { return TabManagerChangedEvent; }
+
 	/**
 	 * Called when actor selection changes
 	 * 
@@ -253,6 +257,9 @@ private:
 
 	FAutoConsoleCommand ToggleImmersiveConsoleCommand;
 
+	/** Multicast delegate executed when the tab manager is changed */
+	FTabManagerChangedEvent TabManagerChangedEvent;
+
 	/** Multicast delegate executed when actor selection changes */
 	FActorSelectionChangedEvent ActorSelectionChangedEvent;
 
