Switched "Levels" tool to a World Browser. Old Levels tool can be enabled with -oldlevels command line switch.

diff --git a/Engine/Source/Editor/EditorStyle/Private/SlateEditorStyle.cpp b/Engine/Source/Editor/EditorStyle/Private/SlateEditorStyle.cpp
--- a/Engine/Source/Editor/EditorStyle/Private/SlateEditorStyle.cpp
+++ b/Engine/Source/Editor/EditorStyle/Private/SlateEditorStyle.cpp
@@ -3338,7 +3338,9 @@ void FSlateEditorStyle::FStyle::SetupLevelEditorStyle()
 			Set( "LevelEditor.Tabs.Outliner", new IMAGE_BRUSH( "/Icons/icon_tab_SceneOutliner_16x", Icon16x16 ) );
 			Set( "LevelEditor.Tabs.ContentBrowser", new IMAGE_BRUSH( "/Icons/icon_tab_ContentBrowser_16x", Icon16x16 ) );
 			Set( "LevelEditor.Tabs.Levels", new IMAGE_BRUSH( "/Icons/icon_tab_Levels_16x", Icon16x16 ) );
-			Set( "LevelEditor.Tabs.WorldBrowser", new IMAGE_BRUSH( "/Icons/icon_tab_Levels_16x", Icon16x16 ) );
+			Set( "LevelEditor.Tabs.WorldBrowser", new IMAGE_BRUSH( "/Icons/icon_tab_levels_16x", Icon16x16 ) );
+			Set( "LevelEditor.Tabs.WorldBrowserDetails", new IMAGE_BRUSH( "/Icons/icon_levels_detailsbutton_16x", Icon16x16 ) );
+			Set( "LevelEditor.Tabs.WorldBrowserComposition", new IMAGE_BRUSH( "/Icons/icon_levels_compositionbutton_16x", Icon16x16 ) );
 			Set( "LevelEditor.Tabs.Layers", new IMAGE_BRUSH( "/Icons/icon_tab_Layers_16x", Icon16x16 ) );
 			Set( "LevelEditor.Tabs.BuildAndSubmit", new IMAGE_BRUSH( "/Icons/icon_tab_BuildSubmit_16x", Icon16x16 ) );
 			Set( "LevelEditor.Tabs.StatsViewer", new IMAGE_BRUSH( "/Icons/icon_tab_Stats_16x", Icon16x16 ) );
diff --git a/Engine/Source/Editor/LevelEditor/Private/SLevelEditor.cpp b/Engine/Source/Editor/LevelEditor/Private/SLevelEditor.cpp
--- a/Engine/Source/Editor/LevelEditor/Private/SLevelEditor.cpp
+++ b/Engine/Source/Editor/LevelEditor/Private/SLevelEditor.cpp
@@ -681,8 +681,10 @@ TSharedRef<SDockTab> SLevelEditor::SpawnLevelEditorTab( const FSpawnTabArgs& Arg
 					]
 				];
 	}
-	else if( TabIdentifier == TEXT("LevelEditorLevelBrowser") )
+	else if( TabIdentifier == WorldBrowserHierarchyTab || TabIdentifier == TEXT("LevelEditorLevelBrowser") )
 	{
+		if (FParse::Param(FCommandLine::Get(), TEXT("oldlevels")))
+		{
 			FLevelsModule& LevelsModule = FModuleManager::LoadModuleChecked<FLevelsModule>( "Levels" );
 			return SNew( SDockTab )
 				.Icon( FEditorStyle::GetBrush( "LevelEditor.Tabs.Levels" ) )
@@ -695,32 +697,34 @@ TSharedRef<SDockTab> SLevelEditor::SpawnLevelEditorTab( const FSpawnTabArgs& Arg
 						LevelsModule.CreateLevelBrowser()
 					]
 				];
+		}
+		else
+		{
+			FWorldBrowserModule& WorldBrowserModule = FModuleManager::LoadModuleChecked<FWorldBrowserModule>( "WorldBrowser" );
+			return SNew( SDockTab )
+				.Icon( FEditorStyle::GetBrush( "LevelEditor.Tabs.WorldBrowser" ) )
+				.Label( NSLOCTEXT("LevelEditor", "WorldBrowserHierarchyTabTitle", "Levels") )
+				[
+					WorldBrowserModule.CreateWorldBrowserHierarchy()
+				];
+		}
+				
 	}
-	else if( GetDefault<UEditorExperimentalSettings>()->bWorldBrowser && TabIdentifier == WorldBrowserHierarchyTab)
-	{
-		FWorldBrowserModule& WorldBrowserModule = FModuleManager::LoadModuleChecked<FWorldBrowserModule>( "WorldBrowser" );
-		return SNew( SDockTab )
-			.Icon( FEditorStyle::GetBrush( "LevelEditor.Tabs.WorldBrowser" ) )
-			.Label( NSLOCTEXT("LevelEditor", "WorldBrowserHierarchyTabTitle", "Levels Hierarchy") )
-			[
-				WorldBrowserModule.CreateWorldBrowserHierarchy()
-			];
-	}
-	else if( GetDefault<UEditorExperimentalSettings>()->bWorldBrowser && TabIdentifier == WorldBrowserDetailsTab)
+	else if( TabIdentifier == WorldBrowserDetailsTab && !FParse::Param(FCommandLine::Get(), TEXT("oldlevels")) )
 	{
 		FWorldBrowserModule& WorldBrowserModule = FModuleManager::LoadModuleChecked<FWorldBrowserModule>( "WorldBrowser" );
 		return SNew( SDockTab )
-			.Icon( FEditorStyle::GetBrush( "LevelEditor.Tabs.WorldBrowser" ) )
-			.Label( NSLOCTEXT("LevelEditor", "WorldBrowserDetailsTabTitle", "Levels Details") )
+			.Icon( FEditorStyle::GetBrush( "LevelEditor.Tabs.WorldBrowserDetails" ) )
+			.Label( NSLOCTEXT("LevelEditor", "WorldBrowserDetailsTabTitle", "Level Details") )
 			[
 				WorldBrowserModule.CreateWorldBrowserDetails()
 			];
 	}
-	else if( GetDefault<UEditorExperimentalSettings>()->bWorldBrowser && TabIdentifier == WorldBrowserCompositionTab)
+	else if( TabIdentifier == WorldBrowserCompositionTab && !FParse::Param(FCommandLine::Get(), TEXT("oldlevels")) )
 	{
 		FWorldBrowserModule& WorldBrowserModule = FModuleManager::LoadModuleChecked<FWorldBrowserModule>( "WorldBrowser" );
 		return SNew( SDockTab )
-			.Icon( FEditorStyle::GetBrush( "LevelEditor.Tabs.WorldBrowser" ) )
+			.Icon( FEditorStyle::GetBrush( "LevelEditor.Tabs.WorldBrowserComposition" ) )
 			.Label( NSLOCTEXT("LevelEditor", "WorldBrowserCompositionTabTitle", "Levels Composition") )
 			[
 				WorldBrowserModule.CreateWorldBrowserComposition()
@@ -1016,6 +1020,7 @@ TSharedRef<SWidget> SLevelEditor::RestoreContentArea( const TSharedRef<SDockTab>
 				.SetIcon( LayersIcon );
 		}
 		
+		if (FParse::Param(FCommandLine::Get(), TEXT("oldlevels")))
 		{
 			const FSlateIcon LevelsIcon(FEditorStyle::GetStyleSetName(), "LevelEditor.Tabs.Levels");
 			LevelEditorTabManager->RegisterTabSpawner( "LevelEditorLevelBrowser", FOnSpawnTab::CreateSP<SLevelEditor, FName, FString>(this, &SLevelEditor::SpawnLevelEditorTab, FName("LevelEditorLevelBrowser"), FString()) )
@@ -1024,12 +1029,27 @@ TSharedRef<SWidget> SLevelEditor::RestoreContentArea( const TSharedRef<SDockTab>
 				.SetGroup( MenuStructure.GetLevelEditorCategory() )
 				.SetIcon( LevelsIcon );
 		}
-
-		if (GetDefault<UEditorExperimentalSettings>()->bWorldBrowser)
+		else
 		{
-			RegisterWorldBrowserTabSpawner();
+			LevelEditorTabManager->RegisterTabSpawner( WorldBrowserHierarchyTab, FOnSpawnTab::CreateSP<SLevelEditor, FName, FString>(this, &SLevelEditor::SpawnLevelEditorTab, WorldBrowserHierarchyTab, FString()) )
+				.SetDisplayName(NSLOCTEXT("LevelEditorTabs", "WorldBrowserHierarchy", "Levels"))
+				.SetTooltipText(NSLOCTEXT("LevelEditorTabs", "WorldBrowserHierarchyTooltipText", "Open the Levels tab. Use this to manage the levels in the current project."))
+				.SetGroup( WorkspaceMenu::GetMenuStructure().GetLevelEditorCategory() )
+				.SetIcon( FSlateIcon(FEditorStyle::GetStyleSetName(), "LevelEditor.Tabs.WorldBrowser") );
+			
+			LevelEditorTabManager->RegisterTabSpawner( WorldBrowserDetailsTab, FOnSpawnTab::CreateSP<SLevelEditor, FName, FString>(this, &SLevelEditor::SpawnLevelEditorTab, WorldBrowserDetailsTab, FString()) )
+				.SetMenuType( ETabSpawnerMenuType::Hide )
+				.SetDisplayName(NSLOCTEXT("LevelEditorTabs", "WorldBrowserDetails", "Level Details"))
+				.SetGroup( WorkspaceMenu::GetMenuStructure().GetLevelEditorCategory() )
+				.SetIcon( FSlateIcon(FEditorStyle::GetStyleSetName(), "LevelEditor.Tabs.WorldBrowserDetails") );
+		
+			LevelEditorTabManager->RegisterTabSpawner( WorldBrowserCompositionTab, FOnSpawnTab::CreateSP<SLevelEditor, FName, FString>(this, &SLevelEditor::SpawnLevelEditorTab, WorldBrowserCompositionTab, FString()) )
+				.SetMenuType( ETabSpawnerMenuType::Hide )
+				.SetDisplayName(NSLOCTEXT("LevelEditorTabs", "WorldBrowserComposition", "Levels Composition"))
+				.SetGroup( WorkspaceMenu::GetMenuStructure().GetLevelEditorCategory() )
+				.SetIcon( FSlateIcon(FEditorStyle::GetStyleSetName(), "LevelEditor.Tabs.WorldBrowserComposition") );
 		}
-
+		
 		{
 			const FSlateIcon StatsViewerIcon(FEditorStyle::GetStyleSetName(), "LevelEditor.Tabs.StatsViewer");
 			LevelEditorTabManager->RegisterTabSpawner( LevelEditorStatsViewerTab, FOnSpawnTab::CreateSP<SLevelEditor, FName, FString>(this, &SLevelEditor::SpawnLevelEditorTab, LevelEditorStatsViewerTab, FString()) )
@@ -1151,47 +1171,6 @@ TSharedRef<SWidget> SLevelEditor::RestoreContentArea( const TSharedRef<SDockTab>
 
 void SLevelEditor::HandleExperimentalSettingChanged(FName PropertyName)
 {
-	if (PropertyName == TEXT("bWorldBrowser"))
-	{
-		if (GetDefault<UEditorExperimentalSettings>()->bWorldBrowser )
-		{
-			RegisterWorldBrowserTabSpawner();
-		}
-		else
-		{
-			UnregisterWorldBrowserTabSpawner();
-		}
-	}
-}
-
-void SLevelEditor::RegisterWorldBrowserTabSpawner()
-{
-	TSharedPtr<FTabManager> LevelEditorTabManager = GetTabManager();
-
-	LevelEditorTabManager->RegisterTabSpawner( WorldBrowserHierarchyTab, FOnSpawnTab::CreateSP<SLevelEditor, FName, FString>(this, &SLevelEditor::SpawnLevelEditorTab, WorldBrowserHierarchyTab, FString()) )
-		.SetDisplayName(NSLOCTEXT("LevelEditorTabs", "WorldBrowserHierarchy", "World Browser"))
-		.SetGroup( WorkspaceMenu::GetMenuStructure().GetLevelEditorCategory() )
-		.SetIcon( FSlateIcon(FEditorStyle::GetStyleSetName(), "LevelEditor.Tabs.WorldBrowser") );
-			
-	LevelEditorTabManager->RegisterTabSpawner( WorldBrowserDetailsTab, FOnSpawnTab::CreateSP<SLevelEditor, FName, FString>(this, &SLevelEditor::SpawnLevelEditorTab, WorldBrowserDetailsTab, FString()) )
-		.SetMenuType( ETabSpawnerMenuType::Hide )
-		.SetDisplayName(NSLOCTEXT("LevelEditorTabs", "WorldBrowserDetails", "Levels Details"))
-		.SetGroup( WorkspaceMenu::GetMenuStructure().GetLevelEditorCategory() )
-		.SetIcon( FSlateIcon(FEditorStyle::GetStyleSetName(), "LevelEditor.Tabs.WorldBrowser") );
-		
-	LevelEditorTabManager->RegisterTabSpawner( WorldBrowserCompositionTab, FOnSpawnTab::CreateSP<SLevelEditor, FName, FString>(this, &SLevelEditor::SpawnLevelEditorTab, WorldBrowserCompositionTab, FString()) )
-		.SetMenuType( ETabSpawnerMenuType::Hide )
-		.SetDisplayName(NSLOCTEXT("LevelEditorTabs", "WorldBrowserComposition", "Levels Composition"))
-		.SetGroup( WorkspaceMenu::GetMenuStructure().GetLevelEditorCategory() )
-		.SetIcon( FSlateIcon(FEditorStyle::GetStyleSetName(), "LevelEditor.Tabs.WorldBrowser") );
-}
-
-void SLevelEditor::UnregisterWorldBrowserTabSpawner()
-{
-	TSharedPtr<FTabManager> LevelEditorTabManager = GetTabManager();
-	LevelEditorTabManager->UnregisterTabSpawner( WorldBrowserHierarchyTab );
-	LevelEditorTabManager->UnregisterTabSpawner( WorldBrowserDetailsTab );
-	LevelEditorTabManager->UnregisterTabSpawner( WorldBrowserCompositionTab );
 }
 
 FName SLevelEditor::GetEditorModeTabId( FEditorModeID ModeID )
diff --git a/Engine/Source/Editor/LevelEditor/Private/SLevelEditor.h b/Engine/Source/Editor/LevelEditor/Private/SLevelEditor.h
--- a/Engine/Source/Editor/LevelEditor/Private/SLevelEditor.h
+++ b/Engine/Source/Editor/LevelEditor/Private/SLevelEditor.h
@@ -145,12 +145,6 @@ private:
 	/** Called when a property is changed */
 	void HandleExperimentalSettingChanged(FName PropertyName);
 
-	/** Registers the tab spawner for the WorldBrowser tab */
-	void RegisterWorldBrowserTabSpawner();
-
-	/** Unregisters the tab spawner for the WorldBrowser tab */
-	void UnregisterWorldBrowserTabSpawner();
-
 	/** Rebuilds the command list for spawning editor modes, this is done when new modes are registered. */
 	void RefreshEditorModeCommands();
 
diff --git a/Engine/Source/Editor/UnrealEd/Classes/Settings/EditorExperimentalSettings.h b/Engine/Source/Editor/UnrealEd/Classes/Settings/EditorExperimentalSettings.h
--- a/Engine/Source/Editor/UnrealEd/Classes/Settings/EditorExperimentalSettings.h
+++ b/Engine/Source/Editor/UnrealEd/Classes/Settings/EditorExperimentalSettings.h
@@ -41,10 +41,6 @@ public:
 	UPROPERTY(EditAnywhere, config, Category=Tools, meta=(DisplayName="Messaging Debugger"))
 	bool bMessagingDebugger;
 
-	/** The World browser allows you to edit a big world by manipulating individual levels positions and streaming properties */
-	UPROPERTY(EditAnywhere, config, Category=Tools, meta=(DisplayName="World Browser"))
-	bool bWorldBrowser;
-
 	/** Allows to use actor merging utilities (Simplygon Proxy LOD, Grouping by Materials)*/
 	UPROPERTY(EditAnywhere, config, Category=Tools, meta=(DisplayName="Actor Merging"))
 	bool bActorMerging;
diff --git a/Engine/Source/Editor/WorldBrowser/Private/SWorldDetails.cpp b/Engine/Source/Editor/WorldBrowser/Private/SWorldDetails.cpp
--- a/Engine/Source/Editor/WorldBrowser/Private/SWorldDetails.cpp
+++ b/Engine/Source/Editor/WorldBrowser/Private/SWorldDetails.cpp
@@ -107,7 +107,7 @@ void SWorldDetails::OnBrowseWorld(UWorld* InWorld)
 						SNew(SButton)
 						.ButtonStyle(FEditorStyle::Get(), "ToggleButton")
 						.OnClicked(this, &SWorldDetails::OnSummonHierarchy)
-						.ToolTipText(LOCTEXT("SummonHierarchyToolTipText", "Summons sublevels hierarchy"))
+						.ToolTipText(LOCTEXT("SummonHierarchyToolTipText", "Summons levels hierarchy"))
 						.HAlign(HAlign_Center)
 						.VAlign(VAlign_Center)
 						.Content()
diff --git a/Engine/Source/Editor/WorldBrowser/Private/SWorldHierarchy.cpp b/Engine/Source/Editor/WorldBrowser/Private/SWorldHierarchy.cpp
--- a/Engine/Source/Editor/WorldBrowser/Private/SWorldHierarchy.cpp
+++ b/Engine/Source/Editor/WorldBrowser/Private/SWorldHierarchy.cpp
@@ -582,7 +582,7 @@ void SWorldHierarchy::OnBrowseWorld(UWorld* InWorld)
 						]
 					]
 
-					// Button to summon levels details tab
+					// Button to summon level details tab
 					+SHorizontalBox::Slot()
 					.AutoWidth()
 					.VAlign(VAlign_Center)
@@ -591,7 +591,7 @@ void SWorldHierarchy::OnBrowseWorld(UWorld* InWorld)
 						SNew(SButton)
 						.ButtonStyle(FEditorStyle::Get(), "ToggleButton")
 						.OnClicked(this, &SWorldHierarchy::OnSummonDetails)
-						.ToolTipText(LOCTEXT("SummonDetailsToolTipText", "Summons sublevels details"))
+						.ToolTipText(LOCTEXT("SummonDetailsToolTipText", "Summons level details"))
 						.HAlign(HAlign_Center)
 						.VAlign(VAlign_Center)
 						.Content()
