Added low FPS detection to the editor

FPerformanceWarningNotification constantly monitors the current framerate of the engine and prompts the user when the framerate falls below a desirable value. The heuristics for "low-fps" can be configured with the following CVars:
    * PerfWarn.FineSampleTime, PerfWarn.FineMinFPS, PerfWarn.FinePercentThreshold
    * PerfWarn.CoarseSampleTime, PerfWarn.CoarseMinFPS, PerfWarn.CoarsePercentThreshold

These variables define the length of sample (in seconds), minimum framerate, and percentage threshold that the sample needs to have been under the minimum framerate to prompt the user. Auto-application of scalability settings will be adde soon.

This addresses TTP#337196 - UE4: EDITOR: Low fps detection, Inform User and provide solution using lower scalability settings

diff --git a/Engine/Source/Editor/EditorStyle/Private/SlateEditorStyle.cpp b/Engine/Source/Editor/EditorStyle/Private/SlateEditorStyle.cpp
--- a/Engine/Source/Editor/EditorStyle/Private/SlateEditorStyle.cpp
+++ b/Engine/Source/Editor/EditorStyle/Private/SlateEditorStyle.cpp
@@ -3797,7 +3797,12 @@ void FSlateEditorStyle::FStyle::SetupLevelEditorStyle()
 				.SetHovered( IMAGE_BRUSH( "/Docking/CloseApp_Hovered", Icon16x16 ) ) );
 		}
 	}
+
+	// Scalability (Performance Warning)
+	{
+		Set( "Scalability.ScalabilitySettings", new IMAGE_BRUSH("Scalability/ScalabilitySettings", FVector2D(473.0f, 266.0f) ) );
 	}
+}
 
 void FSlateEditorStyle::FStyle::SetupPersonaStyle()
 {
diff --git a/Engine/Source/Editor/LevelEditor/Private/PerformanceWarningNotification.cpp b/Engine/Source/Editor/LevelEditor/Private/PerformanceWarningNotification.cpp
new file mode 100644
--- /dev/null
+++ b/Engine/Source/Editor/LevelEditor/Private/PerformanceWarningNotification.cpp
@@ -0,0 +1,542 @@
+// Copyright 1998-2014 Epic Games, Inc. All Rights Reserved.
+
+
+#include "LevelEditor.h"
+#include "SScalabilitySettings.h"
+
+#define LOCTEXT_NAMESPACE "PerformanceWarning"
+
+/**
+ * Helper class for a calculating a moving average. This works by maintaining an accumulator which is then sampled at regular intervals into the "Samples" array.
+ */
+struct FMovingAverage
+{
+	/** Constructor from a sample size, and sample rate */
+	FMovingAverage(const int32 InSampleSize = 0, const float InSampleRateSeconds = 1.f)
+		: SampleRateSeconds(InSampleRateSeconds)
+		, CurrentSampleAccumulator(0)
+		, CurrentSampleStartTime(0)
+		, SampleSize(InSampleSize)
+		, CurrentSampleCount(0)
+		, NextSampleIndex(0)
+		, SampleAverage(0)
+	{
+		Samples.Reserve(SampleSize);
+	}
+	
+	/** Check if this data is reliable or not. Returns false if the sampler is not populated with enough data */
+	inline float IsReliable() const
+	{
+		return SampleSize != 0 && Samples.Num() == SampleSize;
+	}
+
+	/** Reset this sampler to its default (unpopulated) state */
+	inline void Reset()
+	{
+		*this = FMovingAverage(SampleSize, SampleRateSeconds);
+	}
+
+	/** Get the average of all the samples contained in this sampler */
+	inline float GetAverage() const
+	{
+		return SampleAverage;
+	}
+
+	/** Allow this sampler to tick the frame time, and potentially save a new sample */
+	void Tick(double CurrentTime, float Value)
+	{
+		if (SampleSize == 0)
+		{
+			return;
+		}
+
+		if (CurrentSampleStartTime == 0)
+		{
+			CurrentSampleStartTime = CurrentTime;
+		}
+
+		++CurrentSampleCount;
+		CurrentSampleAccumulator += Value;
+
+		if (CurrentTime - CurrentSampleStartTime > SampleRateSeconds)
+		{
+			// Limit to a minimum of 5 frames per second
+			const float AccumulatorAverage = FMath::Max(CurrentSampleAccumulator / CurrentSampleCount, 5.f);
+
+			if (Samples.Num() == SampleSize)
+			{
+				Samples[NextSampleIndex] = AccumulatorAverage;
+			}
+			else
+			{
+				Samples.Add(AccumulatorAverage);
+			}
+
+			NextSampleIndex = ++NextSampleIndex % SampleSize;
+			ensure(NextSampleIndex >= 0 && NextSampleIndex < SampleSize);
+
+			// calculate the average
+			float Sum = 0.0;
+			for (const auto& Sample : Samples)
+			{
+				Sum += Sample;
+			}
+			SampleAverage = Sum / Samples.Num();
+
+			// reset the accumulator and counter
+			CurrentSampleAccumulator = CurrentSampleCount = 0;
+			CurrentSampleStartTime = CurrentTime;
+		}
+	}
+
+	/** Return the percentage of samples that fall below the specified threshold */
+	float PercentageBelowThreshold(const float Threshold) const
+	{
+		check(IsReliable());
+
+		float NumBelowThreshold = 0;
+		for (const auto& Sample : Samples)
+		{
+			if (Sample < Threshold)
+			{
+				++NumBelowThreshold;
+			}
+		}
+
+		return (NumBelowThreshold / Samples.Num()) * 100.0f;
+	}
+private:
+	/** The number of samples in the current accumulator */
+	int32 CurrentSampleCount;
+	/** The cumulative sum of samples for the current time period */
+	float CurrentSampleAccumulator;
+	/** The time at which we started accumulating the current sample */
+	double CurrentSampleStartTime;
+
+	/** The rate at which to store accumulated samples of FPS in seconds */
+	float SampleRateSeconds;
+	/** The maximum number of accumulated samples to store */
+	int32 SampleSize;
+	/** The actual average stored across all samples */
+	float SampleAverage;
+
+	/** The actual FPS samples */
+	TArray<float> Samples;
+	/** The index of the next sample to write to */
+	int32 NextSampleIndex;
+};
+
+/** Scalability dialog widget */
+class SScalabilitySettingsDialog : public SCompoundWidget
+{
+public:
+
+	SLATE_BEGIN_ARGS(SScalabilitySettingsDialog) {}
+		SLATE_ARGUMENT(FOnClicked, OnDoneClicked)
+	SLATE_END_ARGS()
+
+public:
+
+	/** Construct this widget */
+	void Construct(const FArguments& InArgs)
+	{
+		ChildSlot
+		[
+			SNew(SBorder)
+			.HAlign(HAlign_Fill)
+			.BorderImage(FEditorStyle::GetBrush("ChildWindow.Background"))
+			[
+				SNew(SBox)
+				.WidthOverride(500.0f)
+				[
+					SNew(SVerticalBox)
+					+ SVerticalBox::Slot()
+					.AutoHeight()
+					.Padding(5.f)
+					[
+						SNew(STextBlock).Text(LOCTEXT("PerformanceWarningDescription",
+						"The current performance of the editor seems to be low.\nUse the options below to reduce the amount of detail and increase performance.").ToString())
+					]
+					+ SVerticalBox::Slot()
+					.AutoHeight()
+					.Padding(5.f)
+					[
+						SNew(SScalabilitySettings)
+					]
+					+ SVerticalBox::Slot()
+					.AutoHeight()
+					.Padding(5.f)
+					[
+						SNew(STextBlock)
+						.ToolTip(
+							SNew(SToolTip)
+							[
+								SNew(SImage)
+								.Image(FEditorStyle::GetBrush("Scalability.ScalabilitySettings"))
+							]
+						)
+						.AutoWrapText(true)
+						.Text(LOCTEXT("PerformanceWarningChangeLater",
+							"You can modify these settings in future via \"Quick Settings\" button on the level editor toolbar and choosing \"Engine Scalability Settings\".").ToString())
+					]
+					+ SVerticalBox::Slot()
+					.AutoHeight()
+					.Padding(5.f)
+					[
+						SNew(SHorizontalBox)
+						+ SHorizontalBox::Slot()
+						[
+							SNew(SCheckBox)
+							.OnCheckStateChanged(FOnCheckStateChanged::CreateStatic(&OnShowPerformanceWarningChanged))
+							.IsChecked(TAttribute<ESlateCheckBoxState::Type>::Create(&ShouldShowPerformanceWarning))
+							.ToolTipText(LOCTEXT("PerformanceWarningEnableDisableToolTip", "Toggles whether this warning will ever be shown again."))
+							.Content()
+							[
+								SNew(STextBlock)
+								.Text(LOCTEXT("PerformanceWarningEnableDisableCheckbox", "Show this warning again in the future?"))
+							]
+						]
+						+ SHorizontalBox::Slot()
+						.AutoWidth()
+						[
+							SNew(SButton)
+							.Text(LOCTEXT("ScalabilityDone", "Done"))
+							.OnClicked(InArgs._OnDoneClicked)
+						]
+					]
+				]
+			]
+		];
+	}
+
+	/** Called to get the "Show notification" check box state */
+	static ESlateCheckBoxState::Type ShouldShowPerformanceWarning()
+	{
+		const bool bShowPerformanceWarning = GEditor->GetGameAgnosticSettings().bShowPerformanceWarningPreference;
+		return bShowPerformanceWarning ? ESlateCheckBoxState::Checked : ESlateCheckBoxState::Unchecked;
+	}
+
+	/** Called when the state of the "Show notification" check box changes */
+	static void OnShowPerformanceWarningChanged(ESlateCheckBoxState::Type NewState)
+	{
+		const bool bNewEnabledState = (NewState == ESlateCheckBoxState::Checked);
+
+		GEditor->AccessGameAgnosticSettings().bShowPerformanceWarningPreference = bNewEnabledState;
+		GEditor->AccessGameAgnosticSettings().PostEditChange();
+	}
+};
+
+/** Notification class for performance warnings. */
+class FPerformanceWarningNotification
+	: public FTickableEditorObject
+{
+
+public:
+
+	/** Constructor */
+	FPerformanceWarningNotification();
+
+	/** Destructor that unbinds the CVar changed callbacks */
+	~FPerformanceWarningNotification();
+
+	/** Update the moving average samplers to match the settings specified in the console variables */
+	void UpdateMovingAverageSamplers(IConsoleVariable* Unused = nullptr)
+	{
+		static const int NumberOfSamples = 50;
+		IConsoleManager& Console = IConsoleManager::Get();
+
+		float SampleTime = Console.FindConsoleVariable(TEXT("PerfWarn.FineSampleTime"))->GetFloat();
+		FineMovingAverage = FMovingAverage(NumberOfSamples, SampleTime / NumberOfSamples);
+
+		SampleTime = Console.FindConsoleVariable(TEXT("PerfWarn.CoarseSampleTime"))->GetFloat();
+		CoarseMovingAverage = FMovingAverage(NumberOfSamples, SampleTime / NumberOfSamples);
+	}
+
+protected:
+
+	/** FTickableEditorObject interface */
+	virtual bool IsTickable() const override { return true; }
+	virtual void Tick(float DeltaTime) override;
+	virtual TStatId GetStatId() const override { RETURN_QUICK_DECLARE_CYCLE_STAT(FPerformanceWarningNotification, STATGROUP_Tickables); }
+
+private:
+
+	/** Starts the notification. */
+	void ShowPerformanceWarning(FText MessageText);
+
+	/** Ends the notification. */
+	void HidePerformanceWarning();
+
+	/** Display the scalability dialog. */
+	void ShowScalabilityDialog();
+
+	/** Dismiss the notification and disable it for this session. */
+	void DismissAndDisableNotification()
+	{
+		bIsNotificationAllowed = false;
+		HidePerformanceWarning();
+	}
+
+	/** Resets the performance warning data and hides the warning */
+	void Reset()
+	{
+		FineMovingAverage.Reset();
+		CoarseMovingAverage.Reset();
+
+		HidePerformanceWarning();
+	}
+
+	/** Moving average data for the fine and coarse-grained moving averages */
+	FMovingAverage FineMovingAverage, CoarseMovingAverage;
+
+	/** Tracks the last time the notification was started, used to avoid spamming. */
+	double LastEnableTime;
+
+	/** The notification window ptr */
+	TWeakPtr<SNotificationItem> PerformanceWarningNotificationPtr;
+
+	/** The Scalability Setting window we may be using */
+	TWeakPtr<SWindow> ScalabilitySettingsWindowPtr;
+
+	bool bIsNotificationAllowed;
+};
+
+/** Global notification object. */
+FPerformanceWarningNotification GPerformanceWarningNotification;
+
+FPerformanceWarningNotification::FPerformanceWarningNotification()
+{
+	LastEnableTime = 0;
+	bIsNotificationAllowed = true;
+
+	// Register the console variables we need
+	IConsoleVariable* CVar = nullptr;
+
+	CVar = IConsoleManager::Get().RegisterConsoleVariable(TEXT("PerfWarn.FineSampleTime"), 30, TEXT("How many seconds we sample the percentage for the fine-grained minimum FPS."), ECVF_Default);
+	CVar->SetOnChangedCallback(FConsoleVariableDelegate::CreateRaw(this, &FPerformanceWarningNotification::UpdateMovingAverageSamplers));
+	CVar = IConsoleManager::Get().RegisterConsoleVariable(TEXT("PerfWarn.CoarseSampleTime"), 600, TEXT("How many seconds we sample the percentage for the coarse-grained minimum FPS."), ECVF_Default);
+	CVar->SetOnChangedCallback(FConsoleVariableDelegate::CreateRaw(this, &FPerformanceWarningNotification::UpdateMovingAverageSamplers));
+
+	IConsoleManager::Get().RegisterConsoleVariable(TEXT("PerfWarn.FineMinFPS"), 10, TEXT("The FPS threshold below which we warn about for fine-grained sampling."), ECVF_Default);
+	IConsoleManager::Get().RegisterConsoleVariable(TEXT("PerfWarn.CoarseMinFPS"), 20, TEXT("The FPS threshold below which we warn about for coarse-grained sampling."), ECVF_Default);
+
+	IConsoleManager::Get().RegisterConsoleVariable(TEXT("PerfWarn.FinePercentThreshold"), 80, TEXT("The percentage of samples that fall below min FPS above which we warn for."), ECVF_Default);
+	IConsoleManager::Get().RegisterConsoleVariable(TEXT("PerfWarn.CoarsePercentThreshold"), 80, TEXT("The percentage of samples that fall below min FPS above which we warn for."), ECVF_Default);
+
+	// Set up the moving average sasmplers
+	UpdateMovingAverageSamplers();
+}
+
+FPerformanceWarningNotification::~FPerformanceWarningNotification()
+{
+	IConsoleManager& Console = IConsoleManager::Get();
+	Console.FindConsoleVariable(TEXT("PerfWarn.FineSampleTime"))->SetOnChangedCallback(FConsoleVariableDelegate());
+	Console.FindConsoleVariable(TEXT("PerfWarn.CoarseSampleTime"))->SetOnChangedCallback(FConsoleVariableDelegate());
+}
+
+void FPerformanceWarningNotification::ShowPerformanceWarning(FText MessageText)
+{
+	auto AlreadyOpenItem = PerformanceWarningNotificationPtr.Pin();
+	if (AlreadyOpenItem.IsValid())
+	{
+		AlreadyOpenItem->SetText(MessageText);
+	}
+	else if ((FPlatformTime::Seconds() - LastEnableTime) > 30.0)
+	{
+		// Only show a new one if we've not shown one for a while
+		LastEnableTime = FPlatformTime::Seconds();
+
+		// Create notification item
+		FNotificationInfo Info(MessageText);
+		Info.bFireAndForget = false;
+		Info.FadeOutDuration = 1.0f;
+		Info.ExpireDuration = 0.0f;
+		Info.bUseLargeFont = false;
+
+		Info.ButtonDetails.Add(
+			FNotificationButtonInfo(
+			LOCTEXT("PerformanceWarningScalability", "Tweak Scalability Settings"),
+			LOCTEXT("PerformanceWarningScalabilityToolTip", "Opens the Scalability UI for graphical options."),
+			FSimpleDelegate::CreateRaw(this, &FPerformanceWarningNotification::ShowScalabilityDialog),
+			SNotificationItem::CS_None
+			)
+			);
+
+		Info.ButtonDetails.Add(
+			FNotificationButtonInfo(
+			LOCTEXT("PerformanceWarningDismiss", "Dismiss"),
+			LOCTEXT("PerformanceWarningDismissToolTip", "Close and do not show this notification again this session."),
+			FSimpleDelegate::CreateRaw(this, &FPerformanceWarningNotification::DismissAndDisableNotification),
+			SNotificationItem::CS_None
+			)
+			);
+
+		PerformanceWarningNotificationPtr = FSlateNotificationManager::Get().AddNotification(Info);
+	}
+}
+
+void FPerformanceWarningNotification::HidePerformanceWarning()
+{
+	// Finished! Notify the UI.
+	TSharedPtr<SNotificationItem> NotificationItem = PerformanceWarningNotificationPtr.Pin();
+
+	if (NotificationItem.IsValid())
+	{
+		NotificationItem->SetCompletionState(SNotificationItem::CS_Success);
+		NotificationItem->Fadeout();
+
+		PerformanceWarningNotificationPtr.Reset();
+	}
+}
+
+void FPerformanceWarningNotification::Tick(float DeltaTime)
+{
+	if (GEngine->ShouldThrottleCPUUsage())
+	{
+		return;
+	}
+
+	extern ENGINE_API float GAverageFPS;
+	FineMovingAverage.Tick(FPlatformTime::Seconds(), GAverageFPS);
+	CoarseMovingAverage.Tick(FPlatformTime::Seconds(), GAverageFPS);
+
+	const bool bShowPerformanceWarning = GEditor->GetGameAgnosticSettings().bShowPerformanceWarningPreference;
+	if( !bShowPerformanceWarning || !bIsNotificationAllowed )
+	{
+		return;
+	}
+
+	FFormatNamedArguments Arguments;
+	bool bLowFramerate = false;
+	int32 SampleTime = 0;
+	float PercentUnderTarget = 0.f;
+
+	if (FineMovingAverage.IsReliable())
+	{
+		static IConsoleVariable* CVarMinFPS = IConsoleManager::Get().FindConsoleVariable(TEXT("PerfWarn.FineMinFPS"));
+		static IConsoleVariable* CVarPercentThreshold = IConsoleManager::Get().FindConsoleVariable(TEXT("PerfWarn.FinePercentThreshold"));
+
+		PercentUnderTarget = FineMovingAverage.PercentageBelowThreshold(CVarMinFPS->GetFloat());
+
+		if (PercentUnderTarget >= CVarPercentThreshold->GetFloat())
+		{
+			Arguments.Add(TEXT("Framerate"), CVarMinFPS->GetInt());
+			Arguments.Add(TEXT("Percentage"), FMath::FloorToFloat(PercentUnderTarget));
+			SampleTime = IConsoleManager::Get().FindConsoleVariable(TEXT("PerfWarn.FineSampleTime"))->GetInt();
+
+			bLowFramerate = true;
+		}
+	}
+
+
+	if (!bLowFramerate && CoarseMovingAverage.IsReliable())
+	{
+		static IConsoleVariable* CVarMinFPS = IConsoleManager::Get().FindConsoleVariable(TEXT("PerfWarn.CoarseMinFPS"));
+		static IConsoleVariable* CVarPercentThreshold = IConsoleManager::Get().FindConsoleVariable(TEXT("PerfWarn.CoarsePercentThreshold"));
+
+		PercentUnderTarget = CoarseMovingAverage.PercentageBelowThreshold(CVarMinFPS->GetFloat());
+
+		if (PercentUnderTarget >= CVarPercentThreshold->GetFloat())
+		{
+			Arguments.Add(TEXT("Framerate"), CVarMinFPS->GetInt());
+			Arguments.Add(TEXT("Percentage"), FMath::FloorToFloat(PercentUnderTarget));
+			SampleTime = IConsoleManager::Get().FindConsoleVariable(TEXT("PerfWarn.CoarseSampleTime"))->GetInt();
+
+			bLowFramerate = true;
+		}
+	}
+
+	if (!bLowFramerate)
+	{
+		HidePerformanceWarning();
+	}
+	else
+	{
+		enum MessagesEnum { Seconds, SecondsPercent, Minute, MinutePercent, Minutes, MinutesPercent };
+		const FText Messages[] = {
+			LOCTEXT("PerformanceWarningInProgress_Seconds",			"Your framerate has been under {Framerate} FPS for the past {SampleTime} seconds."),
+			LOCTEXT("PerformanceWarningInProgress_Seconds_Percent", "Your framerate has been under {Framerate} FPS for {Percentage}% of the past {SampleTime} seconds."),
+
+			LOCTEXT("PerformanceWarningInProgress_Minute",			"Your framerate has been under {Framerate} FPS for the past minute."),
+			LOCTEXT("PerformanceWarningInProgress_Minute_Percent",	"Your framerate has been under {Framerate} FPS for {Percentage}% of the last minute."),
+
+			LOCTEXT("PerformanceWarningInProgress_Minutes",			"Your framerate has been below {Framerate} FPS for the past {SampleTime} minutes."),
+			LOCTEXT("PerformanceWarningInProgress_Minutes_Percent", "Your framerate has been below {Framerate} FPS for {Percentage}% of the past {SampleTime} minutes."),
+		};
+
+		int32 Message;
+		if (SampleTime < 60)
+		{
+			Arguments.Add(TEXT("SampleTime"), SampleTime);
+			Message = Seconds;
+		}
+		else if (SampleTime == 60)
+		{
+			Message = Minute;
+		}
+		else
+		{
+			Arguments.Add(TEXT("SampleTime"), SampleTime / 60);
+			Message = Minutes;
+		}
+
+		// Use the message with the specific percentage on if applicable
+		if (PercentUnderTarget <= 95.f)
+		{
+			++Message;
+		}
+
+		ShowPerformanceWarning(FText::Format(Messages[Message], Arguments));
+	}
+}
+
+void FPerformanceWarningNotification::ShowScalabilityDialog()
+{
+	auto ExistingWindow = ScalabilitySettingsWindowPtr.Pin();
+
+	if (ExistingWindow.IsValid())
+	{
+		ExistingWindow->BringToFront();
+	}
+	else
+	{
+		// Create the window
+		ScalabilitySettingsWindowPtr = ExistingWindow = SNew(SWindow)
+			.Title(LOCTEXT("PerformanceWarningDialogTitle", "Scalability Options"))
+			.SupportsMaximize(false)
+			.SupportsMinimize(false)
+			.CreateTitleBar(true)
+			.SizingRule(ESizingRule::Autosized);
+
+		ExistingWindow->SetOnWindowClosed(FOnWindowClosed::CreateStatic([](const TSharedRef<SWindow>&, FPerformanceWarningNotification* PerfWarn){
+			PerfWarn->Reset();
+		}, this));
+
+		ExistingWindow->SetContent(
+			SNew(SScalabilitySettingsDialog)
+			.OnDoneClicked(
+				FOnClicked::CreateStatic([](TWeakPtr<SWindow> Window){
+					auto WindowPin = Window.Pin();
+					if (WindowPin.IsValid())
+					{
+						WindowPin->RequestDestroyWindow();
+					}
+					return FReply::Handled();
+				}, ScalabilitySettingsWindowPtr)
+			)
+		);
+
+
+		TSharedPtr<SWindow> RootWindow = FGlobalTabmanager::Get()->GetRootWindow();
+		if (RootWindow.IsValid())
+		{
+			FSlateApplication::Get().AddModalWindow(ExistingWindow.ToSharedRef(), RootWindow);
+		}
+		else
+		{
+			FSlateApplication::Get().AddWindow(ExistingWindow.ToSharedRef());
+		}
+	}
+}
+
+#undef LOCTEXT_NAMESPACE
\ No newline at end of file
diff --git a/Engine/Source/Editor/UnrealEd/Classes/Editor/EditorGameAgnosticSettings.h b/Engine/Source/Editor/UnrealEd/Classes/Editor/EditorGameAgnosticSettings.h
--- a/Engine/Source/Editor/UnrealEd/Classes/Editor/EditorGameAgnosticSettings.h
+++ b/Engine/Source/Editor/UnrealEd/Classes/Editor/EditorGameAgnosticSettings.h
@@ -28,6 +28,9 @@ class UEditorGameAgnosticSettings : public UObject
 	UPROPERTY(config)
 	bool bCopyStarterContentPreference;
 
+	UPROPERTY(config)
+	bool bShowPerformanceWarningPreference;
+
 	/** The id's of the surveys completed */
 	UPROPERTY(config)
 	TArray<FGuid> CompletedSurveys;
diff --git a/Engine/Source/Editor/UnrealEd/Private/Editor/EditorGameAgnosticSettings.cpp b/Engine/Source/Editor/UnrealEd/Private/Editor/EditorGameAgnosticSettings.cpp
--- a/Engine/Source/Editor/UnrealEd/Private/Editor/EditorGameAgnosticSettings.cpp
+++ b/Engine/Source/Editor/UnrealEd/Private/Editor/EditorGameAgnosticSettings.cpp
@@ -6,6 +6,7 @@ UEditorGameAgnosticSettings::UEditorGameAgnosticSettings(const class FPostConstr
 	: Super(PCIP)
 {
 	bCopyStarterContentPreference = true;
+	bShowPerformanceWarningPreference = true;
 }
 
 void UEditorGameAgnosticSettings::PostEditChangeProperty( struct FPropertyChangedEvent& PropertyChangedEvent)
