Editor: Fix typo of "Orthographic"
Integrate CL# 2124151 from 4.3 to Main

diff --git a/Engine/Source/Editor/LevelEditor/Private/SLevelViewportToolBar.cpp b/Engine/Source/Editor/LevelEditor/Private/SLevelViewportToolBar.cpp
--- a/Engine/Source/Editor/LevelEditor/Private/SLevelViewportToolBar.cpp
+++ b/Engine/Source/Editor/LevelEditor/Private/SLevelViewportToolBar.cpp
@@ -790,7 +790,7 @@ TSharedRef<SWidget> SLevelViewportToolBar::GenerateCameraMenu() const
 	// Camera types
 	CameraMenuBuilder.AddMenuEntry( FEditorViewportCommands::Get().Perspective );
 
-	CameraMenuBuilder.BeginSection("LevelViewportCameraType_Ortho", LOCTEXT("CameraTypeHeader_Ortho", "Othographic") );
+	CameraMenuBuilder.BeginSection("LevelViewportCameraType_Ortho", LOCTEXT("CameraTypeHeader_Ortho", "Orthographic") );
 		CameraMenuBuilder.AddMenuEntry(FEditorViewportCommands::Get().Top);
 		CameraMenuBuilder.AddMenuEntry(FEditorViewportCommands::Get().Side);
 		CameraMenuBuilder.AddMenuEntry(FEditorViewportCommands::Get().Front);
diff --git a/Engine/Source/Editor/Persona/Private/SAnimViewportToolBar.cpp b/Engine/Source/Editor/Persona/Private/SAnimViewportToolBar.cpp
--- a/Engine/Source/Editor/Persona/Private/SAnimViewportToolBar.cpp
+++ b/Engine/Source/Editor/Persona/Private/SAnimViewportToolBar.cpp
@@ -704,7 +704,7 @@ TSharedRef<SWidget> SAnimViewportToolBar::GenerateViewportTypeMenu() const
 	// Camera types
 	CameraMenuBuilder.AddMenuEntry(FEditorViewportCommands::Get().Perspective);
 
-	CameraMenuBuilder.BeginSection("LevelViewportCameraType_Ortho", LOCTEXT("CameraTypeHeader_Ortho", "Othographic"));
+	CameraMenuBuilder.BeginSection("LevelViewportCameraType_Ortho", LOCTEXT("CameraTypeHeader_Ortho", "Orthographic"));
 	CameraMenuBuilder.AddMenuEntry(FEditorViewportCommands::Get().Top);
 	CameraMenuBuilder.AddMenuEntry(FEditorViewportCommands::Get().Side);
 	CameraMenuBuilder.AddMenuEntry(FEditorViewportCommands::Get().Front);
