MeshProxy tool will be disabled instead of hidden when user has not simplygon

diff --git a/Engine/Source/Editor/LevelEditor/Private/LevelEditorActions.cpp b/Engine/Source/Editor/LevelEditor/Private/LevelEditorActions.cpp
--- a/Engine/Source/Editor/LevelEditor/Private/LevelEditorActions.cpp
+++ b/Engine/Source/Editor/LevelEditor/Private/LevelEditorActions.cpp
@@ -42,6 +42,7 @@
 #include "AnalyticsEventAttribute.h"
 #include "IAnalyticsProvider.h"
 #include "ReferenceViewer.h"
+#include "Developer/MeshUtilities/Public/MeshUtilities.h"
 
 #include "EditorActorFolders.h"
 #include "ActorPickerMode.h"
@@ -1509,8 +1510,15 @@ void FLevelEditorActionCallbacks::MergeActors_Clicked()
 
 bool FLevelEditorActionCallbacks::CanExecuteMergeActors()
 {
-	FSelectedActorInfo Info = AssetSelectionUtils::GetSelectedActorInfo();
-	return (Info.bHaveStaticMeshComponent || Info.bHaveLandscape);
+	IMeshUtilities* MeshUtilities = FModuleManager::Get().LoadModulePtr<IMeshUtilities>("MeshUtilities");
+	
+	if (MeshUtilities && MeshUtilities->GetMeshMergingInterface() != nullptr)
+	{
+		FSelectedActorInfo Info = AssetSelectionUtils::GetSelectedActorInfo();
+		return (Info.bHaveStaticMeshComponent || Info.bHaveLandscape);
+	}
+	
+	return false;
 }
 
 void FLevelEditorActionCallbacks::MergeActorsByMaterials_Clicked()
diff --git a/Engine/Source/Editor/LevelEditor/Private/LevelEditorContextMenu.cpp b/Engine/Source/Editor/LevelEditor/Private/LevelEditorContextMenu.cpp
--- a/Engine/Source/Editor/LevelEditor/Private/LevelEditorContextMenu.cpp
+++ b/Engine/Source/Editor/LevelEditor/Private/LevelEditorContextMenu.cpp
@@ -27,7 +27,6 @@
 #include "GlobalEditorCommonCommands.h"
 #include "LevelEditorCreateActorMenu.h"
 #include "SourceCodeNavigation.h"
-#include "Developer/MeshUtilities/Public/MeshUtilities.h"
 
 #define LOCTEXT_NAMESPACE "LevelViewportContextMenu"
 
@@ -929,17 +928,12 @@ void FLevelEditorContextMenuImpl::FillMergeActorsMenu( class FMenuBuilder& MenuB
 {
 	MenuBuilder.AddMenuEntry( FLevelEditorCommands::Get().MergeActorsByMaterials );
 
-	IMeshUtilities& MeshUtilities = FModuleManager::Get().LoadModuleChecked<IMeshUtilities>("MeshUtilities");
-	IMeshReduction* MeshReduction = MeshUtilities.GetMeshReductionInterface();
-
-	if (MeshReduction && MeshReduction->IsSupported())
+	// Simplygon ProxyLOD
+	MenuBuilder.BeginSection("ProxySimplygon", LOCTEXT("SimplygonHeading", "Simplygon"));
 	{
-		MenuBuilder.BeginSection("ProxySimplygon", LOCTEXT("SimplygonHeading", "Simplygon"));
-		{
-			MenuBuilder.AddMenuEntry(FLevelEditorCommands::Get().MergeActors);
-		}
-		MenuBuilder.EndSection();
+		MenuBuilder.AddMenuEntry(FLevelEditorCommands::Get().MergeActors);
 	}
+	MenuBuilder.EndSection();
 }
 
 void FLevelScriptEventMenuHelper::FillLevelBlueprintEventsMenu(class FMenuBuilder& MenuBuilder, const TArray<AActor*>& SelectedActors)
