Renamed "Levels Composition" tab into "World Composition"

diff --git a/Engine/Source/Editor/LevelEditor/Private/SLevelEditor.cpp b/Engine/Source/Editor/LevelEditor/Private/SLevelEditor.cpp
--- a/Engine/Source/Editor/LevelEditor/Private/SLevelEditor.cpp
+++ b/Engine/Source/Editor/LevelEditor/Private/SLevelEditor.cpp
@@ -725,7 +725,7 @@ TSharedRef<SDockTab> SLevelEditor::SpawnLevelEditorTab( const FSpawnTabArgs& Arg
 		FWorldBrowserModule& WorldBrowserModule = FModuleManager::LoadModuleChecked<FWorldBrowserModule>( "WorldBrowser" );
 		return SNew( SDockTab )
 			.Icon( FEditorStyle::GetBrush( "LevelEditor.Tabs.WorldBrowserComposition" ) )
-			.Label( NSLOCTEXT("LevelEditor", "WorldBrowserCompositionTabTitle", "Levels Composition") )
+			.Label( NSLOCTEXT("LevelEditor", "WorldBrowserCompositionTabTitle", "World Composition") )
 			[
 				WorldBrowserModule.CreateWorldBrowserComposition()
 			];
@@ -1045,7 +1045,7 @@ TSharedRef<SWidget> SLevelEditor::RestoreContentArea( const TSharedRef<SDockTab>
 		
 			LevelEditorTabManager->RegisterTabSpawner( WorldBrowserCompositionTab, FOnSpawnTab::CreateSP<SLevelEditor, FName, FString>(this, &SLevelEditor::SpawnLevelEditorTab, WorldBrowserCompositionTab, FString()) )
 				.SetMenuType( ETabSpawnerMenuType::Hide )
-				.SetDisplayName(NSLOCTEXT("LevelEditorTabs", "WorldBrowserComposition", "Levels Composition"))
+				.SetDisplayName(NSLOCTEXT("LevelEditorTabs", "WorldBrowserComposition", "World Composition"))
 				.SetGroup( WorkspaceMenu::GetMenuStructure().GetLevelEditorCategory() )
 				.SetIcon( FSlateIcon(FEditorStyle::GetStyleSetName(), "LevelEditor.Tabs.WorldBrowserComposition") );
 		}
