Hide SM3 from the feature level preview menu

diff --git a/Engine/Source/Editor/LevelEditor/Private/LevelEditorToolBar.cpp b/Engine/Source/Editor/LevelEditor/Private/LevelEditorToolBar.cpp
--- a/Engine/Source/Editor/LevelEditor/Private/LevelEditorToolBar.cpp
+++ b/Engine/Source/Editor/LevelEditor/Private/LevelEditorToolBar.cpp
@@ -475,7 +475,10 @@ static void MakeShaderModelPreviewMenu(FMenuBuilder& MenuBuilder)
 	{
 		for (int32 i = GMaxRHIFeatureLevel; i >= 0; --i)
 		{
-			MenuBuilder.AddMenuEntry(FLevelEditorCommands::Get().FeatureLevelPreview[i]);
+			if (i != ERHIFeatureLevel::SM3)
+			{
+				MenuBuilder.AddMenuEntry(FLevelEditorCommands::Get().FeatureLevelPreview[i]);
+			}
 		}
 	}
 	MenuBuilder.EndSection();
