Surface Properties - Display the surface count and material name.

TTP# 340343

diff --git a/Engine/Source/Editor/LevelEditor/Private/SSurfaceProperties.cpp b/Engine/Source/Editor/LevelEditor/Private/SSurfaceProperties.cpp
--- a/Engine/Source/Editor/LevelEditor/Private/SSurfaceProperties.cpp
+++ b/Engine/Source/Editor/LevelEditor/Private/SSurfaceProperties.cpp
@@ -5,6 +5,7 @@
 #include "Slate.h"
 #include "Editor/PropertyEditor/Public/PropertyEditorModule.h"
 #include "Editor/PropertyEditor/Public/IDetailsView.h"
+#include "SurfaceIterators.h"
 
 #define LOCTEXT_NAMESPACE "SSurfaceProperties"
 
@@ -36,6 +37,21 @@ void SSurfaceProperties::Construct( const FArguments& InArgs )
 	ChildSlot
 	[
 		SNew(SVerticalBox)
+
+		+ SVerticalBox::Slot()
+		.AutoHeight()
+		.Padding(0, 0, 0, 5)
+		[
+			SNew(SBorder)
+			.BorderBackgroundColor(BorderColor)
+			.BorderImage(BorderStyle)
+			.Padding(10)
+			[
+				SNew(STextBlock)
+				.Text(this, &SSurfaceProperties::GetSelectedSurfacesDescription)
+			]
+		]
+
 		+SVerticalBox::Slot()
 		.AutoHeight()
 		.Padding(0,0,0,5)
@@ -968,4 +984,56 @@ void SSurfaceProperties::OnToggleRotationDirection( ESlateCheckBoxState::Type Ne
 	bUseNegativeRotation = (NewState == ESlateCheckBoxState::Checked) ? true : false;
 }
 
+FText SSurfaceProperties::GetSelectedSurfacesDescription() const
+{
+	int SurfaceCount = 0;
+	UMaterialInterface* SurfaceMaterial = nullptr;
+	bool bMultipleMaterials = false;
+
+	// This would be easier to follow using a set to keep track of used materials
+	// but since it's called every frame a single pointer and boolean are used instead.
+	for (TSelectedSurfaceIterator<> SurfaceIt(GetWorld()); SurfaceIt; ++SurfaceIt)
+	{
+		FBspSurf* SelectedSurface = *SurfaceIt;
+		++SurfaceCount;
+		if (!bMultipleMaterials)
+		{
+			if (SurfaceMaterial != SelectedSurface->Material)
+			{
+				if (SurfaceCount == 1)
+				{
+					SurfaceMaterial = SelectedSurface->Material;
+				}
+				else
+				{
+					SurfaceMaterial = nullptr;
+					bMultipleMaterials = true;
+				}
+			}
+		}
+	}
+
+	FText MaterialText;
+	if (bMultipleMaterials)
+	{
+		MaterialText = LOCTEXT("MultipleMaterials", "Multiple materials");
+	}
+	else
+	{
+		if (SurfaceMaterial == nullptr)
+		{
+			MaterialText = LOCTEXT("DefaultMaterial", "Default material");
+		}
+		else
+		{
+			MaterialText = FText::FromName(SurfaceMaterial->GetFName());
+		}
+	}
+
+	return FText::Format(LOCTEXT("SurfaceDescriptionFormat", "{0} {1}: {2}"),
+		FText::AsNumber(SurfaceCount),
+		SurfaceCount == 1 ? LOCTEXT("SurfaceSingular", "Surface") : LOCTEXT("SurfacePlural", "Surfaces"),
+		MaterialText);
+}
+
 #undef LOCTEXT_NAMESPACE
diff --git a/Engine/Source/Editor/LevelEditor/Private/SSurfaceProperties.h b/Engine/Source/Editor/LevelEditor/Private/SSurfaceProperties.h
--- a/Engine/Source/Editor/LevelEditor/Private/SSurfaceProperties.h
+++ b/Engine/Source/Editor/LevelEditor/Private/SSurfaceProperties.h
@@ -235,6 +235,10 @@ private:
 	 */
 	void OnToggleRotationDirection( ESlateCheckBoxState::Type NewState );
 
+	/**
+	 * Returns a short summary to display to the user about the currently selected surfaces.
+	**/
+	FText GetSelectedSurfacesDescription() const;
 
 
 	/** typedef to make it easier to handle the lightmass objects */
