* cleaner UI - what is a showflag, what is not

diff --git a/Engine/Source/Editor/LevelEditor/Private/SLevelViewportToolBar.cpp b/Engine/Source/Editor/LevelEditor/Private/SLevelViewportToolBar.cpp
--- a/Engine/Source/Editor/LevelEditor/Private/SLevelViewportToolBar.cpp
+++ b/Engine/Source/Editor/LevelEditor/Private/SLevelViewportToolBar.cpp
@@ -941,7 +941,7 @@ TSharedRef<SWidget> SLevelViewportToolBar::GenerateShowMenu() const
 		if( ShowMenu[SFG_Normal].Num() > 0 )
 		{
 			// Generate entries for the standard show flags
-			ShowMenuBuilder.BeginSection("LevelViewportShowFlagsCommon", LOCTEXT("CommonShowFlagHeader", "Common") );
+			ShowMenuBuilder.BeginSection("LevelViewportShowFlagsCommon", LOCTEXT("CommonShowFlagHeader", "Common Show flags") );
 			{
 				for( int32 EntryIndex = 0; EntryIndex < ShowMenu[SFG_Normal].Num(); ++EntryIndex )
 				{
@@ -952,7 +952,7 @@ TSharedRef<SWidget> SLevelViewportToolBar::GenerateShowMenu() const
 		}
 
 		// Generate entries for the different show flags groups
-		ShowMenuBuilder.BeginSection("LevelViewportShowFlags");
+		ShowMenuBuilder.BeginSection("LevelViewportShowFlags", LOCTEXT("AllShowFlagHeader", "All Show flags"));
 		{
 			ShowMenuBuilder.AddSubMenu( LOCTEXT("PostProcessShowFlagsMenu", "Post Processing"), LOCTEXT("PostProcessShowFlagsMenu_ToolTip", "Post process show flags"),
 				FNewMenuDelegate::CreateStatic( &FillShowMenu, ShowMenu[SFG_PostProcess], 0 ) );
@@ -974,45 +974,49 @@ TSharedRef<SWidget> SLevelViewportToolBar::GenerateShowMenu() const
 		}
 		ShowMenuBuilder.EndSection();
 
-
 		FText ShowAllLabel = LOCTEXT("ShowAllLabel", "Show All");
 		FText HideAllLabel = LOCTEXT("HideAllLabel", "Hide All");
 
-		// Show Volumes sub-menu
+		ShowMenuBuilder.BeginSection("LevelViewportEditorShow", LOCTEXT("EditorShowHeader", "Editor"));
 		{
-			TArray< FLevelViewportCommands::FShowMenuCommand > ShowVolumesMenu;
+			// Show Volumes sub-menu
+			{
+				TArray< FLevelViewportCommands::FShowMenuCommand > ShowVolumesMenu;
 
-			// 'Show All' and 'Hide All' buttons
-			ShowVolumesMenu.Add( FLevelViewportCommands::FShowMenuCommand( Actions.ShowAllVolumes, ShowAllLabel ) );
-			ShowVolumesMenu.Add( FLevelViewportCommands::FShowMenuCommand( Actions.HideAllVolumes, HideAllLabel ) );
+				// 'Show All' and 'Hide All' buttons
+				ShowVolumesMenu.Add(FLevelViewportCommands::FShowMenuCommand(Actions.ShowAllVolumes, ShowAllLabel));
+				ShowVolumesMenu.Add(FLevelViewportCommands::FShowMenuCommand(Actions.HideAllVolumes, HideAllLabel));
 
-			// Get each show flag command and put them in their corresponding groups
-			ShowVolumesMenu += Actions.ShowVolumeCommands;
+				// Get each show flag command and put them in their corresponding groups
+				ShowVolumesMenu += Actions.ShowVolumeCommands;
 
-			ShowMenuBuilder.AddSubMenu( LOCTEXT("ShowVolumesMenu", "Volumes"), LOCTEXT("ShowVolumesMenu_ToolTip", "Show volumes flags"),
-				FNewMenuDelegate::CreateStatic( &FillShowMenu, ShowVolumesMenu, 2 ) );
-		}
+				ShowMenuBuilder.AddSubMenu(LOCTEXT("ShowVolumesMenu", "Volumes"), LOCTEXT("ShowVolumesMenu_ToolTip", "Show volumes flags"),
+					FNewMenuDelegate::CreateStatic(&FillShowMenu, ShowVolumesMenu, 2));
+			}
 
-		// Show Layers sub-menu is dynamically generated when the user enters 'show' menu
-		{
-			ShowMenuBuilder.AddSubMenu( LOCTEXT("ShowLayersMenu", "Layers"), LOCTEXT("ShowLayersMenu_ToolTip", "Show layers flags"),
-				FNewMenuDelegate::CreateStatic( &SLevelViewportToolBar::FillShowLayersMenu, Viewport ) );
+			// Show Layers sub-menu is dynamically generated when the user enters 'show' menu
+			{
+			ShowMenuBuilder.AddSubMenu(LOCTEXT("ShowLayersMenu", "Layers"), LOCTEXT("ShowLayersMenu_ToolTip", "Show layers flags"),
+				FNewMenuDelegate::CreateStatic(&SLevelViewportToolBar::FillShowLayersMenu, Viewport));
 		}
 
-		// Show Sprites sub-menu
-		{
-			TArray< FLevelViewportCommands::FShowMenuCommand > ShowSpritesMenu;
+			// Show Sprites sub-menu
+			{
+				TArray< FLevelViewportCommands::FShowMenuCommand > ShowSpritesMenu;
 
-			// 'Show All' and 'Hide All' buttons
-			ShowSpritesMenu.Add( FLevelViewportCommands::FShowMenuCommand( Actions.ShowAllSprites, ShowAllLabel ) );
-			ShowSpritesMenu.Add( FLevelViewportCommands::FShowMenuCommand( Actions.HideAllSprites, HideAllLabel ) );
+				// 'Show All' and 'Hide All' buttons
+				ShowSpritesMenu.Add(FLevelViewportCommands::FShowMenuCommand(Actions.ShowAllSprites, ShowAllLabel));
+				ShowSpritesMenu.Add(FLevelViewportCommands::FShowMenuCommand(Actions.HideAllSprites, HideAllLabel));
 
-			// Get each show flag command and put them in their corresponding groups
-			ShowSpritesMenu += Actions.ShowSpriteCommands;
+				// Get each show flag command and put them in their corresponding groups
+				ShowSpritesMenu += Actions.ShowSpriteCommands;
 
-			ShowMenuBuilder.AddSubMenu( LOCTEXT("ShowSpritesMenu", "Sprites"), LOCTEXT("ShowSpritesMenu_ToolTip", "Show sprites flags"),
-				FNewMenuDelegate::CreateStatic( &FillShowMenu, ShowSpritesMenu, 2 ) );
+				ShowMenuBuilder.AddSubMenu(LOCTEXT("ShowSpritesMenu", "Sprites"), LOCTEXT("ShowSpritesMenu_ToolTip", "Show sprites flags"),
+					FNewMenuDelegate::CreateStatic(&FillShowMenu, ShowSpritesMenu, 2));
+			}
 		}
+
+		ShowMenuBuilder.EndSection();
 	}
 
 	return ShowMenuBuilder.MakeWidget();
