TTP 341939 - Fix for crash when taking a high res screenshot from a viewport-locked camera, with custom depth buffer masking enabled.

Several other fixes for high res screenshot issues.

Default behaviour of screenshots with viewport-locked cameras is to only capture the actual viewrect, not including the black bars. Removed button for setting capture region to view rect.

diff --git a/Engine/Source/Editor/LevelEditor/Private/HighResScreenshotUI.cpp b/Engine/Source/Editor/LevelEditor/Private/HighResScreenshotUI.cpp
--- a/Engine/Source/Editor/LevelEditor/Private/HighResScreenshotUI.cpp
+++ b/Engine/Source/Editor/LevelEditor/Private/HighResScreenshotUI.cpp
@@ -184,20 +184,6 @@ void SHighResScreenshotDialog::Construct( const FArguments& InArgs )
 							.Image(FEditorStyle::GetBrush("HighresScreenshot.FullViewportCaptureRegion"))
 						]
 					]
-					+ SHorizontalBox::Slot()
-					.HAlign(HAlign_Right)
-					.AutoWidth()
-					[
-						SNew(SButton)
-						.ToolTipText(NSLOCTEXT("HighResScreenshot", "ScreenshotCameraSafeAreaCaptureRegionTooltip", "Set the capture rectangle to the attached camera safe area (only enabled when the viewport is locked to a camera actor)"))
-						.Visibility(this, &SHighResScreenshotDialog::GetCaptureRegionControlsVisibility)
-						.OnClicked(this, &SHighResScreenshotDialog::OnSetCameraSafeAreaCaptureRegionClicked)
-						.IsEnabled(this, &SHighResScreenshotDialog::IsSetCameraSafeAreaCaptureRegionEnabled)
-						[
-							SNew(SImage)
-							.Image(FEditorStyle::GetBrush("HighresScreenshot.CameraSafeAreaCaptureRegion"))
-						]
-					]
 					+SHorizontalBox::Slot()
 						// for padding
 					+SHorizontalBox::Slot()
@@ -330,7 +316,7 @@ FReply SHighResScreenshotDialog::OnSelectCaptureRegionClicked()
 	// Only enable the capture region widget if the owning viewport gave us one
 	if (CaptureRegionWidget.IsValid())
 	{
-		CaptureRegionWidget->Activate(Config.UnscaledCaptureRegion.Width() == -1 || Config.UnscaledCaptureRegion.Height() == -1);
+		CaptureRegionWidget->Activate(Config.UnscaledCaptureRegion.Area() > 0);
 		bCaptureRegionControlsVisible = true;
 	}
 	return FReply::Handled();
@@ -347,13 +333,12 @@ FReply SHighResScreenshotDialog::OnCaptureClicked()
 			GScreenshotResolutionY = ConfigViewport->GetSizeXY().Y * Config.ResolutionMultiplier;
 			FIntRect ScaledCaptureRegion = Config.UnscaledCaptureRegion;
 
-		if (ScaledCaptureRegion.Width() == -1 || ScaledCaptureRegion.Height() == -1)
-		{
-			ScaledCaptureRegion = FIntRect(0, 0, ConfigViewport->GetSizeXY().X, ConfigViewport->GetSizeXY().Y);
-		}
+			if (ScaledCaptureRegion.Area() > 0)
+			{
+				ScaledCaptureRegion.Clip(FIntRect(FIntPoint::ZeroValue, ConfigViewport->GetSizeXY()));
+				ScaledCaptureRegion *= Config.ResolutionMultiplier;
+			}
 
-			ScaledCaptureRegion.Clip(FIntRect(FIntPoint::ZeroValue, ConfigViewport->GetSizeXY()));
-			ScaledCaptureRegion *= Config.ResolutionMultiplier;
 			Config.CaptureRegion = ScaledCaptureRegion;
 
 			// Trigger the screenshot on the owning viewport
@@ -402,7 +387,7 @@ FReply SHighResScreenshotDialog::OnSetFullViewportCaptureRegionClicked()
 		ConfigViewport->Invalidate();
 	}
 
-	Config.UnscaledCaptureRegion = FIntRect(0, 0, -1, -1);
+	Config.UnscaledCaptureRegion = FIntRect(0, 0, 0, 0);
 	CaptureRegionWidget->Reset();
 	return FReply::Handled();
 }
diff --git a/Engine/Source/Editor/UnrealEd/Private/EditorViewportClient.cpp b/Engine/Source/Editor/UnrealEd/Private/EditorViewportClient.cpp
--- a/Engine/Source/Editor/UnrealEd/Private/EditorViewportClient.cpp
+++ b/Engine/Source/Editor/UnrealEd/Private/EditorViewportClient.cpp
@@ -3844,12 +3844,44 @@ void FEditorViewportClient::ProcessScreenShots(FViewport* InViewport)
 {
 	if (GIsDumpingMovie || FScreenshotRequest::IsScreenshotRequested() || GIsHighResScreenshot)
 	{
+		// Default capture region is the entire viewport
+		FIntRect CaptureRect(0, 0, 0, 0);
+
+		bool bCaptureAreaValid = GetHighResScreenshotConfig().CaptureRegion.Area() > 0;
+
+		// If capture region isn't valid, we need to determine which rectangle to capture from.
+		// We need to calculate a proper view rectangle so that we can take into account camera
+		// properties, such as it being aspect ratio constrainted
+		if (GIsHighResScreenshot && !bCaptureAreaValid)
+		{
+			FSceneViewFamilyContext ViewFamily(FSceneViewFamily::ConstructionValues(
+				InViewport,
+				GetScene(),
+				EngineShowFlags)
+				.SetRealtimeUpdate(IsRealtime()));
+			auto* ViewportBak = Viewport;
+			Viewport = InViewport;
+			FSceneView* View = CalcSceneView(&ViewFamily);
+			Viewport = ViewportBak;
+			CaptureRect = View->ViewRect;
+		}
+
 		FString ScreenShotName = FScreenshotRequest::GetFilename();
 		TArray<FColor> Bitmap;
-		if( GetViewportScreenShot(InViewport, Bitmap) )
+		if (GetViewportScreenShot(InViewport, Bitmap, CaptureRect))
 		{
-			FIntPoint BitmapSize = InViewport->GetSizeXY();
-			FIntRect SourceRect = GIsHighResScreenshot ? GetHighResScreenshotConfig().CaptureRegion : FIntRect(0, 0, GScreenshotResolutionX, GScreenshotResolutionY);
+			// Determine the size of the captured viewport data.
+			FIntPoint BitmapSize = CaptureRect.Area() > 0 ? CaptureRect.Size() : InViewport->GetSizeXY();
+			
+			// Determine which region of the captured data we want to save out. If the highres screenshot capture region 
+			// is not valid, we want to save out everything in the viewrect that we just grabbed.
+			FIntRect SourceRect = FIntRect(0, 0, 0, 0);
+			if (GIsHighResScreenshot && bCaptureAreaValid)
+			{
+				// Highres screenshot capture region is valid, so use that
+				SourceRect = GetHighResScreenshotConfig().CaptureRegion;
+			}
+
 			bool bWriteAlpha = false;
 
 			// If this is a high resolution screenshot and we are using the masking feature,
@@ -3859,6 +3891,7 @@ void FEditorViewportClient::ProcessScreenShots(FViewport* InViewport)
 				bWriteAlpha = GetHighResScreenshotConfig().MergeMaskIntoAlpha(Bitmap);
 			}
 
+			// Save the bitmap to disc
 			FFileHelper::CreateBitmap(*ScreenShotName, BitmapSize.X, BitmapSize.Y, Bitmap.GetTypedData(), &SourceRect, &IFileManager::Get(), NULL, bWriteAlpha);
 		}
 		
diff --git a/Engine/Source/Runtime/Engine/Private/HighResScreenshot.cpp b/Engine/Source/Runtime/Engine/Private/HighResScreenshot.cpp
--- a/Engine/Source/Runtime/Engine/Private/HighResScreenshot.cpp
+++ b/Engine/Source/Runtime/Engine/Private/HighResScreenshot.cpp
@@ -29,7 +29,7 @@ void FHighResScreenshotConfig::ChangeViewport(TWeakPtr<FSceneViewport> InViewpor
 		Viewport->Invalidate();
 	}
 
-	UnscaledCaptureRegion = FIntRect(0, 0, -1, -1);
+	UnscaledCaptureRegion = FIntRect(0, 0, 0, 0);
 	CaptureRegion = UnscaledCaptureRegion;
 	bMaskEnabled = false;
 	bDumpBufferVisualizationTargets = false;
diff --git a/Engine/Source/Runtime/Engine/Private/SceneView.cpp b/Engine/Source/Runtime/Engine/Private/SceneView.cpp
--- a/Engine/Source/Runtime/Engine/Private/SceneView.cpp
+++ b/Engine/Source/Runtime/Engine/Private/SceneView.cpp
@@ -1046,7 +1046,7 @@ void FSceneView::EndFinalPostprocessSettings()
 	if (Config.bDisplayCaptureRegion && !GIsHighResScreenshot)
 	{
 		// Only enable the capture region effect if the capture region is different from the view rectangle...
-		if (Config.UnscaledCaptureRegion != ViewRect && Config.UnscaledCaptureRegion.Width() != -1 && Config.UnscaledCaptureRegion.Height() != -1 && State != NULL)
+		if ((Config.UnscaledCaptureRegion != ViewRect) && (Config.UnscaledCaptureRegion.Area() > 0) && (State != NULL))
 		{
 			// ...and if this is the viewport associated with the highres screenshot UI
 			auto ConfigViewport = Config.TargetViewport.Pin();
diff --git a/Engine/Source/Runtime/Engine/Private/UnrealClient.cpp b/Engine/Source/Runtime/Engine/Private/UnrealClient.cpp
--- a/Engine/Source/Runtime/Engine/Private/UnrealClient.cpp
+++ b/Engine/Source/Runtime/Engine/Private/UnrealClient.cpp
@@ -1583,12 +1583,12 @@ void FViewport::SetInitialSize( FIntPoint InitialSizeXY )
 
 
 
-ENGINE_API bool GetViewportScreenShot(FViewport* Viewport, TArray<FColor>& Bitmap)
+ENGINE_API bool GetViewportScreenShot(FViewport* Viewport, TArray<FColor>& Bitmap, const FIntRect& ViewRect /*= FIntRect()*/)
 {
 	// Read the contents of the viewport into an array.
-	if(Viewport->ReadPixels(Bitmap))
+	if (Viewport->ReadPixels(Bitmap, FReadSurfaceDataFlags(), ViewRect))
 	{
-		check(Bitmap.Num() == Viewport->GetSizeXY().X * Viewport->GetSizeXY().Y);
+		check(Bitmap.Num() == ViewRect.Area() || (Bitmap.Num() == Viewport->GetSizeXY().X * Viewport->GetSizeXY().Y));
 		return true;
 	}
 
diff --git a/Engine/Source/Runtime/Engine/Public/UnrealClient.h b/Engine/Source/Runtime/Engine/Public/UnrealClient.h
--- a/Engine/Source/Runtime/Engine/Public/UnrealClient.h
+++ b/Engine/Source/Runtime/Engine/Public/UnrealClient.h
@@ -602,7 +602,7 @@ extern ENGINE_API bool IsCtrlDown(FViewport* Viewport);
 extern ENGINE_API bool IsShiftDown(FViewport* Viewport);
 extern ENGINE_API bool IsAltDown(FViewport* Viewport);
 
-extern ENGINE_API bool GetViewportScreenShot(FViewport* Viewport, TArray<FColor>& Bitmap);
+extern ENGINE_API bool GetViewportScreenShot(FViewport* Viewport, TArray<FColor>& Bitmap, const FIntRect& ViewRect = FIntRect());
 extern ENGINE_API bool GetHighResScreenShotInput(const TCHAR* Cmd, FOutputDevice& Ar, uint32& OutXRes, uint32& OutYRes, float& OutResMult, FIntRect& OutCaptureRegion, bool& OutShouldEnableMask);
 
 /**
diff --git a/Engine/Source/Runtime/Renderer/Private/PostProcess/RenderingCompositionGraph.cpp b/Engine/Source/Runtime/Renderer/Private/PostProcess/RenderingCompositionGraph.cpp
--- a/Engine/Source/Runtime/Renderer/Private/PostProcess/RenderingCompositionGraph.cpp
+++ b/Engine/Source/Runtime/Renderer/Private/PostProcess/RenderingCompositionGraph.cpp
@@ -269,8 +269,16 @@ void FRenderingCompositionGraph::DumpOutputToFile(FRenderingCompositePassContext
 	if (GIsHighResScreenshot)
 	{
 		SourceRect = GetHighResScreenshotConfig().CaptureRegion;
-		SourceRect.Min.X *= CaptureRegionScaleX;
-		SourceRect.Max.X *= CaptureRegionScaleX;
+
+		if (SourceRect.Area() > 0)
+		{
+			SourceRect = Context.View.ViewRect;
+		}
+		else
+		{
+			SourceRect.Min.X *= CaptureRegionScaleX;
+			SourceRect.Max.X *= CaptureRegionScaleX;
+		}
 	}
 	
 	FFileHelper::CreateBitmap(File, ExtendXWithMSAA, Extent.Y, Bitmap.GetTypedData(), GIsHighResScreenshot ? &SourceRect : NULL);	
@@ -554,7 +562,7 @@ void FRenderingCompositionGraph::RecursivelyProcess(const FRenderingCompositeOut
 	}
 #endif
 
-	// iterate through all inputs of this pass and decrement the references for it's inputs
+	// iterate through all inputs of tghis pass and decrement the references for it's inputs
 	// this can release some intermediate RT so they can be reused
 	{
 		uint32 InputId = 0;
