Fix typos

diff --git a/Engine/Source/Editor/LevelEditor/Private/LevelEditorActions.cpp b/Engine/Source/Editor/LevelEditor/Private/LevelEditorActions.cpp
--- a/Engine/Source/Editor/LevelEditor/Private/LevelEditorActions.cpp
+++ b/Engine/Source/Editor/LevelEditor/Private/LevelEditorActions.cpp
@@ -2763,7 +2763,7 @@ void FLevelEditorCommands::RegisterCommands()
 	UI_COMMAND( SurfSelectOr, "Surface Selection OR", "Replace the current selection with only the surfaces which are both currently selected and contained within the saved selection in memory", EUserInterfaceActionType::Button, FInputGesture(EModifierKey::Shift, EKeys::O) ); 
 	UI_COMMAND( SurfSelectAnd, "Surface Selection AND", "Add the selection of surfaces saved in memory to the current selection", EUserInterfaceActionType::Button, FInputGesture(EModifierKey::Shift, EKeys::A) ); 
 	UI_COMMAND( SurfSelectXor, "Surace Selection XOR", " Replace the current selection with only the surfaces that are not in both the current selection and the selection saved in memory", EUserInterfaceActionType::Button, FInputGesture(EModifierKey::Shift, EKeys::X) ); 
-	UI_COMMAND( SurfUnalign, "Align Surface Default", "Default surface alignmment", EUserInterfaceActionType::Button, FInputGesture() );
+	UI_COMMAND( SurfUnalign, "Align Surface Default", "Default surface alignment", EUserInterfaceActionType::Button, FInputGesture() );
 	UI_COMMAND( SurfAlignPlanarAuto, "Align Surface Planar", "Planar surface alignment", EUserInterfaceActionType::Button, FInputGesture() );
 	UI_COMMAND( SurfAlignPlanarWall, "Align Surface Planar Wall", "Planar wall surface alignment", EUserInterfaceActionType::Button, FInputGesture() );
 	UI_COMMAND( SurfAlignPlanarFloor, "Align Surface Planar Floor", "Planar floor surface alignment", EUserInterfaceActionType::Button, FInputGesture() );
diff --git a/Engine/Source/Editor/UnrealEd/Classes/Settings/LevelEditorViewportSettings.h b/Engine/Source/Editor/UnrealEd/Classes/Settings/LevelEditorViewportSettings.h
--- a/Engine/Source/Editor/UnrealEd/Classes/Settings/LevelEditorViewportSettings.h
+++ b/Engine/Source/Editor/UnrealEd/Classes/Settings/LevelEditorViewportSettings.h
@@ -262,7 +262,7 @@ class UNREALED_API ULevelEditorViewportSettings
 public:
 
 	/** If enabled will use power of 2 grid settings (e.g, 1,2,4,8,16,...,1024) instead of decimal grid sizes */
-	UPROPERTY(EditAnywhere, config, Category=GridSnapping, meta=(DisplayName = "User Power of Two Snap Size"))
+	UPROPERTY(EditAnywhere, config, Category=GridSnapping, meta=(DisplayName = "Use Power of Two Snap Size"))
 	bool bUsePowerOf2SnapSize;
 
 	/** Decimal grid sizes (for translation snapping and grid rendering) */
