Disabled code compilation in editor on promoted builds as they don't have any already-generated code and this causes compilation errors.

#ttp 342056: CRITICAL: Regression: UE4: EDITOR: Game code fails to recompile in edtior due to header changes

diff --git a/Engine/Source/Editor/LevelEditor/Private/LevelEditorActions.cpp b/Engine/Source/Editor/LevelEditor/Private/LevelEditorActions.cpp
--- a/Engine/Source/Editor/LevelEditor/Private/LevelEditorActions.cpp
+++ b/Engine/Source/Editor/LevelEditor/Private/LevelEditorActions.cpp
@@ -47,6 +47,7 @@
 
 #include "EditorActorFolders.h"
 #include "ActorPickerMode.h"
+#include "EngineBuildSettings.h"
 
 DEFINE_LOG_CATEGORY_STATIC(LevelEditorActions, Log, All);
 
@@ -953,7 +954,7 @@ void FLevelEditorActionCallbacks::RecompileGameCode_Clicked()
 bool FLevelEditorActionCallbacks::Recompile_CanExecute()
 {
 	// We're not able to recompile if a compile is already in progress!
-	return !FModuleManager::Get().IsCurrentlyCompiling();
+	return !FModuleManager::Get().IsCurrentlyCompiling() && !(GEngineVersion.IsPromotedBuild() && FEngineBuildSettings::IsPerforceBuild());
 }
 
 bool FLevelEditorActionCallbacks::Reload_CanExecute()
