Disable news feed

diff --git a/Engine/Source/Editor/LevelEditor/Private/SLevelEditor.cpp b/Engine/Source/Editor/LevelEditor/Private/SLevelEditor.cpp
--- a/Engine/Source/Editor/LevelEditor/Private/SLevelEditor.cpp
+++ b/Engine/Source/Editor/LevelEditor/Private/SLevelEditor.cpp
@@ -240,6 +240,9 @@ void SLevelEditor::ConstructNotificationBar()
 			FLevelEditorMenu::MakeNotificationBar( LevelEditorCommands, SharedThis(this ) )
 		];
 
+#define SHOW_NEWS_FEED 0
+
+#if SHOW_NEWS_FEED
 	// news feed button
 	INewsFeedModule& NewsFeedModule = FModuleManager::LoadModuleChecked<INewsFeedModule>(NewsFeedModuleName);
 
@@ -251,6 +254,7 @@ void SLevelEditor::ConstructNotificationBar()
 			NewsFeedModule.CreateNewsFeedButton()
 		];
 
+#endif
 	// developer tools
 	const IMainFrameModule& MainFrameModule = FModuleManager::GetModuleChecked<IMainFrameModule>(MainFrameModuleName);
 
