Quick access to project settings from the menu bar

diff --git a/Engine/Source/Editor/EditorStyle/Private/SlateEditorStyle.cpp b/Engine/Source/Editor/EditorStyle/Private/SlateEditorStyle.cpp
--- a/Engine/Source/Editor/EditorStyle/Private/SlateEditorStyle.cpp
+++ b/Engine/Source/Editor/EditorStyle/Private/SlateEditorStyle.cpp
@@ -3626,6 +3626,9 @@ void FSlateEditorStyle::FStyle::SetupLevelEditorStyle()
 		Set( "LevelEditor.ViewOptions", new IMAGE_BRUSH( "Icons/icon_view_40x", Icon40x40 ) );
 		Set( "LevelEditor.ViewOptions.Small", new IMAGE_BRUSH( "Icons/icon_view_40x", Icon20x20 ) );
 
+		Set( "LevelEditor.GameSettings", new IMAGE_BRUSH( "Icons/icon_game_settings_40x", Icon40x40 ) );
+		Set( "LevelEditor.GameSettings.Small", new IMAGE_BRUSH( "Icons/icon_game_settings_40x", Icon20x20 ) );
+
 		Set( "LevelEditor.Create", new IMAGE_BRUSH( "Icons/icon_Mode_Placement_40px", Icon40x40 ) );
 		Set( "LevelEditor.Create.Small", new IMAGE_BRUSH( "Icons/icon_Mode_Placement_40px", Icon20x20 ) );
 		Set( "LevelEditor.Create.OutlineHoriz", new IMAGE_BRUSH( "Common/WorkingFrame_Marquee", FVector2D(34.0f, 3.0f), FLinearColor::White, ESlateBrushTileType::Horizontal) );
@@ -3660,8 +3663,8 @@ void FSlateEditorStyle::FStyle::SetupLevelEditorStyle()
 		Set( "LevelEditor.BspMode.Selected", new IMAGE_BRUSH( "Icons/icon_Mode_GeoEdit-a_40px", Icon40x40 ) );
 		Set( "LevelEditor.BspMode.Selected.Small", new IMAGE_BRUSH( "Icons/icon_Mode_GeoEdit-a_40px", Icon20x20 ) );
 
-		Set( "LevelEditor.WorldProperties", new IMAGE_BRUSH( "Icons/icon_worldscript_40x", Icon40x40 ) );
-		Set( "LevelEditor.WorldProperties.Small", new IMAGE_BRUSH( "Icons/icon_worldscript_40x", Icon20x20 ) );
+		//Set( "LevelEditor.WorldProperties", new IMAGE_BRUSH( "Icons/icon_worldscript_40x", Icon40x40 ) );
+		//Set( "LevelEditor.WorldProperties.Small", new IMAGE_BRUSH( "Icons/icon_worldscript_40x", Icon20x20 ) );
 		Set( "LevelEditor.WorldProperties.Tab", new IMAGE_BRUSH( "Icons/icon_worldscript_40x", Icon16x16 ) );
 		Set( "LevelEditor.OpenContentBrowser", new IMAGE_BRUSH( "Icons/icon_ContentBrowser_40x", Icon40x40 ) );
 		Set( "LevelEditor.OpenContentBrowser.Small", new IMAGE_BRUSH( "Icons/icon_ContentBrowser_40x", Icon20x20 ) );
diff --git a/Engine/Source/Editor/LevelEditor/Private/LevelEditorActions.cpp b/Engine/Source/Editor/LevelEditor/Private/LevelEditorActions.cpp
--- a/Engine/Source/Editor/LevelEditor/Private/LevelEditorActions.cpp
+++ b/Engine/Source/Editor/LevelEditor/Private/LevelEditorActions.cpp
@@ -2748,7 +2748,7 @@ void FLevelEditorCommands::RegisterCommands()
 
 	UI_COMMAND( FindActorInLevelScript, "Find in Level Blueprint", "Finds any references to the selected actor in its level's blueprint", EUserInterfaceActionType::Button, FInputGesture() );
 
-	UI_COMMAND( WorldProperties, "World Properties", "Displays the world properties", EUserInterfaceActionType::Button, FInputGesture() );
+	UI_COMMAND( WorldProperties, "World Settings", "Displays the world settings", EUserInterfaceActionType::Button, FInputGesture() );
 	UI_COMMAND( OpenContentBrowser, "Open Content Browser", "Opens the Content Browser", EUserInterfaceActionType::Button, FInputGesture(EModifierKey::Control|EModifierKey::Shift, EKeys::F) );
 	UI_COMMAND( OpenMarketplace, "Open Marketplace", "Opens the Marketplace", EUserInterfaceActionType::Button, FInputGesture() );
 	UI_COMMAND( AddMatinee, "Add Matinee", "Creates a new matinee actor to edit", EUserInterfaceActionType::Button, FInputGesture() );
diff --git a/Engine/Source/Editor/LevelEditor/Private/LevelEditorToolBar.cpp b/Engine/Source/Editor/LevelEditor/Private/LevelEditorToolBar.cpp
--- a/Engine/Source/Editor/LevelEditor/Private/LevelEditorToolBar.cpp
+++ b/Engine/Source/Editor/LevelEditor/Private/LevelEditorToolBar.cpp
@@ -16,7 +16,7 @@
 #include "EngineBuildSettings.h"
 #include "Matinee/MatineeActor.h"
 #include "Engine/LevelScriptBlueprint.h"
-
+#include "Settings.h"
 
 /**
  * Static: Creates a widget for the level editor tool bar
@@ -62,8 +62,14 @@ TSharedRef< SWidget > FLevelEditorToolBar::MakeLevelEditorToolBar( const TShared
 			FSlateIcon(FEditorStyle::GetStyleSetName(), "LevelEditor.ViewOptions")
 			);
 
-		ToolbarBuilder.AddToolBarButton( FLevelEditorCommands::Get().WorldProperties, NAME_None, LOCTEXT( "WorldProperties_Override", "World Settings" ), LOCTEXT( "WorldProperties_ToolTipOverride", "Displays the world settings" ), TAttribute<FSlateIcon>(), "LevelToolbarWorldSettings" );
-
+		ToolbarBuilder.AddComboButton(
+			FUIAction(),
+			FOnGetContent::CreateStatic( &FLevelEditorToolBar::GenerateOpenGameSettingsMenu, InCommandList ),
+			LOCTEXT( "OpenConfiguration_Label", "Game Settings" ),
+			LOCTEXT( "OpenConfiguration_ToolTip", "Open game settings windows" ),
+			FSlateIcon(FEditorStyle::GetStyleSetName(), "LevelEditor.GameSettings")/*,
+			"LevelToolbarWorldSettings"*/
+			);
 	}
 	ToolbarBuilder.EndSection();
 
@@ -489,6 +495,8 @@ static void MakeScalabilityMenu( FMenuBuilder& MenuBuilder )
 	MenuBuilder.AddWidget(SNew(SScalabilitySettings), FText(), true);
 }
 
+
+
 TSharedRef< SWidget > FLevelEditorToolBar::GenerateQuickSettingsMenu( TSharedRef<FUICommandList> InCommandList )
 {
 #define LOCTEXT_NAMESPACE "LevelToolBarViewMenu"
@@ -613,6 +621,34 @@ TSharedRef< SWidget > FLevelEditorToolBar::GenerateQuickSettingsMenu( TSharedRef
 	return MenuBuilder.MakeWidget();
 }
 
+TSharedRef< SWidget > FLevelEditorToolBar::GenerateOpenGameSettingsMenu( TSharedRef<FUICommandList> InCommandList )
+{
+#define LOCTEXT_NAMESPACE "OpenGameSettingsMenu"
+	const bool bShouldCloseWindowAfterMenuSelection = true;
+	FMenuBuilder MenuBuilder( bShouldCloseWindowAfterMenuSelection, InCommandList );
+
+	MenuBuilder.AddMenuEntry( FLevelEditorCommands::Get().WorldProperties );
+
+	struct Local
+	{
+		static void OpenSettings( FName ContainerName, FName CategoryName, FName SectionName )
+		{
+			FModuleManager::LoadModuleChecked<ISettingsModule>("Settings").ShowViewer(ContainerName, CategoryName, SectionName);
+		}
+	};
+
+	MenuBuilder.AddMenuEntry(
+		LOCTEXT("ProjectSettingsMenuLabel", "Project Settings..."),
+		LOCTEXT("ProjectSettingsMenuToolTip", "Change the settings of the currently loaded project"),
+		FSlateIcon(),
+		FUIAction(FExecuteAction::CreateStatic(&Local::OpenSettings, FName("Project"), FName("Game"), FName("General")))
+	);
+
+#undef LOCTEXT_NAMESPACE
+
+	return MenuBuilder.MakeWidget();
+}
+
 TSharedRef< SWidget > FLevelEditorToolBar::GenerateCompileMenuContent( TSharedRef<FUICommandList> InCommandList )
 {
 #define LOCTEXT_NAMESPACE "LevelToolBarViewMenu"
diff --git a/Engine/Source/Editor/LevelEditor/Private/LevelEditorToolBar.h b/Engine/Source/Editor/LevelEditor/Private/LevelEditorToolBar.h
--- a/Engine/Source/Editor/LevelEditor/Private/LevelEditorToolBar.h
+++ b/Engine/Source/Editor/LevelEditor/Private/LevelEditorToolBar.h
@@ -49,6 +49,13 @@ protected:
 	 */
 	static TSharedRef< SWidget > GenerateQuickSettingsMenu( TSharedRef<FUICommandList> InCommandList );
 
+	/**
+	 * Generates menu content for the settings display combo button drop down menu
+	 *
+	 * @return	Menu content widget
+	 */
+	static TSharedRef< SWidget > GenerateOpenGameSettingsMenu( TSharedRef<FUICommandList> InCommandList );
+
 	/**
 	 * Generates menu content for the compile combo button drop down menu
 	 *
