Fixed up tagging of combo buttons in multibox widgets to use passed-in names if supplied.

The issue in question was because C++ happily converts a string literal to a bool with no complaints, so supplying the tag to the incorrect parameter of AddComboButton() caused no errors/warnings.

TTP# 343944 - Specifying a tutorial hilight name causes combo button icon to not display

diff --git a/Engine/Source/Editor/LevelEditor/Private/LevelEditorToolBar.cpp b/Engine/Source/Editor/LevelEditor/Private/LevelEditorToolBar.cpp
--- a/Engine/Source/Editor/LevelEditor/Private/LevelEditorToolBar.cpp
+++ b/Engine/Source/Editor/LevelEditor/Private/LevelEditorToolBar.cpp
@@ -59,7 +59,9 @@ TSharedRef< SWidget > FLevelEditorToolBar::MakeLevelEditorToolBar( const TShared
 			FOnGetContent::CreateStatic(&FLevelEditorToolBar::GenerateQuickSettingsMenu, InCommandList),
 			LOCTEXT("QuickSettingsCombo", "Quick Settings"),
 			LOCTEXT("QuickSettingsCombo_ToolTip", "Quick level editor settings"),
-			FSlateIcon(FEditorStyle::GetStyleSetName(), "LevelEditor.ViewOptions")
+			FSlateIcon(FEditorStyle::GetStyleSetName(), "LevelEditor.ViewOptions"),
+			false,
+			"LevelToolbarQuickSettings"
 			);
 
 		ToolbarBuilder.AddComboButton(
@@ -67,8 +69,9 @@ TSharedRef< SWidget > FLevelEditorToolBar::MakeLevelEditorToolBar( const TShared
 			FOnGetContent::CreateStatic( &FLevelEditorToolBar::GenerateOpenGameSettingsMenu, InCommandList ),
 			LOCTEXT( "OpenConfiguration_Label", "Game Settings" ),
 			LOCTEXT( "OpenConfiguration_ToolTip", "Open game settings windows" ),
-			FSlateIcon(FEditorStyle::GetStyleSetName(), "LevelEditor.GameSettings")/*,
-			"LevelToolbarWorldSettings"*/
+			FSlateIcon(FEditorStyle::GetStyleSetName(), "LevelEditor.GameSettings"),
+			false,
+			"LevelToolbarWorldSettings"
 			);
 	}
 	ToolbarBuilder.EndSection();
diff --git a/Engine/Source/Runtime/Slate/Private/Framework/MultiBox/SToolBarComboButtonBlock.cpp b/Engine/Source/Runtime/Slate/Private/Framework/MultiBox/SToolBarComboButtonBlock.cpp
--- a/Engine/Source/Runtime/Slate/Private/Framework/MultiBox/SToolBarComboButtonBlock.cpp
+++ b/Engine/Source/Runtime/Slate/Private/Framework/MultiBox/SToolBarComboButtonBlock.cpp
@@ -104,14 +104,23 @@ void SToolBarComboButtonBlock::BuildMultiBlockWidget(const ISlateStyle* StyleSet
 		Label = ToolBarComboButtonBlock->Label;
 	}
 
+	bool bBool = TEXT("Bool");
+
 	// Create a tag
-	FString TagName;
-	TagName = FString::Printf(TEXT("ToolbarComboButton,%s,0"), *Label.Get().ToString());
+	FName TagName;
+	if (ToolBarComboButtonBlock->GetTutorialHighlightName() == NAME_None)
+	{
+		TagName = *FString::Printf(TEXT("ToolbarComboButton,%s,0"), *Label.Get().ToString());
+	}
+	else
+	{
+		TagName = ToolBarComboButtonBlock->GetTutorialHighlightName();
+	}
 	
 	// Create the content for our button
 	TSharedRef< SWidget > ButtonContent =
 		SNew( SVerticalBox )
-		.Tag(*TagName)
+		.Tag(TagName)
 		// Icon image
 		+ SVerticalBox::Slot()
 		.AutoHeight()
