Fix non-unity compile

diff --git a/Engine/Source/Editor/LevelEditor/Private/LevelEditorToolBar.cpp b/Engine/Source/Editor/LevelEditor/Private/LevelEditorToolBar.cpp
--- a/Engine/Source/Editor/LevelEditor/Private/LevelEditorToolBar.cpp
+++ b/Engine/Source/Editor/LevelEditor/Private/LevelEditorToolBar.cpp
@@ -17,6 +17,7 @@
 #include "Matinee/MatineeActor.h"
 #include "Engine/LevelScriptBlueprint.h"
 #include "Settings.h"
+#include "DesktopPlatformModule.h"
 
 /**
  * Static: Creates a widget for the level editor tool bar
