Slate/UMG - The game viewport now disables throttling when it contains focus so that Slate UI in the game doesn't trigger throttling in the editor.  When foucs is lost it restores the old value of Slate.bAllowThrottling.

diff --git a/Engine/Source/Editor/LevelEditor/Private/SLevelViewport.cpp b/Engine/Source/Editor/LevelEditor/Private/SLevelViewport.cpp
--- a/Engine/Source/Editor/LevelEditor/Private/SLevelViewport.cpp
+++ b/Engine/Source/Editor/LevelEditor/Private/SLevelViewport.cpp
@@ -60,6 +60,9 @@ SLevelViewport::SLevelViewport()
 	, bViewTransitionAnimPending( false )
 	, DeviceProfile("Default")
 	, PIEOverlaySlotIndex(0)
+	, bPIEContainsFocus(false)
+	, bPIEHasFocus(false)
+	, UserAllowThrottlingValue(0)
 {
 }
 
@@ -782,6 +785,28 @@ FReply SLevelViewport::OnDrop( const FGeometry& MyGeometry, const FDragDropEvent
 
 void SLevelViewport::Tick( const FGeometry& AllottedGeometry, const double InCurrentTime, const float InDeltaTime )
 {
+	const bool bContainsFocus = HasFocusedDescendants();
+
+	// When we have focus we update the 'Allow Throttling' option in slate to be disabled so that interactions in the
+	// viewport with Slate widgets that are part of the game, don't throttle.
+	if ( bPIEContainsFocus != bContainsFocus )
+	{
+		IConsoleVariable* AllowThrottling = IConsoleManager::Get().FindConsoleVariable(TEXT("Slate.bAllowThrottling"));
+		check(AllowThrottling);
+
+		if ( bContainsFocus )
+		{
+			UserAllowThrottlingValue = AllowThrottling->GetInt();
+			AllowThrottling->Set(0);
+		}
+		else
+		{
+			AllowThrottling->Set(UserAllowThrottlingValue);
+		}
+
+		bPIEContainsFocus = bContainsFocus;
+	}
+
 	// We defer starting animation playback because very often there may be a large hitch after the frame in which
 	// the animation was triggered, and we don't want to start animating until after that hitch.  Otherwise, the
 	// user could miss part of the animation, or even the whole thing!
diff --git a/Engine/Source/Editor/LevelEditor/Public/SLevelViewport.h b/Engine/Source/Editor/LevelEditor/Public/SLevelViewport.h
--- a/Engine/Source/Editor/LevelEditor/Public/SLevelViewport.h
+++ b/Engine/Source/Editor/LevelEditor/Public/SLevelViewport.h
@@ -783,6 +783,12 @@ private:
 	/** Whether the PIE view has focus so we can track when to reshow the mouse control label */
 	bool bPIEHasFocus;
 
+	/** Whether the PIE view contains focus (even if not captured), if so we disable throttling. */
+	bool bPIEContainsFocus;
+
+	/** The users value for allowing throttling, we restore this value when we lose focus. */
+	int32 UserAllowThrottlingValue;
+
 protected:
 	void LockActorInternal(AActor* NewActorToLock);
 
