"Click for mouse control" label no longer appears if the PlayerController::bShowCursor is set

diff --git a/Engine/Source/Editor/LevelEditor/Private/SLevelViewport.cpp b/Engine/Source/Editor/LevelEditor/Private/SLevelViewport.cpp
--- a/Engine/Source/Editor/LevelEditor/Private/SLevelViewport.cpp
+++ b/Engine/Source/Editor/LevelEditor/Private/SLevelViewport.cpp
@@ -3220,6 +3220,23 @@ void SLevelViewport::StartPlayInEditorSession( UGameViewportClient* PlayClient,
 	GEngine->BroadcastLevelActorListChanged();
 }
 
+EVisibility SLevelViewport::GetMouseCaptureLabelVisibility() const
+{
+	EVisibility visibility = EVisibility::Collapsed;
+
+	if (GEditor->PlayWorld)
+	{
+		// Show the label if the local player's PC isn't set to show the cursor
+		auto const TargetPlayer = GEngine->GetLocalPlayerFromControllerId(GEditor->PlayWorld, 0);
+		if (TargetPlayer && TargetPlayer->PlayerController && !TargetPlayer->PlayerController->bShowMouseCursor)
+		{
+			visibility = EVisibility::Visible;
+		}
+	}
+
+	return visibility;
+}
+
 FLinearColor SLevelViewport::GetMouseCaptureLabelColorAndOpacity() const
 {
 	FSlateColor SlateColor = FEditorStyle::GetSlateColor("DefaultForeground");
@@ -3296,6 +3313,7 @@ void SLevelViewport::ShowMouseCaptureLabel(ELabelAnchorMode AnchorMode)
 	[
 		SNew( SBorder )
 		.BorderImage( FEditorStyle::GetBrush("NoBorder") )
+		.Visibility(this, &SLevelViewport::GetMouseCaptureLabelVisibility)
 		.ColorAndOpacity( this, &SLevelViewport::GetMouseCaptureLabelColorAndOpacity )
 		.ForegroundColor( FLinearColor::White )
 		.Padding(15.0f)
diff --git a/Engine/Source/Editor/LevelEditor/Public/SLevelViewport.h b/Engine/Source/Editor/LevelEditor/Public/SLevelViewport.h
--- a/Engine/Source/Editor/LevelEditor/Public/SLevelViewport.h
+++ b/Engine/Source/Editor/LevelEditor/Public/SLevelViewport.h
@@ -630,6 +630,9 @@ private:
 	/** @return Whether or not redo can be executed */
 	bool CanExecuteRedo() const;
 
+	/** @return Whether the mouse capture label is visible */
+	EVisibility GetMouseCaptureLabelVisibility() const;
+
 	/** @return The current color & opacity for the mouse capture label */
 	FLinearColor GetMouseCaptureLabelColorAndOpacity() const;
 
