Can now select World Settings game mode classes from the Blueprints drop down, as well as edit or create new classes based on a specific parent class

#ttp 324358 - RocketGDC: POSTLAUNCH: DEV: Blueprint Button: Need a way to select a blueprint for the game type / player controller / pawn, in addition to creating a new BP

diff --git a/Engine/Source/Editor/LevelEditor/Private/LevelEditorToolBar.cpp b/Engine/Source/Editor/LevelEditor/Private/LevelEditorToolBar.cpp
--- a/Engine/Source/Editor/LevelEditor/Private/LevelEditorToolBar.cpp
+++ b/Engine/Source/Editor/LevelEditor/Private/LevelEditorToolBar.cpp
@@ -18,6 +18,142 @@
 #include "Engine/LevelScriptBlueprint.h"
 #include "Settings.h"
 #include "DesktopPlatformModule.h"
+#include "Editor/ClassViewer/Private/SClassViewer.h"
+#include "Editor/ClassViewer/Public/ClassViewerFilter.h"
+#include "KismetEditorUtilities.h"
+
+namespace LevelEditorActionHelpers
+{
+	/** Filters out any classes for the Class Picker when creating or selecting classes in the Blueprints dropdown */
+	class FBlueprintParentFilter_MapModeSettings : public IClassViewerFilter
+	{
+	public:
+		/** Classes to not allow any children of into the Class Viewer/Picker. */
+		TSet< const UClass* > AllowedChildrenOfClasses;
+
+		virtual bool IsClassAllowed(const FClassViewerInitializationOptions& InInitOptions, const UClass* InClass, TSharedRef< FClassViewerFilterFuncs > InFilterFuncs ) OVERRIDE
+		{
+			return InFilterFuncs->IfInChildOfClassesSet(AllowedChildrenOfClasses, InClass) == EFilterReturn::Passed;
+		}
+
+		virtual bool IsUnloadedClassAllowed(const FClassViewerInitializationOptions& InInitOptions, const TSharedRef< const IUnloadedBlueprintData > InUnloadedClassData, TSharedRef< FClassViewerFilterFuncs > InFilterFuncs) OVERRIDE
+		{
+			return InFilterFuncs->IfInChildOfClassesSet(AllowedChildrenOfClasses, InUnloadedClassData) == EFilterReturn::Passed;
+		}
+	};
+
+	/**
+	 * Builds a sub-menu for selecting a class
+	 *
+	 * @param InMenuBuilder		Object to append menu items/widgets to
+	 * @param InRootClass		The root class to filter the Class Viewer by to only show children of
+	 * @param InOnClassPicked	Callback delegate to fire when a class is picked
+	 */
+	void GetSelectSettingsClassSubMenu(FMenuBuilder& InMenuBuilder, UClass* InRootClass, FOnClassPicked InOnClassPicked)
+	{
+		FClassViewerInitializationOptions Options;
+		Options.Mode = EClassViewerMode::ClassPicker;
+		Options.DisplayMode = EClassViewerDisplayMode::ListView;
+		Options.bShowObjectRootClass = true;
+		Options.bShowNoneOption = true;
+
+		// Only want blueprint actor base classes.
+		Options.bIsBlueprintBaseOnly = true;
+
+		// This will allow unloaded blueprints to be shown.
+		Options.bShowUnloadedBlueprints = true;
+
+		TSharedPtr< FBlueprintParentFilter_MapModeSettings > Filter = MakeShareable(new FBlueprintParentFilter_MapModeSettings);
+		Filter->AllowedChildrenOfClasses.Add(InRootClass);
+		Options.ClassFilter = Filter;
+
+		FText RootClassName = FText::FromString(InRootClass->GetName());
+		auto ClassViewer = StaticCastSharedRef<SClassViewer>(FModuleManager::LoadModuleChecked<FClassViewerModule>("ClassViewer").CreateClassViewer(Options, InOnClassPicked));
+		InMenuBuilder.BeginSection(NAME_None, FText::Format( NSLOCTEXT("LevelToolBarViewMenu", "SelectGameModeLabel", "Select {RootClass} class"), RootClassName ));
+		InMenuBuilder.AddWidget(ClassViewer, FText::GetEmpty(), true);
+		InMenuBuilder.EndSection();
+	}
+
+	/**
+	 * Builds a sub-menu for creating a class
+	 *
+	 * @param InMenuBuilder		Object to append menu items/widgets to
+	 * @param InRootClass		The root class to filter the Class Viewer by to only show children of
+	 * @param InOnClassPicked	Callback delegate to fire when a class is picked
+	 */
+	void GetCreateSettingsClassSubMenu(FMenuBuilder& InMenuBuilder, UClass* InRootClass, FOnClassPicked InOnClassPicked)
+	{
+		FClassViewerInitializationOptions Options;
+		Options.Mode = EClassViewerMode::ClassPicker;
+		Options.DisplayMode = EClassViewerDisplayMode::ListView;
+		Options.bShowObjectRootClass = true;
+
+		// Only want blueprint actor base classes.
+		Options.bIsBlueprintBaseOnly = true;
+
+		// This will allow unloaded blueprints to be shown.
+		Options.bShowUnloadedBlueprints = true;
+
+		TSharedPtr< FBlueprintParentFilter_MapModeSettings > Filter = MakeShareable(new FBlueprintParentFilter_MapModeSettings);
+		Filter->AllowedChildrenOfClasses.Add(InRootClass);
+		Options.ClassFilter = Filter;
+
+		FText RootClassName = FText::FromString(InRootClass->GetName());
+		auto ClassViewer = StaticCastSharedRef<SClassViewer>(FModuleManager::LoadModuleChecked<FClassViewerModule>("ClassViewer").CreateClassViewer(Options, InOnClassPicked));
+		InMenuBuilder.BeginSection(NAME_None, FText::Format( NSLOCTEXT("LevelToolBarViewMenu", "CreateGameModeLabel", "Select {RootClass} parent class"), RootClassName ));
+		InMenuBuilder.AddWidget(ClassViewer, FText::GetEmpty(), true);
+		InMenuBuilder.EndSection();
+	}
+
+	/** Helper struct for passing all required data to the GetBlueprintSettingsSubMenu function */
+	struct FBlueprintMenuSettings
+	{
+		/** The UI command for editing the Blueprint class associated with the menu */
+		TSharedPtr< FUICommandInfo > EditCommand;
+
+		/** Current class associated with the menu */
+		UClass* CurrentClass;
+
+		/** Root class that defines what class children can be set through the menu */
+		UClass* RootClass;
+
+		/** Callback when a class is picked, to assign the new class */
+		FOnClassPicked OnSelectClassPicked;
+
+		/** Callback when a class is picked, to create a new child class of and assign */
+		FOnClassPicked OnCreateClassPicked;
+	};
+
+	/**
+	 * A sub-menu for the Blueprints dropdown, facilitates all the sub-menu actions such as creating, editing, and selecting classes for the world settings game mode.
+	 *
+	 * @param InMenuBuilder		Object to append menu items/widgets to
+	 * @param InSettingsData	All the data needed to create the menu actions
+	 */
+	void GetBlueprintSettingsSubMenu(FMenuBuilder& InMenuBuilder, FBlueprintMenuSettings InSettingsData)
+	{
+#define LOCTEXT_NAMESPACE "LevelToolBarViewMenu"
+
+		FSlateIcon EditBPIcon(FEditorStyle::Get().GetStyleSetName(), TEXT("PropertyWindow.Button_Edit"));
+		FSlateIcon NewBPIcon(FEditorStyle::Get().GetStyleSetName(), TEXT("PropertyWindow.Button_AddToArray"));
+		FText RootClassName = FText::FromString(InSettingsData.RootClass->GetName());
+
+		// If there is currently a valid GameMode Blueprint, offer to edit the Blueprint
+		if(InSettingsData.CurrentClass && InSettingsData.CurrentClass->ClassGeneratedBy)
+		{
+			FText BlueprintName = FText::FromString(InSettingsData.CurrentClass->ClassGeneratedBy->GetName());
+			InMenuBuilder.AddMenuEntry( InSettingsData.EditCommand, NAME_None, FText::Format( LOCTEXT("EditBlueprint", "Edit {Blueprint}"), BlueprintName), FText::Format( LOCTEXT("EditBlueprint_Tooltip", "Open the world's assigned {RootClass} blueprint"), RootClassName), EditBPIcon );
+		}
+
+		// Create a new GameMode, this is always available so the user can easily create a new one
+		InMenuBuilder.AddSubMenu( LOCTEXT("CreateBlueprint", "Create..."), FText::Format( LOCTEXT("CreateClass_Tooltip", "Create a new {RootClass} based on a selected class and auto-assign it to the world"), RootClassName ), FNewMenuDelegate::CreateStatic( &LevelEditorActionHelpers::GetCreateSettingsClassSubMenu, InSettingsData.RootClass, InSettingsData.OnCreateClassPicked ), false, NewBPIcon );
+
+		// Select a game mode, this is always available so the user can switch his selection
+		InMenuBuilder.AddSubMenu( FText::Format( LOCTEXT("SelectGameModeClass", "Select {RootClass} Class"), RootClassName ), FText::Format( LOCTEXT("SelectGameModeClass_Tooltip", "Select a class to be the active {RootClass} in the world"), RootClassName ), FNewMenuDelegate::CreateStatic( &LevelEditorActionHelpers::GetSelectSettingsClassSubMenu, InSettingsData.RootClass, InSettingsData.OnSelectClassPicked ) );
+
+#undef LOCTEXT_NAMESPACE
+	}
+}
 
 /**
  * Static: Creates a widget for the level editor tool bar
@@ -722,7 +858,7 @@ TSharedRef< SWidget > FLevelEditorToolBar::GenerateOpenBlueprintMenuContent( TSh
 		}
 
 
-		/** Generates 'eopn blueprint' sub-menu */
+		/** Generates 'open blueprint' sub-menu */
 		static void MakeOpenClassBPMenu(FMenuBuilder& InMenuBuilder)
 		{
 			FContentBrowserModule& ContentBrowserModule = FModuleManager::Get().LoadModuleChecked<FContentBrowserModule>(TEXT("ContentBrowser"));
@@ -776,60 +912,71 @@ TSharedRef< SWidget > FLevelEditorToolBar::GenerateOpenBlueprintMenuContent( TSh
 	}
 	MenuBuilder.EndSection();
 
-	MenuBuilder.BeginSection(NAME_None, LOCTEXT("GameBlueprints", "Game Blueprints"));
+	MenuBuilder.BeginSection(NAME_None, LOCTEXT("WorldSettingsClasses", "World Settings Classes"));
 	{
-		FSlateIcon EditBPIcon(FEditorStyle::Get().GetStyleSetName(), TEXT("PropertyWindow.Button_Edit"));
-		FSlateIcon NewBPIcon(FEditorStyle::Get().GetStyleSetName(), TEXT("PropertyWindow.Button_AddToArray"));
+		AWorldSettings* WorldSettings = InLevelEditor.Pin()->GetWorld()->GetWorldSettings();
 
 		// Game Mode
-		TAttribute<FText>::FGetter DynamicGameModeGetter;
-		DynamicGameModeGetter.BindStatic(&FLevelEditorToolBar::GetOpenGameModeBlueprintLabel, InLevelEditor);
-		TAttribute<FText> DynamicGameModeLabel = TAttribute<FText>::Create(DynamicGameModeGetter);
-
-		TAttribute<FText>::FGetter DynamicGameModeGetter_Tooltip;
-		DynamicGameModeGetter_Tooltip.BindStatic(&FLevelEditorToolBar::GetOpenGameModeBlueprintTooltip, InLevelEditor);
-		TAttribute<FText> DynamicGameModeTooltip = TAttribute<FText>::Create(DynamicGameModeGetter_Tooltip);
-		MenuBuilder.AddMenuEntry( FLevelEditorCommands::Get().OpenGameModeBlueprint, NAME_None, DynamicGameModeLabel, DynamicGameModeTooltip, IsValidGameModeBlueprint(InLevelEditor)? EditBPIcon : NewBPIcon );
-
-		// Game State
-		TAttribute<FText>::FGetter DynamicGameStateGetter;
-		DynamicGameStateGetter.BindStatic(&FLevelEditorToolBar::GetOpenGameStateBlueprintLabel, InLevelEditor);
-		TAttribute<FText> DynamicGameStateLabel = TAttribute<FText>::Create(DynamicGameStateGetter);
-
-		TAttribute<FText>::FGetter DynamicGameStateGetter_Tooltip;
-		DynamicGameStateGetter_Tooltip.BindStatic(&FLevelEditorToolBar::GetOpenGameStateBlueprintTooltip, InLevelEditor);
-		TAttribute<FText> DynamicGameStateTooltip = TAttribute<FText>::Create(DynamicGameStateGetter_Tooltip);
-		MenuBuilder.AddMenuEntry( FLevelEditorCommands::Get().OpenGameStateBlueprint, NAME_None, DynamicGameStateLabel, DynamicGameStateTooltip, IsValidGameStateBlueprint(InLevelEditor)? EditBPIcon : NewBPIcon );
-
-		// Pawn
-		TAttribute<FText>::FGetter DynamicDefaultPawnGetter;
-		DynamicDefaultPawnGetter.BindStatic(&FLevelEditorToolBar::GetOpenPawnBlueprintLabel, InLevelEditor);
-		TAttribute<FText> DynamicDefaultPawnLabel = TAttribute<FText>::Create(DynamicDefaultPawnGetter);
-
-		TAttribute<FText>::FGetter DynamicDefaultPawnGetter_Tooltip;
-		DynamicDefaultPawnGetter_Tooltip.BindStatic(&FLevelEditorToolBar::GetOpenPawnBlueprintTooltip, InLevelEditor);
-		TAttribute<FText> DynamicDefaultPawnTooltip = TAttribute<FText>::Create(DynamicDefaultPawnGetter_Tooltip);
-		MenuBuilder.AddMenuEntry( FLevelEditorCommands::Get().OpenDefaultPawnBlueprint, NAME_None, DynamicDefaultPawnLabel, DynamicDefaultPawnTooltip, IsValidPawnBlueprint(InLevelEditor)? EditBPIcon : NewBPIcon );
-
-		// HUD
-		TAttribute<FText>::FGetter DynamicHUDGetter;
-		DynamicHUDGetter.BindStatic(&FLevelEditorToolBar::GetOpenHUDBlueprintLabel, InLevelEditor);
-		TAttribute<FText> DynamicHUDLabel = TAttribute<FText>::Create(DynamicHUDGetter);
-
-		TAttribute<FText>::FGetter DynamicHUDGetter_Tooltip;
-		DynamicHUDGetter_Tooltip.BindStatic(&FLevelEditorToolBar::GetOpenHUDBlueprintTooltip, InLevelEditor);
-		TAttribute<FText> DynamicHUDTooltip = TAttribute<FText>::Create(DynamicHUDGetter_Tooltip);
-		MenuBuilder.AddMenuEntry( FLevelEditorCommands::Get().OpenHUDBlueprint, NAME_None, DynamicHUDLabel, DynamicHUDTooltip, IsValidHUDBlueprint(InLevelEditor)? EditBPIcon : NewBPIcon );
-
-		// Player Controller
-		TAttribute<FText>::FGetter DynamicPlayerControllerGetter;
-		DynamicPlayerControllerGetter.BindStatic(&FLevelEditorToolBar::GetOpenPlayerControllerBlueprintLabel, InLevelEditor);
-		TAttribute<FText> DynamicPlayerControllerLabel = TAttribute<FText>::Create(DynamicPlayerControllerGetter);
-
-		TAttribute<FText>::FGetter DynamicPlayerControllerGetter_Tooltip;
-		DynamicPlayerControllerGetter_Tooltip.BindStatic(&FLevelEditorToolBar::GetOpenPlayerControllerBlueprintTooltip, InLevelEditor);
-		TAttribute<FText> DynamicPlayerControllerTooltip = TAttribute<FText>::Create(DynamicPlayerControllerGetter_Tooltip);
-		MenuBuilder.AddMenuEntry( FLevelEditorCommands::Get().OpenPlayerControllerBlueprint, NAME_None, DynamicPlayerControllerLabel, DynamicPlayerControllerTooltip, IsValidPlayerControllerBlueprint(InLevelEditor)? EditBPIcon : NewBPIcon );
+		AGameMode* ActiveGameMode = NULL;
+
+		// If the WorldSettings' default game mode is valid, pull out the default object so we can edit it
+		if(WorldSettings->DefaultGameMode)
+		{
+			ActiveGameMode = Cast<AGameMode>(InLevelEditor.Pin()->GetWorld()->GetWorldSettings()->DefaultGameMode->GetDefaultObject());
+		}
+
+		LevelEditorActionHelpers::FBlueprintMenuSettings GameModeMenuSettings;
+		GameModeMenuSettings.EditCommand = FLevelEditorCommands::Get().OpenGameModeBlueprint;
+		GameModeMenuSettings.OnCreateClassPicked = FOnClassPicked::CreateStatic( &FLevelEditorToolBar::OnCreateGameModeClassPicked, InLevelEditor );;
+		GameModeMenuSettings.OnSelectClassPicked = FOnClassPicked::CreateStatic( &FLevelEditorToolBar::OnSelectGameModeClassPicked, InLevelEditor );;
+		GameModeMenuSettings.CurrentClass = GetActiveGameModeClass(InLevelEditor);
+		GameModeMenuSettings.RootClass = AGameMode::StaticClass();
+
+		MenuBuilder.AddSubMenu( GetOpenGameModeBlueprintLabel(InLevelEditor), LOCTEXT("GameMode_Tooltip", "Select, edit, or create a new GameMode blueprint for the world"), FNewMenuDelegate::CreateStatic( &LevelEditorActionHelpers::GetBlueprintSettingsSubMenu, GameModeMenuSettings) );
+
+		if(ActiveGameMode && ActiveGameMode->GetClass()->ClassGeneratedBy)
+		{
+			// Game State
+			LevelEditorActionHelpers::FBlueprintMenuSettings GameStateMenuSettings;
+			GameStateMenuSettings.EditCommand = FLevelEditorCommands::Get().OpenGameStateBlueprint;
+			GameStateMenuSettings.OnCreateClassPicked = FOnClassPicked::CreateStatic( &FLevelEditorToolBar::OnCreateGameStateClassPicked, InLevelEditor );;
+			GameStateMenuSettings.OnSelectClassPicked = FOnClassPicked::CreateStatic( &FLevelEditorToolBar::OnSelectGameStateClassPicked, InLevelEditor );;
+			GameStateMenuSettings.CurrentClass = GetActiveGameStateClass(InLevelEditor);
+			GameStateMenuSettings.RootClass = AGameState::StaticClass();
+
+			MenuBuilder.AddSubMenu( GetOpenGameStateBlueprintLabel(InLevelEditor), LOCTEXT("GameState_Tooltip", "Select, edit, or create a new GameState blueprint for the world"), FNewMenuDelegate::CreateStatic( &LevelEditorActionHelpers::GetBlueprintSettingsSubMenu, GameStateMenuSettings ) );
+		
+			// Pawn
+			LevelEditorActionHelpers::FBlueprintMenuSettings PawnMenuSettings;
+			PawnMenuSettings.EditCommand = FLevelEditorCommands::Get().OpenDefaultPawnBlueprint;
+			PawnMenuSettings.OnCreateClassPicked = FOnClassPicked::CreateStatic( &FLevelEditorToolBar::OnCreatePawnClassPicked, InLevelEditor );;
+			PawnMenuSettings.OnSelectClassPicked = FOnClassPicked::CreateStatic( &FLevelEditorToolBar::OnSelectPawnClassPicked, InLevelEditor );;
+			PawnMenuSettings.CurrentClass = GetActivePawnClass(InLevelEditor);
+			PawnMenuSettings.RootClass = APawn::StaticClass();
+
+			MenuBuilder.AddSubMenu( GetOpenPawnBlueprintLabel(InLevelEditor), LOCTEXT("Pawn_Tooltip", "Select, edit, or create a new Pawn blueprint for the world"), FNewMenuDelegate::CreateStatic( &LevelEditorActionHelpers::GetBlueprintSettingsSubMenu, PawnMenuSettings ) );
+
+			// HUD
+			LevelEditorActionHelpers::FBlueprintMenuSettings HUDMenuSettings;
+			HUDMenuSettings.EditCommand = FLevelEditorCommands::Get().OpenHUDBlueprint;
+			HUDMenuSettings.OnCreateClassPicked = FOnClassPicked::CreateStatic( &FLevelEditorToolBar::OnCreateHUDClassPicked, InLevelEditor );;
+			HUDMenuSettings.OnSelectClassPicked = FOnClassPicked::CreateStatic( &FLevelEditorToolBar::OnSelectHUDClassPicked, InLevelEditor );;
+			HUDMenuSettings.CurrentClass = GetActiveHUDClass(InLevelEditor);
+			HUDMenuSettings.RootClass = AHUD::StaticClass();
+
+			MenuBuilder.AddSubMenu( GetOpenHUDBlueprintLabel(InLevelEditor), LOCTEXT("HUD_Tooltip", "Select, edit, or create a new HUD blueprint for the world"), FNewMenuDelegate::CreateStatic( &LevelEditorActionHelpers::GetBlueprintSettingsSubMenu, HUDMenuSettings ) );
+
+			// Player Controller
+			LevelEditorActionHelpers::FBlueprintMenuSettings PlayerControllerMenuSettings;
+			PlayerControllerMenuSettings.EditCommand = FLevelEditorCommands::Get().OpenHUDBlueprint;
+			PlayerControllerMenuSettings.OnCreateClassPicked = FOnClassPicked::CreateStatic( &FLevelEditorToolBar::OnCreatePlayerControllerClassPicked, InLevelEditor );;
+			PlayerControllerMenuSettings.OnSelectClassPicked = FOnClassPicked::CreateStatic( &FLevelEditorToolBar::OnSelectPlayerControllerClassPicked, InLevelEditor );;
+			PlayerControllerMenuSettings.CurrentClass = GetActivePlayerControllerClass(InLevelEditor);
+			PlayerControllerMenuSettings.RootClass = APlayerController::StaticClass();
+
+			MenuBuilder.AddSubMenu( GetOpenPlayerControllerBlueprintLabel(InLevelEditor), LOCTEXT("PlayerController_Tooltip", "Select, edit, or create a new Player controller blueprint for the world"), FNewMenuDelegate::CreateStatic( &LevelEditorActionHelpers::GetBlueprintSettingsSubMenu, PlayerControllerMenuSettings ) );
+
+		}
 	}
 	MenuBuilder.EndSection();
 
@@ -868,179 +1015,312 @@ void FLevelEditorToolBar::OnOpenSubLevelBlueprint( ULevel* InLevel )
 	}
 }
 
-bool FLevelEditorToolBar::IsValidGameModeBlueprint(TWeakPtr< SLevelEditor > InLevelEditor)
+UClass* FLevelEditorToolBar::GetActiveGameModeClass(TWeakPtr< SLevelEditor > InLevelEditor)
 {
 	AWorldSettings* WorldSettings = InLevelEditor.Pin()->GetWorld()->GetWorldSettings();
-	return WorldSettings->DefaultGameMode && WorldSettings->DefaultGameMode->ClassGeneratedBy;
+	if(WorldSettings->DefaultGameMode)
+	{
+		return WorldSettings->DefaultGameMode;
+	}
+	return NULL;
 }
 
 FText FLevelEditorToolBar::GetOpenGameModeBlueprintLabel(TWeakPtr< SLevelEditor > InLevelEditor)
 {
 #define LOCTEXT_NAMESPACE "LevelToolBarViewMenu"
-	if(IsValidGameModeBlueprint(InLevelEditor))
+	if(UClass* GameModeClass = GetActiveGameModeClass(InLevelEditor))
 	{
-		return FText::Format( LOCTEXT("GameModeEditBlueprint", "GameMode: Edit {GameModeName}"), FText::FromString(InLevelEditor.Pin()->GetWorld()->GetWorldSettings()->DefaultGameMode->ClassGeneratedBy->GetName()));
+		if(GameModeClass->ClassGeneratedBy)
+		{
+			return FText::Format( LOCTEXT("GameModeEditBlueprint", "GameMode: Edit {GameModeName}"), FText::FromString(GameModeClass->ClassGeneratedBy->GetName()));
+		}
+
+		return FText::Format( LOCTEXT("GameModeBlueprint", "GameMode: {GameModeName}"), FText::FromString(GameModeClass->GetName()));
 	}
 
-	return LOCTEXT("GameModeCreateBlueprint", "GameMode: Create...");
+	return LOCTEXT("GameModeCreateBlueprint", "GameMode: New...");
 #undef LOCTEXT_NAMESPACE
 }
 
-FText FLevelEditorToolBar::GetOpenGameModeBlueprintTooltip(TWeakPtr< SLevelEditor > InLevelEditor)
+void FLevelEditorToolBar::OnCreateGameModeClassPicked(UClass* InChosenClass, TWeakPtr< SLevelEditor > InLevelEditor)
 {
-#define LOCTEXT_NAMESPACE "LevelToolBarViewMenu"
-	if(IsValidGameModeBlueprint(InLevelEditor))
+	if(InChosenClass)
 	{
-		return LOCTEXT("GameModeEditBlueprint_Tooltip", "Open the World Setting's assigned GameMode Blueprint");
-	}
+		const FString NewBPName(TEXT("NewGameMode"));
+		UBlueprint* Blueprint = FKismetEditorUtilities::CreateBlueprintFromClass(NSLOCTEXT("LevelEditorCommands", "CreateGameModeBlueprint_Title", "Create GameMode Blueprint"), InChosenClass, NewBPName);
 
-	return LOCTEXT("GameModeCreateBlueprint_Tooltip", "Creates a new GameMode Blueprint and auto-assigns it to the current World Settings");
-#undef LOCTEXT_NAMESPACE
+		if(Blueprint)
+		{
+			// @todo Re-enable once world centric works
+			const bool bOpenWorldCentric = false;
+			FAssetEditorManager::Get().OpenEditorForAsset(
+				Blueprint,
+				bOpenWorldCentric ? EToolkitMode::WorldCentric : EToolkitMode::Standalone,
+				InLevelEditor.Pin()  );
+
+			InLevelEditor.Pin()->GetWorld()->GetWorldSettings()->DefaultGameMode = Cast<UClass>(Blueprint->GeneratedClass);
+		}
+	}
+	FSlateApplication::Get().DismissAllMenus();
 }
 
+void FLevelEditorToolBar::OnSelectGameModeClassPicked(UClass* InChosenClass, TWeakPtr< SLevelEditor > InLevelEditor)
+{
+	InLevelEditor.Pin()->GetWorld()->GetWorldSettings()->DefaultGameMode = InChosenClass;
 
-bool FLevelEditorToolBar::IsValidGameStateBlueprint(TWeakPtr< SLevelEditor > InLevelEditor)
+	FSlateApplication::Get().DismissAllMenus();
+}
+
+UClass* FLevelEditorToolBar::GetActiveGameStateClass(TWeakPtr< SLevelEditor > InLevelEditor)
 {
 	AWorldSettings* WorldSettings = InLevelEditor.Pin()->GetWorld()->GetWorldSettings();
 	if(WorldSettings->DefaultGameMode)
 	{
 		AGameMode* ActiveGameMode = Cast<AGameMode>(WorldSettings->DefaultGameMode->GetDefaultObject());
-		return ActiveGameMode && ActiveGameMode->GameStateClass && ActiveGameMode->GameStateClass->ClassGeneratedBy;
+		if(ActiveGameMode)
+		{
+			return ActiveGameMode->GameStateClass;
+		}
 	}
-	return false;
+	return NULL;
 }
 
 FText FLevelEditorToolBar::GetOpenGameStateBlueprintLabel(TWeakPtr< SLevelEditor > InLevelEditor)
 {
 #define LOCTEXT_NAMESPACE "LevelToolBarViewMenu"
-	if(IsValidGameStateBlueprint(InLevelEditor))
+	if(UClass* GameStateClass = GetActiveGameStateClass(InLevelEditor))
 	{
-		AGameMode* ActiveGameMode = Cast<AGameMode>(InLevelEditor.Pin()->GetWorld()->GetWorldSettings()->DefaultGameMode->GetDefaultObject());
-		return FText::Format( LOCTEXT("GameStateEditBlueprint", "GameState: Edit {GameStateName}"), FText::FromString(ActiveGameMode->GameStateClass->ClassGeneratedBy->GetName()));
+		if(GameStateClass->ClassGeneratedBy)
+		{
+			return FText::Format( LOCTEXT("GameStateEditBlueprint", "GameState: Edit {GameStateName}"), FText::FromString(GameStateClass->ClassGeneratedBy->GetName()));
+		}
+
+		return FText::Format( LOCTEXT("GameStateBlueprint", "GameState: {GameStateName}"), FText::FromString(GameStateClass->GetName()));
 	}
 
-	return LOCTEXT("GameStateCreateBlueprint", "GameState: Create...");
+	return LOCTEXT("GameStateCreateBlueprint", "GameState: New...");
 #undef LOCTEXT_NAMESPACE
 }
 
-FText FLevelEditorToolBar::GetOpenGameStateBlueprintTooltip(TWeakPtr< SLevelEditor > InLevelEditor)
+void FLevelEditorToolBar::OnCreateGameStateClassPicked(UClass* InChosenClass, TWeakPtr< SLevelEditor > InLevelEditor)
 {
-#define LOCTEXT_NAMESPACE "LevelToolBarViewMenu"
-	if(IsValidGameStateBlueprint(InLevelEditor))
+	if(InChosenClass)
 	{
-		return LOCTEXT("GameStateEditBlueprint_Tooltip", "Open the World Setting's assigned GameState Blueprint");
+		const FString NewBPName(TEXT("NewGameState"));
+		UBlueprint* Blueprint = FKismetEditorUtilities::CreateBlueprintFromClass(NSLOCTEXT("LevelEditorCommands", "CreateGameStateBlueprint_Title", "Create GameState Blueprint"), InChosenClass, NewBPName);
+
+		if(Blueprint)
+		{
+			// @todo Re-enable once world centric works
+			const bool bOpenWorldCentric = false;
+			FAssetEditorManager::Get().OpenEditorForAsset(
+				Blueprint,
+				bOpenWorldCentric ? EToolkitMode::WorldCentric : EToolkitMode::Standalone,
+				InLevelEditor.Pin()  );
+
+			AGameMode* ActiveGameMode = Cast<AGameMode>(InLevelEditor.Pin()->GetWorld()->GetWorldSettings()->DefaultGameMode->GetDefaultObject());
+			ActiveGameMode->GameStateClass = Cast<UClass>(Blueprint->GeneratedClass);
+		}
 	}
+	FSlateApplication::Get().DismissAllMenus();
+}
 
-	return LOCTEXT("GameStateCreateBlueprint_Tooltip", "Creates a new GameState Blueprint and auto-assigns it to the current World Settings");
-#undef LOCTEXT_NAMESPACE
+void FLevelEditorToolBar::OnSelectGameStateClassPicked(UClass* InChosenClass, TWeakPtr< SLevelEditor > InLevelEditor)
+{
+	AGameMode* ActiveGameMode = Cast<AGameMode>(InLevelEditor.Pin()->GetWorld()->GetWorldSettings()->DefaultGameMode->GetDefaultObject());
+
+	ActiveGameMode->GameStateClass = InChosenClass;
+
+	FSlateApplication::Get().DismissAllMenus();
 }
 
-bool FLevelEditorToolBar::IsValidPawnBlueprint(TWeakPtr< SLevelEditor > InLevelEditor)
+UClass* FLevelEditorToolBar::GetActivePawnClass(TWeakPtr< SLevelEditor > InLevelEditor)
 {
 	AWorldSettings* WorldSettings = InLevelEditor.Pin()->GetWorld()->GetWorldSettings();
 	if(WorldSettings->DefaultGameMode)
 	{
 		AGameMode* ActiveGameMode = Cast<AGameMode>(WorldSettings->DefaultGameMode->GetDefaultObject());
-		return ActiveGameMode && ActiveGameMode->DefaultPawnClass && ActiveGameMode->DefaultPawnClass->ClassGeneratedBy;
+		if(ActiveGameMode)
+		{
+			return ActiveGameMode->DefaultPawnClass;
+		}
 	}
-	return false;
+	return NULL;
 }
 
 FText FLevelEditorToolBar::GetOpenPawnBlueprintLabel(TWeakPtr< SLevelEditor > InLevelEditor)
 {
 #define LOCTEXT_NAMESPACE "LevelToolBarViewMenu"
-	if(IsValidPawnBlueprint(InLevelEditor))
+	if(UClass* PawnClass = GetActivePawnClass(InLevelEditor))
 	{
-		AGameMode* ActiveGameMode = Cast<AGameMode>(InLevelEditor.Pin()->GetWorld()->GetWorldSettings()->DefaultGameMode->GetDefaultObject());
-		return FText::Format( LOCTEXT("PawnEditBlueprint", "Pawn: Edit {PawnName}"), FText::FromString(ActiveGameMode->DefaultPawnClass->ClassGeneratedBy->GetName()));
+		if(PawnClass->ClassGeneratedBy)
+		{
+			return FText::Format( LOCTEXT("PawnEditBlueprint", "Pawn: Edit {PawnName}"), FText::FromString(PawnClass->ClassGeneratedBy->GetName()));
+		}
+
+		return FText::Format( LOCTEXT("PawnBlueprint", "Pawn: {PawnName}"), FText::FromString(PawnClass->GetName()));
 	}
 
-	return LOCTEXT("PawnCreateBlueprint", "Pawn: Create...");
+	return LOCTEXT("PawnCreateBlueprint", "Pawn: New...");
 #undef LOCTEXT_NAMESPACE
 }
 
-FText FLevelEditorToolBar::GetOpenPawnBlueprintTooltip(TWeakPtr< SLevelEditor > InLevelEditor)
+void FLevelEditorToolBar::OnCreatePawnClassPicked(UClass* InChosenClass, TWeakPtr< SLevelEditor > InLevelEditor)
 {
-#define LOCTEXT_NAMESPACE "LevelToolBarViewMenu"
-	if(IsValidPawnBlueprint(InLevelEditor))
+	if(InChosenClass)
 	{
-		return LOCTEXT("PawnEditBlueprint_Tooltip", "Open the World Setting's assigned Pawn Blueprint");
+		const FString NewBPName(TEXT("NewPawn"));
+		UBlueprint* Blueprint = FKismetEditorUtilities::CreateBlueprintFromClass(NSLOCTEXT("LevelEditorCommands", "CreatePawnBlueprint_Title", "Create Pawn Blueprint"), InChosenClass, NewBPName);
+
+		if(Blueprint)
+		{
+			// @todo Re-enable once world centric works
+			const bool bOpenWorldCentric = false;
+			FAssetEditorManager::Get().OpenEditorForAsset(
+				Blueprint,
+				bOpenWorldCentric ? EToolkitMode::WorldCentric : EToolkitMode::Standalone,
+				InLevelEditor.Pin()  );
+
+			AGameMode* ActiveGameMode = Cast<AGameMode>(InLevelEditor.Pin()->GetWorld()->GetWorldSettings()->DefaultGameMode->GetDefaultObject());
+			ActiveGameMode->DefaultPawnClass = Cast<UClass>(Blueprint->GeneratedClass);
+		}
 	}
+	FSlateApplication::Get().DismissAllMenus();
+}
 
-	return LOCTEXT("PawnCreateBlueprint_Tooltip", "Creates a new Pawn Blueprint and auto-assigns it to the current World Settings");
-#undef LOCTEXT_NAMESPACE
+void FLevelEditorToolBar::OnSelectPawnClassPicked(UClass* InChosenClass, TWeakPtr< SLevelEditor > InLevelEditor)
+{
+	AGameMode* ActiveGameMode = Cast<AGameMode>(InLevelEditor.Pin()->GetWorld()->GetWorldSettings()->DefaultGameMode->GetDefaultObject());
+
+	ActiveGameMode->DefaultPawnClass = InChosenClass;
+
+	FSlateApplication::Get().DismissAllMenus();
 }
 
-bool FLevelEditorToolBar::IsValidHUDBlueprint(TWeakPtr< SLevelEditor > InLevelEditor)
+UClass* FLevelEditorToolBar::GetActiveHUDClass(TWeakPtr< SLevelEditor > InLevelEditor)
 {
 	AWorldSettings* WorldSettings = InLevelEditor.Pin()->GetWorld()->GetWorldSettings();
 	if(WorldSettings->DefaultGameMode)
 	{
 		AGameMode* ActiveGameMode = Cast<AGameMode>(WorldSettings->DefaultGameMode->GetDefaultObject());
-		return ActiveGameMode && ActiveGameMode->HUDClass && ActiveGameMode->HUDClass->ClassGeneratedBy;
+		if(ActiveGameMode)
+		{
+			return ActiveGameMode->HUDClass;
+		}
 	}
-	return false;
+	return NULL;
 }
 
 FText FLevelEditorToolBar::GetOpenHUDBlueprintLabel(TWeakPtr< SLevelEditor > InLevelEditor)
 {
 #define LOCTEXT_NAMESPACE "LevelToolBarViewMenu"
-	if(IsValidHUDBlueprint(InLevelEditor))
+	if(UClass* HUDClass = GetActiveHUDClass(InLevelEditor))
 	{
-		AGameMode* ActiveGameMode = Cast<AGameMode>(InLevelEditor.Pin()->GetWorld()->GetWorldSettings()->DefaultGameMode->GetDefaultObject());
-		return FText::Format( LOCTEXT("HUDEditBlueprint", "HUD: Edit {HUDName}"), FText::FromString(ActiveGameMode->HUDClass->ClassGeneratedBy->GetName()));
+		if(HUDClass->ClassGeneratedBy)
+		{
+			return FText::Format( LOCTEXT("HUDEditBlueprint", "HUD: Edit {HUDName}"), FText::FromString(HUDClass->ClassGeneratedBy->GetName()));
+		}
+
+		return FText::Format( LOCTEXT("HUDBlueprint", "HUD: {HUDName}"), FText::FromString(HUDClass->GetName()));
 	}
 
-	return LOCTEXT("HUDCreateBlueprint", "HUD: Create...");
+	return LOCTEXT("HUDCreateBlueprint", "HUD: New...");
 #undef LOCTEXT_NAMESPACE
 }
 
-FText FLevelEditorToolBar::GetOpenHUDBlueprintTooltip(TWeakPtr< SLevelEditor > InLevelEditor)
+void FLevelEditorToolBar::OnCreateHUDClassPicked(UClass* InChosenClass, TWeakPtr< SLevelEditor > InLevelEditor)
 {
-#define LOCTEXT_NAMESPACE "LevelToolBarViewMenu"
-	if(IsValidHUDBlueprint(InLevelEditor))
+	if(InChosenClass)
 	{
-		return LOCTEXT("HUDEditBlueprint_Tooltip", "Open the World Setting's assigned HUD blueprint");
+		const FString NewBPName(TEXT("NewHUD"));
+		UBlueprint* Blueprint = FKismetEditorUtilities::CreateBlueprintFromClass(NSLOCTEXT("LevelEditorCommands", "CreateHUDBlueprint_Title", "Create HUD Blueprint"), InChosenClass, NewBPName);
+
+		if(Blueprint)
+		{
+			// @todo Re-enable once world centric works
+			const bool bOpenWorldCentric = false;
+			FAssetEditorManager::Get().OpenEditorForAsset(
+				Blueprint,
+				bOpenWorldCentric ? EToolkitMode::WorldCentric : EToolkitMode::Standalone,
+				InLevelEditor.Pin()  );
+
+			AGameMode* ActiveGameMode = Cast<AGameMode>(InLevelEditor.Pin()->GetWorld()->GetWorldSettings()->DefaultGameMode->GetDefaultObject());
+			ActiveGameMode->HUDClass = Cast<UClass>(Blueprint->GeneratedClass);
+		}
 	}
+	FSlateApplication::Get().DismissAllMenus();
+}
 
-	return LOCTEXT("HUDCreateBlueprint_Tooltip", "Creates a new HUD Blueprint and auto-assigns it to the current World Settings");
-#undef LOCTEXT_NAMESPACE
+void FLevelEditorToolBar::OnSelectHUDClassPicked(UClass* InChosenClass, TWeakPtr< SLevelEditor > InLevelEditor)
+{
+	AGameMode* ActiveGameMode = Cast<AGameMode>(InLevelEditor.Pin()->GetWorld()->GetWorldSettings()->DefaultGameMode->GetDefaultObject());
+
+	ActiveGameMode->HUDClass = InChosenClass;
+
+	FSlateApplication::Get().DismissAllMenus();
 }
 
-bool FLevelEditorToolBar::IsValidPlayerControllerBlueprint(TWeakPtr< SLevelEditor > InLevelEditor)
+UClass* FLevelEditorToolBar::GetActivePlayerControllerClass(TWeakPtr< SLevelEditor > InLevelEditor)
 {
 	AWorldSettings* WorldSettings = InLevelEditor.Pin()->GetWorld()->GetWorldSettings();
 	if(WorldSettings->DefaultGameMode)
 	{
 		AGameMode* ActiveGameMode = Cast<AGameMode>(WorldSettings->DefaultGameMode->GetDefaultObject());
-		return ActiveGameMode && ActiveGameMode->PlayerControllerClass && ActiveGameMode->PlayerControllerClass->ClassGeneratedBy;
+		if(ActiveGameMode)
+		{
+			return ActiveGameMode->PlayerControllerClass;
+		}
 	}
-	return false;
+	return NULL;
 }
 
 FText FLevelEditorToolBar::GetOpenPlayerControllerBlueprintLabel(TWeakPtr< SLevelEditor > InLevelEditor)
 {
 #define LOCTEXT_NAMESPACE "LevelToolBarViewMenu"
-	if(IsValidPlayerControllerBlueprint(InLevelEditor))
+	if(UClass* PlayerControllerClass = GetActivePlayerControllerClass(InLevelEditor))
 	{
-		AGameMode* ActiveGameMode = Cast<AGameMode>(InLevelEditor.Pin()->GetWorld()->GetWorldSettings()->DefaultGameMode->GetDefaultObject());
-		return FText::Format( LOCTEXT("PlayerControllerEditBlueprint", "PlayerController: Edit {PlayerControllerName}"), FText::FromString(ActiveGameMode->PlayerControllerClass->ClassGeneratedBy->GetName()));
+		if(PlayerControllerClass->ClassGeneratedBy)
+		{
+			return FText::Format( LOCTEXT("PlayerControllerEditBlueprint", "PlayerController: Edit {PlayerControllerName}"), FText::FromString(PlayerControllerClass->ClassGeneratedBy->GetName()));
+		}
+
+		return FText::Format( LOCTEXT("PlayerControllerBlueprint", "PlayerController: {PlayerControllerName}"), FText::FromString(PlayerControllerClass->GetName()));
 	}
 
-	return LOCTEXT("PlayerControllerCreateBlueprint", "PlayerController: Create...");
+	return LOCTEXT("PlayerControllerCreateBlueprint", "PlayerController: New...");
 #undef LOCTEXT_NAMESPACE
 }
 
-FText FLevelEditorToolBar::GetOpenPlayerControllerBlueprintTooltip(TWeakPtr< SLevelEditor > InLevelEditor)
+void FLevelEditorToolBar::OnCreatePlayerControllerClassPicked(UClass* InChosenClass, TWeakPtr< SLevelEditor > InLevelEditor)
 {
-#define LOCTEXT_NAMESPACE "LevelToolBarViewMenu"
-	if(IsValidPlayerControllerBlueprint(InLevelEditor))
+	if(InChosenClass)
 	{
-		return LOCTEXT("PlayerControllerEditBlueprint_Tooltip", "Open the World Setting's assigned PlayerController Blueprint");
+		const FString NewBPName(TEXT("NewPlayerController"));
+		UBlueprint* Blueprint = FKismetEditorUtilities::CreateBlueprintFromClass(NSLOCTEXT("LevelEditorCommands", "CreatePlayerControllerBlueprint_Title", "Create PlayerController Blueprint"), InChosenClass, NewBPName);
+
+		if(Blueprint)
+		{
+			// @todo Re-enable once world centric works
+			const bool bOpenWorldCentric = false;
+			FAssetEditorManager::Get().OpenEditorForAsset(
+				Blueprint,
+				bOpenWorldCentric ? EToolkitMode::WorldCentric : EToolkitMode::Standalone,
+				InLevelEditor.Pin()  );
+
+			AGameMode* ActiveGameMode = Cast<AGameMode>(InLevelEditor.Pin()->GetWorld()->GetWorldSettings()->DefaultGameMode->GetDefaultObject());
+			ActiveGameMode->PlayerControllerClass = Cast<UClass>(Blueprint->GeneratedClass);
+		}
 	}
+	FSlateApplication::Get().DismissAllMenus();
+}
 
-	return LOCTEXT("PlayerControllerCreateBlueprint_Tooltip", "Creates a new PlayerController Blueprint and auto-assigns it to the current World Settings");
-#undef LOCTEXT_NAMESPACE
+void FLevelEditorToolBar::OnSelectPlayerControllerClassPicked(UClass* InChosenClass, TWeakPtr< SLevelEditor > InLevelEditor)
+{
+	AGameMode* ActiveGameMode = Cast<AGameMode>(InLevelEditor.Pin()->GetWorld()->GetWorldSettings()->DefaultGameMode->GetDefaultObject());
+
+	ActiveGameMode->PlayerControllerClass = InChosenClass;
+
+	FSlateApplication::Get().DismissAllMenus();
 }
 
 TSharedRef< SWidget > FLevelEditorToolBar::GenerateMatineeMenuContent( TSharedRef<FUICommandList> InCommandList, TWeakPtr<SLevelEditor> LevelEditorWeakPtr )
diff --git a/Engine/Source/Editor/LevelEditor/Private/LevelEditorToolBar.h b/Engine/Source/Editor/LevelEditor/Private/LevelEditorToolBar.h
--- a/Engine/Source/Editor/LevelEditor/Private/LevelEditorToolBar.h
+++ b/Engine/Source/Editor/LevelEditor/Private/LevelEditorToolBar.h
@@ -91,40 +91,46 @@ protected:
 	static void OnOpenSubLevelBlueprint( ULevel* InLevel );
 
 	/**
-	 * Checks if the passed in world setting's GameMode is a valid Blueprint
+	 * Retrieves the active GameMode class from
 	 *
 	 * @param InLevelEditor		The editor to extract the world from
-	 * @return					TRUE if the GameMode is a Blueprint
+	 * @return					The active GameMode class in the World Settings
 	 */
-	static bool IsValidGameModeBlueprint(TWeakPtr< SLevelEditor > InLevelEditor);
+	static UClass* GetActiveGameModeClass(TWeakPtr< SLevelEditor > InLevelEditor);
 
 	/** Callback for the label to display for the GameMode menu selection */
 	static FText GetOpenGameModeBlueprintLabel(TWeakPtr< SLevelEditor > InLevelEditor);
 
-	/** Callback for the tooltip to display for the GameMode menu selection */
-	static FText GetOpenGameModeBlueprintTooltip(TWeakPtr< SLevelEditor > InLevelEditor);
+	/** Callback when selecting a GameMode class, assigns it to the world */
+	static void OnSelectGameModeClassPicked(UClass* InChosenClass, TWeakPtr< SLevelEditor > InLevelEditor);
+
+	/** Callback when creating a new GameMode class, creates the Blueprint and assigns it to the world */
+	static void OnCreateGameModeClassPicked(UClass* InChosenClass, TWeakPtr< SLevelEditor > InLevelEditor);
 
 	/**
-	 * Checks if the passed in world setting's GameState is a valid Blueprint
+	 * Retrieves the active GameState class from
 	 *
 	 * @param InLevelEditor		The editor to extract the world from
-	 * @return					TRUE if the GameState is a Blueprint
+	 * @return					The active GameState class in the World Settings
 	 */
-	static bool IsValidGameStateBlueprint(TWeakPtr< SLevelEditor > InLevelEditor);
+	static UClass* GetActiveGameStateClass(TWeakPtr< SLevelEditor > InLevelEditor);
 
 	/** Callback for the label to display for the GameState menu selection */
 	static FText GetOpenGameStateBlueprintLabel(TWeakPtr< SLevelEditor > InLevelEditor);
 	
-	/** Callback for the tooltip to display for the GameState menu selection */
-	static FText GetOpenGameStateBlueprintTooltip(TWeakPtr< SLevelEditor > InLevelEditor);
+	/** Callback when selecting a GameState class, assigns it to the world */
+	static void OnSelectGameStateClassPicked(UClass* InChosenClass, TWeakPtr< SLevelEditor > InLevelEditor);
+
+	/** Callback when creating a new GameState class, creates the Blueprint and assigns it to the world */
+	static void OnCreateGameStateClassPicked(UClass* InChosenClass, TWeakPtr< SLevelEditor > InLevelEditor);
 
 	/**
-	 * Checks if the passed in world setting's Pawn is a valid Blueprint
+	 * Retrieves the active Pawn class from
 	 *
 	 * @param InLevelEditor		The editor to extract the world from
-	 * @return					TRUE if the Pawn is a Blueprint
+	 * @return					The active Pawn class in the World Settings
 	 */
-	static bool IsValidPawnBlueprint(TWeakPtr< SLevelEditor > InLevelEditor);
+	static UClass* GetActivePawnClass(TWeakPtr< SLevelEditor > InLevelEditor);
 
 	/** Callback for the label to display for the Pawn menu selection */
 	static FText GetOpenPawnBlueprintLabel(TWeakPtr< SLevelEditor > InLevelEditor);
@@ -132,33 +138,45 @@ protected:
 	/** Callback for the tooltip to display for the Pawn menu selection */
 	static FText GetOpenPawnBlueprintTooltip(TWeakPtr< SLevelEditor > InLevelEditor);
 
+	/** Callback when selecting a Pawn class, assigns it to the world */
+	static void OnSelectPawnClassPicked(UClass* InChosenClass, TWeakPtr< SLevelEditor > InLevelEditor);
+
+	/** Callback when creating a new Pawn class, creates the Blueprint and assigns it to the world */
+	static void OnCreatePawnClassPicked(UClass* InChosenClass, TWeakPtr< SLevelEditor > InLevelEditor);
+
 	/**
-	 * Checks if the passed in world setting's HUD is a valid Blueprint
+	 * Retrieves the active HUD class from
 	 *
 	 * @param InLevelEditor		The editor to extract the world from
-	 * @return					TRUE if the HUD is a Blueprint
+	 * @return					The active HUD class in the World Settings
 	 */
-	static bool IsValidHUDBlueprint(TWeakPtr< SLevelEditor > InLevelEditor);
+	static UClass* GetActiveHUDClass(TWeakPtr< SLevelEditor > InLevelEditor);
 
 	/** Callback for the label to display for the HUD menu selection */
 	static FText GetOpenHUDBlueprintLabel(TWeakPtr< SLevelEditor > InLevelEditor);
 	
-	/** Callback for the tooltip to display for the HUD menu selection */
-	static FText GetOpenHUDBlueprintTooltip(TWeakPtr< SLevelEditor > InLevelEditor);
+	/** Callback when selecting a HUD class, assigns it to the world */
+	static void OnSelectHUDClassPicked(UClass* InChosenClass, TWeakPtr< SLevelEditor > InLevelEditor);
+
+	/** Callback when creating a new HUD class, creates the Blueprint and assigns it to the world */
+	static void OnCreateHUDClassPicked(UClass* InChosenClass, TWeakPtr< SLevelEditor > InLevelEditor);
 
 	/**
-	 * Checks if the passed in world setting's PlayerController is a valid Blueprint
+	 * Retrieves the active PlayerController class from
 	 *
 	 * @param InLevelEditor		The editor to extract the world from
-	 * @return					TRUE if the PlayerController is a Blueprint
+	 * @return					The active PlayerController class in the World Settings
 	 */
-	static bool IsValidPlayerControllerBlueprint(TWeakPtr< SLevelEditor > InLevelEditor);
+	static UClass* GetActivePlayerControllerClass(TWeakPtr< SLevelEditor > InLevelEditor);
 
 	/** Callback for the label to display for the PlayerController menu selection */
 	static FText GetOpenPlayerControllerBlueprintLabel(TWeakPtr< SLevelEditor > InLevelEditor);
 
-	/** Callback for the tooltip to display for the PlayerController menu selection */
-	static FText GetOpenPlayerControllerBlueprintTooltip(TWeakPtr< SLevelEditor > InLevelEditor);
+	/** Callback when selecting a PlayerController class, assigns it to the world */
+	static void OnSelectPlayerControllerClassPicked(UClass* InChosenClass, TWeakPtr< SLevelEditor > InLevelEditor);
+
+	/** Callback when creating a new PlayerController class, creates the Blueprint and assigns it to the world */
+	static void OnCreatePlayerControllerClassPicked(UClass* InChosenClass, TWeakPtr< SLevelEditor > InLevelEditor);
 };
 
 
