Consolidated the settings buttons in the toolbar into one button.

diff --git a/Engine/Source/Editor/DetailCustomizations/Private/ActorDetails.cpp b/Engine/Source/Editor/DetailCustomizations/Private/ActorDetails.cpp
--- a/Engine/Source/Editor/DetailCustomizations/Private/ActorDetails.cpp
+++ b/Engine/Source/Editor/DetailCustomizations/Private/ActorDetails.cpp
@@ -126,10 +126,10 @@ void FActorDetails::CustomizeDetails( IDetailLayoutBuilder& DetailLayout )
 
 		OnExtendActorDetails.Broadcast(DetailLayout, FGetSelectedActors::CreateSP(this, &FActorDetails::GetSelectedActors));
 
-		if (!HideCategories.Contains(TEXT("Layers")))
+		/*if (!HideCategories.Contains(TEXT("Layers")))
 		{
 			AddLayersCategory(DetailLayout);
-		}
+		}*/
 
 		//AddComponentsCategory( DetailLayout );
 	}
diff --git a/Engine/Source/Editor/EditorStyle/Private/SlateEditorStyle.cpp b/Engine/Source/Editor/EditorStyle/Private/SlateEditorStyle.cpp
--- a/Engine/Source/Editor/EditorStyle/Private/SlateEditorStyle.cpp
+++ b/Engine/Source/Editor/EditorStyle/Private/SlateEditorStyle.cpp
@@ -3660,9 +3660,6 @@ void FSlateEditorStyle::FStyle::SetupLevelEditorStyle()
 		Set( "LevelEditor.Recompile", new IMAGE_BRUSH( "Icons/icon_compile_40x", Icon40x40 ) );
 		Set( "LevelEditor.Recompile.Small", new IMAGE_BRUSH( "Icons/icon_compile_40x", Icon20x20 ) );
 
-		Set( "LevelEditor.ViewOptions", new IMAGE_BRUSH( "Icons/icon_view_40x", Icon40x40 ) );
-		Set( "LevelEditor.ViewOptions.Small", new IMAGE_BRUSH( "Icons/icon_view_40x", Icon20x20 ) );
-
 		Set( "LevelEditor.GameSettings", new IMAGE_BRUSH( "Icons/icon_game_settings_40x", Icon40x40 ) );
 		Set( "LevelEditor.GameSettings.Small", new IMAGE_BRUSH( "Icons/icon_game_settings_40x", Icon20x20 ) );
 
diff --git a/Engine/Source/Editor/LevelEditor/Private/LevelEditor.cpp b/Engine/Source/Editor/LevelEditor/Private/LevelEditor.cpp
--- a/Engine/Source/Editor/LevelEditor/Private/LevelEditor.cpp
+++ b/Engine/Source/Editor/LevelEditor/Private/LevelEditor.cpp
@@ -1412,13 +1412,6 @@ void FLevelEditorModule::BindGlobalLevelEditorCommands()
 		FIsActionChecked::CreateStatic( &FLevelEditorActionCallbacks::OnIsStrictBoxSelectEnabled ) 
 		);
 
-	ActionList.MapAction(
-		Commands.DrawBrushMarkerPolys,
-		FExecuteAction::CreateStatic( &FLevelEditorActionCallbacks::OnDrawBrushMarkerPolys ),
-		FCanExecuteAction(),
-		FIsActionChecked::CreateStatic( &FLevelEditorActionCallbacks::OnIsDrawBrushMarkerPolysEnabled ) 
-		);
-
 	ActionList.MapAction(
 		Commands.OnlyLoadVisibleInPIE,
 		FExecuteAction::CreateStatic( &FLevelEditorActionCallbacks::OnToggleOnlyLoadVisibleInPIE ),
diff --git a/Engine/Source/Editor/LevelEditor/Private/LevelEditorActions.cpp b/Engine/Source/Editor/LevelEditor/Private/LevelEditorActions.cpp
--- a/Engine/Source/Editor/LevelEditor/Private/LevelEditorActions.cpp
+++ b/Engine/Source/Editor/LevelEditor/Private/LevelEditorActions.cpp
@@ -1941,16 +1941,6 @@ bool FLevelEditorActionCallbacks::OnIsStrictBoxSelectEnabled()
 	return GetDefault<ULevelEditorViewportSettings>()->bStrictBoxSelection;
 }
 
-void FLevelEditorActionCallbacks::OnDrawBrushMarkerPolys()
-{
-	GEditor->Exec( GetWorld(), *FString::Printf( TEXT("MODE SHOWBRUSHMARKERPOLYS=%d"), !GEditor->bShowBrushMarkerPolys ? 1 : 0 ) );
-	GEditor->SaveConfig();
-}
-
-bool FLevelEditorActionCallbacks::OnIsDrawBrushMarkerPolysEnabled()
-{
-	return GEditor->bShowBrushMarkerPolys;
-}
 
 void FLevelEditorActionCallbacks::OnToggleOnlyLoadVisibleInPIE()
 {
@@ -2766,7 +2756,7 @@ void FLevelEditorCommands::RegisterCommands()
 	UI_COMMAND( AllowTranslucentSelection, "Allow Translucent Selection", "Allows translucent objects to be selected", EUserInterfaceActionType::ToggleButton, FInputGesture(EKeys::T) );
 	UI_COMMAND( AllowGroupSelection, "Allow Group Selection", "Allows actor groups to be selected", EUserInterfaceActionType::ToggleButton, FInputGesture(EModifierKey::Control|EModifierKey::Shift, EKeys::G) );
 	UI_COMMAND( StrictBoxSelect, "Strict Box Selection", "When enabled an object must be entirely encompassed by the selection box when marquee box selecting", EUserInterfaceActionType::ToggleButton, FInputGesture() );
-	UI_COMMAND( DrawBrushMarkerPolys, "Draw Brush Marker Polys", "Draws semi-transparent polygons around a brush when selected", EUserInterfaceActionType::ToggleButton, FInputGesture() );
+	
 	UI_COMMAND( OnlyLoadVisibleInPIE, "Only Load Visible Levels in Game Preview", "If enabled, when game preview starts, only visible levels will be loaded", EUserInterfaceActionType::ToggleButton, FInputGesture() );
 	UI_COMMAND( ToggleSocketSnapping, "Enable Socket Snapping", "Enables or disables snapping to sockets", EUserInterfaceActionType::ToggleButton, FInputGesture() ); 
 	UI_COMMAND( ToggleParticleSystemLOD, "Enable Particle System LOD Switching", "If enabled particle systems will use distance LOD switching in perspective viewports", EUserInterfaceActionType::ToggleButton, FInputGesture() );
diff --git a/Engine/Source/Editor/LevelEditor/Private/LevelEditorToolBar.cpp b/Engine/Source/Editor/LevelEditor/Private/LevelEditorToolBar.cpp
--- a/Engine/Source/Editor/LevelEditor/Private/LevelEditorToolBar.cpp
+++ b/Engine/Source/Editor/LevelEditor/Private/LevelEditorToolBar.cpp
@@ -192,22 +192,13 @@ TSharedRef< SWidget > FLevelEditorToolBar::MakeLevelEditorToolBar( const TShared
 		ToolbarBuilder.AddComboButton(
 			FUIAction(),
 			FOnGetContent::CreateStatic(&FLevelEditorToolBar::GenerateQuickSettingsMenu, InCommandList),
-			LOCTEXT("QuickSettingsCombo", "Quick Settings"),
-			LOCTEXT("QuickSettingsCombo_ToolTip", "Quick level editor settings"),
-			FSlateIcon(FEditorStyle::GetStyleSetName(), "LevelEditor.ViewOptions"),
+			LOCTEXT("QuickSettingsCombo", "Settings"),
+			LOCTEXT("QuickSettingsCombo_ToolTip", "Project and Editor settings"),
+			FSlateIcon(FEditorStyle::GetStyleSetName(), "LevelEditor.GameSettings"),
 			false,
 			"LevelToolbarQuickSettings"
 			);
 
-		ToolbarBuilder.AddComboButton(
-			FUIAction(),
-			FOnGetContent::CreateStatic( &FLevelEditorToolBar::GenerateOpenGameSettingsMenu, InCommandList ),
-			LOCTEXT( "OpenConfiguration_Label", "Game Settings" ),
-			LOCTEXT( "OpenConfiguration_ToolTip", "Open game settings windows" ),
-			FSlateIcon(FEditorStyle::GetStyleSetName(), "LevelEditor.GameSettings"),
-			false,
-			"LevelToolbarWorldSettings"
-			);
 	}
 	ToolbarBuilder.EndSection();
 
@@ -624,7 +615,22 @@ static void MakeScalabilityMenu( FMenuBuilder& MenuBuilder )
 	MenuBuilder.AddWidget(SNew(SScalabilitySettings), FText(), true);
 }
 
+static void MakePreviewSettingsMenu( FMenuBuilder& MenuBuilder )
+{
+#define LOCTEXT_NAMESPACE "LevelToolBarViewMenu"
 
+	MenuBuilder.BeginSection("LevelEditorPreview", LOCTEXT("PreviewHeading", "Previewing"));
+	{
+		MenuBuilder.AddMenuEntry(FLevelEditorCommands::Get().OnlyLoadVisibleInPIE);
+		MenuBuilder.AddMenuEntry(FLevelEditorCommands::Get().ToggleParticleSystemLOD);
+		MenuBuilder.AddMenuEntry(FLevelEditorCommands::Get().ToggleParticleSystemHelpers);
+		MenuBuilder.AddMenuEntry(FLevelEditorCommands::Get().ToggleFreezeParticleSimulation);
+		MenuBuilder.AddMenuEntry(FLevelEditorCommands::Get().ToggleLODViewLocking);
+		MenuBuilder.AddMenuEntry(FLevelEditorCommands::Get().LevelStreamingVolumePrevis);
+	}
+	MenuBuilder.EndSection();
+#undef LOCTEXT_NAMESPACE
+}
 
 TSharedRef< SWidget > FLevelEditorToolBar::GenerateQuickSettingsMenu( TSharedRef<FUICommandList> InCommandList )
 {
@@ -647,30 +653,39 @@ TSharedRef< SWidget > FLevelEditorToolBar::GenerateQuickSettingsMenu( TSharedRef
 	const bool bShouldCloseWindowAfterMenuSelection = true;
 	FMenuBuilder MenuBuilder( bShouldCloseWindowAfterMenuSelection, InCommandList, MenuExtender );
 
-	MenuBuilder.BeginSection("LevelEditorSelection", LOCTEXT("SelectionHeading","Selection") );
+	struct Local
 	{
-		MenuBuilder.AddMenuEntry( FLevelEditorCommands::Get().AllowTranslucentSelection );
-		MenuBuilder.AddMenuEntry( FLevelEditorCommands::Get().AllowGroupSelection );
-		MenuBuilder.AddMenuEntry( FLevelEditorCommands::Get().StrictBoxSelect );
+		static void OpenSettings(FName ContainerName, FName CategoryName, FName SectionName)
+		{
+			FModuleManager::LoadModuleChecked<ISettingsModule>("Settings").ShowViewer(ContainerName, CategoryName, SectionName);
+		}
+	};
+
+	MenuBuilder.BeginSection("ProjectSettingsSection", LOCTEXT("ProjectSettings","Game Specific Settings") );
+	{
+
+		MenuBuilder.AddMenuEntry(FLevelEditorCommands::Get().WorldProperties);
+
+		MenuBuilder.AddMenuEntry(
+			LOCTEXT("ProjectSettingsMenuLabel", "Project Settings..."),
+			LOCTEXT("ProjectSettingsMenuToolTip", "Change the settings of the currently loaded project"),
+			FSlateIcon(),
+			FUIAction(FExecuteAction::CreateStatic(&Local::OpenSettings, FName("Project"), FName("Project"), FName("General")))
+			);
 	}
 	MenuBuilder.EndSection();
 
-	MenuBuilder.BeginSection("LevelEditorEditing", LOCTEXT("EditingHeading", "Editing") );
+	MenuBuilder.BeginSection("LevelEditorSelection", LOCTEXT("SelectionHeading","Selection") );
 	{
+		MenuBuilder.AddMenuEntry( FLevelEditorCommands::Get().AllowTranslucentSelection );
+		MenuBuilder.AddMenuEntry( FLevelEditorCommands::Get().AllowGroupSelection );
+		MenuBuilder.AddMenuEntry( FLevelEditorCommands::Get().StrictBoxSelect );
 		MenuBuilder.AddMenuEntry( FLevelEditorCommands::Get().ShowTransformWidget );
-		MenuBuilder.AddMenuEntry( FLevelEditorCommands::Get().DrawBrushMarkerPolys );
 	}
 	MenuBuilder.EndSection();
 
-	MenuBuilder.BeginSection("LevelEditorPreview", LOCTEXT("PreviewHeading", "Previewing") );
+	MenuBuilder.BeginSection("LevelEditorScalability", LOCTEXT("ScalabilityHeading", "Scalability") );
 	{
-		MenuBuilder.AddMenuEntry( FLevelEditorCommands::Get().OnlyLoadVisibleInPIE );
-		MenuBuilder.AddMenuEntry( FLevelEditorCommands::Get().ToggleParticleSystemLOD );
-		MenuBuilder.AddMenuEntry( FLevelEditorCommands::Get().ToggleParticleSystemHelpers );
-		MenuBuilder.AddMenuEntry( FLevelEditorCommands::Get().ToggleFreezeParticleSimulation );
-		MenuBuilder.AddMenuEntry( FLevelEditorCommands::Get().ToggleLODViewLocking );
-		MenuBuilder.AddMenuEntry( FLevelEditorCommands::Get().LevelStreamingVolumePrevis );
-
 		MenuBuilder.AddSubMenu(
 			LOCTEXT( "ScalabilitySubMenu", "Engine Scalability Settings" ),
 			LOCTEXT( "ScalabilitySubMenu_ToolTip", "Open the engine scalability settings" ),
@@ -713,7 +728,7 @@ TSharedRef< SWidget > FLevelEditorToolBar::GenerateQuickSettingsMenu( TSharedRef
 	}
 	MenuBuilder.EndSection();
 
-	MenuBuilder.BeginSection("LevelEditorActorSnap", LOCTEXT("ActorSnapHeading","Actor Snap") );
+	MenuBuilder.BeginSection( "Snapping", LOCTEXT("SnappingHeading","Snapping") );
 	{
 		MenuBuilder.AddMenuEntry( FLevelEditorCommands::Get().EnableActorSnap );
 		TSharedRef<SWidget> SnapItem = 
@@ -729,11 +744,7 @@ TSharedRef< SWidget > FLevelEditorToolBar::GenerateQuickSettingsMenu( TSharedRef
 	          +SHorizontalBox::Slot()
 	          .FillWidth(0.1f);
 		MenuBuilder.AddWidget(SnapItem, LOCTEXT("ActorSnapLabel","Distance"));
-	}
-	MenuBuilder.EndSection();
 
-	MenuBuilder.BeginSection( "Snapping", LOCTEXT("SnappingHeading","Snapping") );
-	{
 		MenuBuilder.AddMenuEntry( FLevelEditorCommands::Get().ToggleSocketSnapping );
 		MenuBuilder.AddMenuEntry( FLevelEditorCommands::Get().EnableVertexSnap );
 	}
@@ -742,6 +753,8 @@ TSharedRef< SWidget > FLevelEditorToolBar::GenerateQuickSettingsMenu( TSharedRef
 	MenuBuilder.BeginSection("LevelEditorViewport", LOCTEXT("ViewportHeading", "Viewport") );
 	{
 		MenuBuilder.AddMenuEntry( FLevelEditorCommands::Get().ToggleHideViewportUI );
+
+		MenuBuilder.AddSubMenu( LOCTEXT("PreviewMenu", "Previewing"), LOCTEXT("PreviewMenuTooltip","Game Preview Settings"), FNewMenuDelegate::CreateStatic( &MakePreviewSettingsMenu ) );
 	}
 	MenuBuilder.EndSection();
 
@@ -750,34 +763,6 @@ TSharedRef< SWidget > FLevelEditorToolBar::GenerateQuickSettingsMenu( TSharedRef
 	return MenuBuilder.MakeWidget();
 }
 
-TSharedRef< SWidget > FLevelEditorToolBar::GenerateOpenGameSettingsMenu( TSharedRef<FUICommandList> InCommandList )
-{
-#define LOCTEXT_NAMESPACE "OpenGameSettingsMenu"
-	const bool bShouldCloseWindowAfterMenuSelection = true;
-	FMenuBuilder MenuBuilder( bShouldCloseWindowAfterMenuSelection, InCommandList );
-
-	MenuBuilder.AddMenuEntry( FLevelEditorCommands::Get().WorldProperties );
-
-	struct Local
-	{
-		static void OpenSettings( FName ContainerName, FName CategoryName, FName SectionName )
-		{
-			FModuleManager::LoadModuleChecked<ISettingsModule>("Settings").ShowViewer(ContainerName, CategoryName, SectionName);
-		}
-	};
-
-	MenuBuilder.AddMenuEntry(
-		LOCTEXT("ProjectSettingsMenuLabel", "Project Settings..."),
-		LOCTEXT("ProjectSettingsMenuToolTip", "Change the settings of the currently loaded project"),
-		FSlateIcon(),
-		FUIAction(FExecuteAction::CreateStatic(&Local::OpenSettings, FName("Project"), FName("Project"), FName("General")))
-	);
-
-#undef LOCTEXT_NAMESPACE
-
-	return MenuBuilder.MakeWidget();
-}
-
 TSharedRef< SWidget > FLevelEditorToolBar::GenerateCompileMenuContent( TSharedRef<FUICommandList> InCommandList )
 {
 #define LOCTEXT_NAMESPACE "LevelToolBarViewMenu"
diff --git a/Engine/Source/Editor/LevelEditor/Private/LevelEditorToolBar.h b/Engine/Source/Editor/LevelEditor/Private/LevelEditorToolBar.h
--- a/Engine/Source/Editor/LevelEditor/Private/LevelEditorToolBar.h
+++ b/Engine/Source/Editor/LevelEditor/Private/LevelEditorToolBar.h
@@ -49,13 +49,6 @@ protected:
 	 */
 	static TSharedRef< SWidget > GenerateQuickSettingsMenu( TSharedRef<FUICommandList> InCommandList );
 
-	/**
-	 * Generates menu content for the settings display combo button drop down menu
-	 *
-	 * @return	Menu content widget
-	 */
-	static TSharedRef< SWidget > GenerateOpenGameSettingsMenu( TSharedRef<FUICommandList> InCommandList );
-
 	/**
 	 * Generates menu content for the compile combo button drop down menu
 	 *
diff --git a/Engine/Source/Editor/LevelEditor/Public/LevelEditorActions.h b/Engine/Source/Editor/LevelEditor/Public/LevelEditorActions.h
--- a/Engine/Source/Editor/LevelEditor/Public/LevelEditorActions.h
+++ b/Engine/Source/Editor/LevelEditor/Public/LevelEditorActions.h
@@ -520,7 +520,6 @@ public:
 	TSharedPtr< FUICommandInfo > AllowGroupSelection;
 
 	TSharedPtr< FUICommandInfo > StrictBoxSelect;
-	TSharedPtr< FUICommandInfo > DrawBrushMarkerPolys;
 	TSharedPtr< FUICommandInfo > OnlyLoadVisibleInPIE;
 
 	TSharedPtr< FUICommandInfo > ToggleSocketSnapping; 
