Added back missing "Draw Brush Polys" setting

diff --git a/Engine/Source/Editor/LevelEditor/Private/LevelEditor.cpp b/Engine/Source/Editor/LevelEditor/Private/LevelEditor.cpp
--- a/Engine/Source/Editor/LevelEditor/Private/LevelEditor.cpp
+++ b/Engine/Source/Editor/LevelEditor/Private/LevelEditor.cpp
@@ -1412,6 +1412,13 @@ void FLevelEditorModule::BindGlobalLevelEditorCommands()
 		FIsActionChecked::CreateStatic( &FLevelEditorActionCallbacks::OnIsStrictBoxSelectEnabled ) 
 		);
 
+	ActionList.MapAction(
+		Commands.DrawBrushMarkerPolys,
+		FExecuteAction::CreateStatic( &FLevelEditorActionCallbacks::OnDrawBrushMarkerPolys ),
+		FCanExecuteAction(),
+		FIsActionChecked::CreateStatic( &FLevelEditorActionCallbacks::OnIsDrawBrushMarkerPolysEnabled ) 
+		);
+
 	ActionList.MapAction(
 		Commands.OnlyLoadVisibleInPIE,
 		FExecuteAction::CreateStatic( &FLevelEditorActionCallbacks::OnToggleOnlyLoadVisibleInPIE ),
diff --git a/Engine/Source/Editor/LevelEditor/Private/LevelEditorActions.cpp b/Engine/Source/Editor/LevelEditor/Private/LevelEditorActions.cpp
--- a/Engine/Source/Editor/LevelEditor/Private/LevelEditorActions.cpp
+++ b/Engine/Source/Editor/LevelEditor/Private/LevelEditorActions.cpp
@@ -1941,6 +1941,16 @@ bool FLevelEditorActionCallbacks::OnIsStrictBoxSelectEnabled()
 	return GetDefault<ULevelEditorViewportSettings>()->bStrictBoxSelection;
 }
 
+void FLevelEditorActionCallbacks::OnDrawBrushMarkerPolys()
+{
+	GEditor->Exec( GetWorld(), *FString::Printf( TEXT("MODE SHOWBRUSHMARKERPOLYS=%d"), !GEditor->bShowBrushMarkerPolys ? 1 : 0 ) );
+	GEditor->SaveConfig();
+}
+
+bool FLevelEditorActionCallbacks::OnIsDrawBrushMarkerPolysEnabled()
+{
+	return GEditor->bShowBrushMarkerPolys;
+}
 
 void FLevelEditorActionCallbacks::OnToggleOnlyLoadVisibleInPIE()
 {
@@ -2756,7 +2766,7 @@ void FLevelEditorCommands::RegisterCommands()
 	UI_COMMAND( AllowTranslucentSelection, "Allow Translucent Selection", "Allows translucent objects to be selected", EUserInterfaceActionType::ToggleButton, FInputGesture(EKeys::T) );
 	UI_COMMAND( AllowGroupSelection, "Allow Group Selection", "Allows actor groups to be selected", EUserInterfaceActionType::ToggleButton, FInputGesture(EModifierKey::Control|EModifierKey::Shift, EKeys::G) );
 	UI_COMMAND( StrictBoxSelect, "Strict Box Selection", "When enabled an object must be entirely encompassed by the selection box when marquee box selecting", EUserInterfaceActionType::ToggleButton, FInputGesture() );
-	
+	UI_COMMAND( DrawBrushMarkerPolys, "Draw Brush Polys", "Draws semi-transparent polygons around a brush when selected", EUserInterfaceActionType::ToggleButton, FInputGesture() );
 	UI_COMMAND( OnlyLoadVisibleInPIE, "Only Load Visible Levels in Game Preview", "If enabled, when game preview starts, only visible levels will be loaded", EUserInterfaceActionType::ToggleButton, FInputGesture() );
 	UI_COMMAND( ToggleSocketSnapping, "Enable Socket Snapping", "Enables or disables snapping to sockets", EUserInterfaceActionType::ToggleButton, FInputGesture() ); 
 	UI_COMMAND( ToggleParticleSystemLOD, "Enable Particle System LOD Switching", "If enabled particle systems will use distance LOD switching in perspective viewports", EUserInterfaceActionType::ToggleButton, FInputGesture() );
diff --git a/Engine/Source/Editor/LevelEditor/Private/LevelEditorToolBar.cpp b/Engine/Source/Editor/LevelEditor/Private/LevelEditorToolBar.cpp
--- a/Engine/Source/Editor/LevelEditor/Private/LevelEditorToolBar.cpp
+++ b/Engine/Source/Editor/LevelEditor/Private/LevelEditorToolBar.cpp
@@ -621,6 +621,7 @@ static void MakePreviewSettingsMenu( FMenuBuilder& MenuBuilder )
 
 	MenuBuilder.BeginSection("LevelEditorPreview", LOCTEXT("PreviewHeading", "Previewing"));
 	{
+		MenuBuilder.AddMenuEntry(FLevelEditorCommands::Get().DrawBrushMarkerPolys);
 		MenuBuilder.AddMenuEntry(FLevelEditorCommands::Get().OnlyLoadVisibleInPIE);
 		MenuBuilder.AddMenuEntry(FLevelEditorCommands::Get().ToggleParticleSystemLOD);
 		MenuBuilder.AddMenuEntry(FLevelEditorCommands::Get().ToggleParticleSystemHelpers);
diff --git a/Engine/Source/Editor/LevelEditor/Public/LevelEditorActions.h b/Engine/Source/Editor/LevelEditor/Public/LevelEditorActions.h
--- a/Engine/Source/Editor/LevelEditor/Public/LevelEditorActions.h
+++ b/Engine/Source/Editor/LevelEditor/Public/LevelEditorActions.h
@@ -520,6 +520,7 @@ public:
 	TSharedPtr< FUICommandInfo > AllowGroupSelection;
 
 	TSharedPtr< FUICommandInfo > StrictBoxSelect;
+	TSharedPtr< FUICommandInfo > DrawBrushMarkerPolys;
 	TSharedPtr< FUICommandInfo > OnlyLoadVisibleInPIE;
 
 	TSharedPtr< FUICommandInfo > ToggleSocketSnapping; 
