In the Blueprints menu in the Editor toolbar, fixed an issue attempting to edit the player controller asking the user to create a Hud or open the Hud for edit.

#ttp 345332 - Clicking edit player controller opens the HUD blueprint

diff --git a/Engine/Source/Editor/LevelEditor/Private/LevelEditorToolBar.cpp b/Engine/Source/Editor/LevelEditor/Private/LevelEditorToolBar.cpp
--- a/Engine/Source/Editor/LevelEditor/Private/LevelEditorToolBar.cpp
+++ b/Engine/Source/Editor/LevelEditor/Private/LevelEditorToolBar.cpp
@@ -954,7 +954,7 @@ TSharedRef< SWidget > FLevelEditorToolBar::GenerateOpenBlueprintMenuContent( TSh
 
 			// Player Controller
 			LevelEditorActionHelpers::FBlueprintMenuSettings PlayerControllerMenuSettings;
-			PlayerControllerMenuSettings.EditCommand = FLevelEditorCommands::Get().OpenHUDBlueprint;
+			PlayerControllerMenuSettings.EditCommand = FLevelEditorCommands::Get().OpenPlayerControllerBlueprint;
 			PlayerControllerMenuSettings.OnCreateClassPicked = FOnClassPicked::CreateStatic( &FLevelEditorToolBar::OnCreatePlayerControllerClassPicked, InLevelEditor );;
 			PlayerControllerMenuSettings.OnSelectClassPicked = FOnClassPicked::CreateStatic( &FLevelEditorToolBar::OnSelectPlayerControllerClassPicked, InLevelEditor );;
 			PlayerControllerMenuSettings.CurrentClass = GetActivePlayerControllerClass(InLevelEditor);
