Fix viewports drawing when their tabs are not active.

diff --git a/Engine/Source/Editor/Kismet/Private/SSCSEditorViewport.cpp b/Engine/Source/Editor/Kismet/Private/SSCSEditorViewport.cpp
--- a/Engine/Source/Editor/Kismet/Private/SSCSEditorViewport.cpp
+++ b/Engine/Source/Editor/Kismet/Private/SSCSEditorViewport.cpp
@@ -121,7 +121,7 @@ SSCSEditorViewport::~SSCSEditorViewport()
 bool SSCSEditorViewport::IsVisible() const
 {
 	// We consider the viewport to be visible if the reference is valid
-	return ViewportWidget.IsValid();
+	return ViewportWidget.IsValid() && SEditorViewport::IsVisible();
 }
 
 EVisibility SSCSEditorViewport::GetWidgetVisibility() const
@@ -239,7 +239,7 @@ bool SSCSEditorViewport::GetIsSimulateEnabled()
 
 void SSCSEditorViewport::Tick(const FGeometry& AllottedGeometry, const double InCurrentTime, const float InDeltaTime)
 {
-	SCompoundWidget::Tick(AllottedGeometry, InCurrentTime, InDeltaTime);
+	SEditorViewport::Tick(AllottedGeometry, InCurrentTime, InDeltaTime);
 
 	// If the preview scene is no longer valid (i.e. all actors have destroyed themselves), then attempt to recreate the scene. This way we can "loop" certain "finite" Blueprints that might destroy themselves.
 	if(ViewportClient.IsValid() && (bPreviewNeedsUpdating || !ViewportClient->IsPreviewSceneValid()))
diff --git a/Engine/Source/Editor/LevelEditor/Private/SLevelViewport.cpp b/Engine/Source/Editor/LevelEditor/Private/SLevelViewport.cpp
--- a/Engine/Source/Editor/LevelEditor/Private/SLevelViewport.cpp
+++ b/Engine/Source/Editor/LevelEditor/Private/SLevelViewport.cpp
@@ -106,7 +106,7 @@ void SLevelViewport::HandleViewportSettingChanged(FName PropertyName)
 bool SLevelViewport::IsVisible() const
 {
 	// The viewport is visible if we don't have a parent layout (likely a floating window) or this viewport is visible in the parent layout
-	return ViewportWidget.IsValid() && ( !ParentLayout.IsValid() || ParentLayout.Pin()->IsLevelViewportVisible( *this ) );
+	return ViewportWidget.IsValid() && (!ParentLayout.IsValid() || ParentLayout.Pin()->IsLevelViewportVisible(*this)) && SEditorViewport::IsVisible();
 }
 
 void SLevelViewport::Construct(const FArguments& InArgs)
@@ -786,8 +786,12 @@ FReply SLevelViewport::OnDrop( const FGeometry& MyGeometry, const FDragDropEvent
 
 void SLevelViewport::Tick( const FGeometry& AllottedGeometry, const double InCurrentTime, const float InDeltaTime )
 {
+	SEditorViewport::Tick(AllottedGeometry, InCurrentTime, InDeltaTime);
+
 	const bool bContainsFocus = HasFocusedDescendants();
 
+	LastTickTime = FPlatformTime::Seconds();
+
 	// When we have focus we update the 'Allow Throttling' option in slate to be disabled so that interactions in the
 	// viewport with Slate widgets that are part of the game, don't throttle.
 	if ( bPIEContainsFocus != bContainsFocus )
diff --git a/Engine/Source/Editor/MaterialEditor/Private/SMaterialEditorViewport.cpp b/Engine/Source/Editor/MaterialEditor/Private/SMaterialEditorViewport.cpp
--- a/Engine/Source/Editor/MaterialEditor/Private/SMaterialEditorViewport.cpp
+++ b/Engine/Source/Editor/MaterialEditor/Private/SMaterialEditorViewport.cpp
@@ -178,7 +178,7 @@ void SMaterialEditorViewport::RefreshViewport()
 
 void SMaterialEditorViewport::Tick(const FGeometry& AllottedGeometry, const double InCurrentTime, const float InDeltaTime)
 {
-	SCompoundWidget::Tick(AllottedGeometry, InCurrentTime, InDeltaTime);
+	SEditorViewport::Tick(AllottedGeometry, InCurrentTime, InDeltaTime);
 }
 
 bool SMaterialEditorViewport::SetPreviewMesh(UStaticMesh* InStaticMesh, USkeletalMesh* InSkeletalMesh)
@@ -292,7 +292,7 @@ void SMaterialEditorViewport::OnAddedToTab( const TSharedRef<SDockTab>& OwnerTab
 
 bool SMaterialEditorViewport::IsVisible() const
 {
-	return ViewportWidget.IsValid() && (!ParentTab.IsValid() || ParentTab.Pin()->IsForeground());
+	return ViewportWidget.IsValid() && (!ParentTab.IsValid() || ParentTab.Pin()->IsForeground()) && SEditorViewport::IsVisible() ;
 }
 
 void SMaterialEditorViewport::BindCommands()
diff --git a/Engine/Source/Editor/UnrealEd/Private/Editor.cpp b/Engine/Source/Editor/UnrealEd/Private/Editor.cpp
--- a/Engine/Source/Editor/UnrealEd/Private/Editor.cpp
+++ b/Engine/Source/Editor/UnrealEd/Private/Editor.cpp
@@ -1347,7 +1347,7 @@ void UEditorEngine::Tick( float DeltaSeconds, bool bIdleMode )
 					continue;
 				}
 
-				if (!ViewportClient->IsLevelEditorClient() || ViewportClient->IsVisible())
+				if ( ViewportClient->IsVisible() )
 				{
 					// Only update ortho viewports if that mode is turned on, the viewport client we are about to update is orthographic and the current editing viewport is orthographic and tracking mouse movement.
 					bUpdateLinkedOrthoViewports = GetDefault<ULevelEditorViewportSettings>()->bUseLinkedOrthographicViewports && ViewportClient->IsOrtho() && GCurrentLevelEditingViewportClient && GCurrentLevelEditingViewportClient->IsOrtho() && GCurrentLevelEditingViewportClient->IsTracking();
diff --git a/Engine/Source/Editor/UnrealEd/Private/SEditorViewport.cpp b/Engine/Source/Editor/UnrealEd/Private/SEditorViewport.cpp
--- a/Engine/Source/Editor/UnrealEd/Private/SEditorViewport.cpp
+++ b/Engine/Source/Editor/UnrealEd/Private/SEditorViewport.cpp
@@ -10,6 +10,7 @@
 #define LOCTEXT_NAMESPACE "EditorViewport"
 
 SEditorViewport::SEditorViewport()
+	: LastTickTime(0)
 {
 }
 
@@ -98,6 +99,11 @@ FReply SEditorViewport::OnKeyboardFocusReceived( const FGeometry& MyGeometry, co
 	return FReply::Handled().SetKeyboardFocus( ViewportWidget.ToSharedRef(), InKeyboardFocusEvent.GetCause() );
 }
 
+void SEditorViewport::Tick( const FGeometry& AllottedGeometry, const double InCurrentTime, const float InDeltaTime )
+{
+	LastTickTime = FPlatformTime::Seconds();
+}
+
 void SEditorViewport::BindCommands()
 {
 	FUICommandList& CommandListRef = *CommandList;
@@ -368,6 +374,13 @@ bool SEditorViewport::IsRealtime() const
 	return Client->IsRealtime();
 }
 
+bool SEditorViewport::IsVisible() const
+{
+	const float VisibilityTimeThreshold = .25f;
+	// The viewport is visible if we don't have a parent layout (likely a floating window) or this viewport is visible in the parent layout
+	return FPlatformTime::Seconds() - LastTickTime <= VisibilityTimeThreshold;
+}
+
 void SEditorViewport::OnScreenCapture()
 {
 	Client->TakeScreenshot(Client->Viewport, true);
diff --git a/Engine/Source/Editor/UnrealEd/Public/SEditorViewport.h b/Engine/Source/Editor/UnrealEd/Public/SEditorViewport.h
--- a/Engine/Source/Editor/UnrealEd/Public/SEditorViewport.h
+++ b/Engine/Source/Editor/UnrealEd/Public/SEditorViewport.h
@@ -17,12 +17,16 @@ public:
 	virtual FReply OnKeyDown( const FGeometry& MyGeometry, const FKeyboardEvent& InKeyboardEvent ) override;
 	virtual bool SupportsKeyboardFocus() const override;
 	virtual FReply OnKeyboardFocusReceived( const FGeometry& MyGeometry, const FKeyboardFocusEvent& InKeyboardFocusEvent ) override;
-
+	virtual void Tick(const FGeometry& AllottedGeometry, const double InCurrentTime, const float InDeltaTime) override;
+	
 	/**
 	 * @return True if the viewport is being updated in realtime
 	 */
 	bool IsRealtime() const;
 
+	/** @return True if the viewport is currently visible */
+	virtual bool IsVisible() const;
+
 	/**
 	 * @return true if the specified coordinate system the active one active
 	 */
@@ -142,4 +146,6 @@ protected:
 	TSharedPtr<FEditorViewportClient> Client;
 	/** Commandlist used in the viewport (Maps commands to viewport specific actions) */
 	TSharedPtr<FUICommandList> CommandList;
+	/** The last time the viewport was ticked (for visibility determination) */
+	double LastTickTime;
 };
