More widget tags for tutorials
Removed unused function prototype

diff --git a/Engine/Source/Editor/IntroTutorials/Private/STutorialOverlay.h b/Engine/Source/Editor/IntroTutorials/Private/STutorialOverlay.h
--- a/Engine/Source/Editor/IntroTutorials/Private/STutorialOverlay.h
+++ b/Engine/Source/Editor/IntroTutorials/Private/STutorialOverlay.h
@@ -55,9 +55,6 @@ private:
 	/* Do any interaction stuff for a widget - open browser, scroll to node etc */
 	void PerformWidgetInteractions(const FTutorialWidgetContent &WidgetContent);
 
-	/* Creates a map of broswers required by the various widget types. */
-	void AddTabInfo();
-
 private:
 	/** Reference to the canvas we use to position our content widgets */
 	TSharedPtr<SCanvas> OverlayCanvas;
@@ -80,6 +77,4 @@ private:
 	/** Delegate used to inform widgets of the current window size, so they can auto-adjust layout */
 	FOnCacheWindowSize OnCacheWindowSize;
 
-	/* A map of which tab requires invoking based on the tag name of the asset */
-	TMap< FString, FString >	BrowserTabMap;
 };
\ No newline at end of file
diff --git a/Engine/Source/Editor/LandscapeEditor/Private/LandscapeEditorDetailCustomization_NewLandscape.cpp b/Engine/Source/Editor/LandscapeEditor/Private/LandscapeEditorDetailCustomization_NewLandscape.cpp
--- a/Engine/Source/Editor/LandscapeEditor/Private/LandscapeEditorDetailCustomization_NewLandscape.cpp
+++ b/Engine/Source/Editor/LandscapeEditor/Private/LandscapeEditorDetailCustomization_NewLandscape.cpp
@@ -29,6 +29,7 @@
 
 #include "Landscape/Landscape.h"
 #include "Landscape/LandscapeLayerInfoObject.h"
+#include "TutorialMetaData.h"
 
 #define LOCTEXT_NAMESPACE "LandscapeEditor.NewLandscape"
 
@@ -440,6 +441,7 @@ void FLandscapeEditorDetailCustomization_NewLandscape::CustomizeDetails(IDetailL
 			SNew(SButton)
 			.Visibility_Static(&GetVisibilityOnlyInNewLandscapeMode, ENewLandscapePreviewMode::NewLandscape)
 			.Text(LOCTEXT("FillWorld", "Fill World"))
+			.AddMetaData<FTutorialMetaData>(FTutorialMetaData(TEXT("FillWorldButton"), TEXT("LevelEditorToolBox")))
 			.OnClicked(this, &FLandscapeEditorDetailCustomization_NewLandscape::OnFillWorldButtonClicked)
 		]
 		+ SHorizontalBox::Slot()
@@ -448,6 +450,7 @@ void FLandscapeEditorDetailCustomization_NewLandscape::CustomizeDetails(IDetailL
 			SNew(SButton)
 			.Visibility_Static(&GetVisibilityOnlyInNewLandscapeMode, ENewLandscapePreviewMode::ImportLandscape)
 			.Text(LOCTEXT("FitToData", "Fit To Data"))
+			.AddMetaData<FTagMetaData>(TEXT("ImportButton"))
 			.OnClicked(this, &FLandscapeEditorDetailCustomization_NewLandscape::OnFitImportDataButtonClicked)
 		]
 		+ SHorizontalBox::Slot()
@@ -460,6 +463,7 @@ void FLandscapeEditorDetailCustomization_NewLandscape::CustomizeDetails(IDetailL
 			SNew(SButton)
 			.Visibility_Static(&GetVisibilityOnlyInNewLandscapeMode, ENewLandscapePreviewMode::NewLandscape)
 			.Text(LOCTEXT("Create", "Create"))
+			.AddMetaData<FTutorialMetaData>(FTutorialMetaData(TEXT("CreateButton"), TEXT("LevelEditorToolBox")))
 			.OnClicked(this, &FLandscapeEditorDetailCustomization_NewLandscape::OnCreateButtonClicked)
 		]
 		+ SHorizontalBox::Slot()
diff --git a/Engine/Source/Editor/LevelEditor/Private/SLevelViewportToolBar.cpp b/Engine/Source/Editor/LevelEditor/Private/SLevelViewportToolBar.cpp
--- a/Engine/Source/Editor/LevelEditor/Private/SLevelViewportToolBar.cpp
+++ b/Engine/Source/Editor/LevelEditor/Private/SLevelViewportToolBar.cpp
@@ -133,6 +133,7 @@ void SLevelViewportToolBar::Construct( const FArguments& InArgs )
 					.ParentToolBar( SharedThis( this ) )
 					.Cursor( EMouseCursor::Default )
 					.Image( "EditorViewportToolBar.MenuDropdown" )
+					.AddMetaData<FTagMetaData>(FTagMetaData(TEXT("EditorViewportToolBar.MenuDropdown")))
 					.OnGetMenuContent( this, &SLevelViewportToolBar::GenerateOptionsMenu )
 				]
 				+ SHorizontalBox::Slot()
@@ -144,6 +145,7 @@ void SLevelViewportToolBar::Construct( const FArguments& InArgs )
 					.Cursor( EMouseCursor::Default )
 					.Label( this, &SLevelViewportToolBar::GetCameraMenuLabel )
 					.LabelIcon( this, &SLevelViewportToolBar::GetCameraMenuLabelIcon )
+					.AddMetaData<FTagMetaData>(FTagMetaData(TEXT("EditorViewportToolBar.CameraMenu")))
 					.OnGetMenuContent( this, &SLevelViewportToolBar::GenerateCameraMenu ) 
 				]
 				+ SHorizontalBox::Slot()
@@ -153,6 +155,7 @@ void SLevelViewportToolBar::Construct( const FArguments& InArgs )
 					SNew( SLevelEditorViewportViewMenu, ViewportRef, SharedThis(this) )
 					.Cursor( EMouseCursor::Default )
 					.MenuExtenders( GetViewMenuExtender() )
+					.AddMetaData<FTagMetaData>(FTagMetaData(TEXT("ViewMenuButton")))
 				]
 				+ SHorizontalBox::Slot()
 				.AutoWidth()
@@ -162,6 +165,7 @@ void SLevelViewportToolBar::Construct( const FArguments& InArgs )
 					.Label( LOCTEXT("ShowMenuTitle", "Show") )
 					.Cursor( EMouseCursor::Default )
 					.ParentToolBar( SharedThis( this ) )
+					.AddMetaData<FTagMetaData>(FTagMetaData(TEXT("EditorViewportToolBar.ShowMenu")))
 					.OnGetMenuContent( this, &SLevelViewportToolBar::GenerateShowMenu ) 
 				]
 				+ SHorizontalBox::Slot()
@@ -184,7 +188,7 @@ void SLevelViewportToolBar::Construct( const FArguments& InArgs )
 					SNew(STransformViewportToolBar)
 					.Viewport(ViewportRef)
 					.CommandList(LevelEditorModule.GetGlobalLevelEditorActions())
-					.Extenders(LevelEditorModule.GetToolBarExtensibilityManager()->GetAllExtenders())
+					.Extenders(LevelEditorModule.GetToolBarExtensibilityManager()->GetAllExtenders())					
 					.Visibility(ViewportRef, &SLevelViewport::GetTransformToolbarVisibility)
 				]
 				+ SHorizontalBox::Slot()
diff --git a/Engine/Source/Editor/UnrealEd/Private/STransformViewportToolbar.cpp b/Engine/Source/Editor/UnrealEd/Private/STransformViewportToolbar.cpp
--- a/Engine/Source/Editor/UnrealEd/Private/STransformViewportToolbar.cpp
+++ b/Engine/Source/Editor/UnrealEd/Private/STransformViewportToolbar.cpp
@@ -659,6 +659,7 @@ TSharedRef< SWidget > STransformViewportToolBar::MakeTransformToolBar( const TSh
 									.ToolTipText(LOCTEXT("CameraSpeed_ToolTip","Camera Speed"))
 									.Icon(FSlateIcon(FEditorStyle::GetStyleSetName(), "EditorViewport.CamSpeedSetting"))
 									.ParentToolBar(SharedThis(this))
+									.AddMetaData<FTagMetaData>(FTagMetaData(TEXT("CameraSpeedButton")))
 									);
 	}
 	ToolbarBuilder.EndBlockGroup();
