Fix crash opening level viewport context menu

diff --git a/Engine/Source/Editor/LevelEditor/Private/SLevelEditor.cpp b/Engine/Source/Editor/LevelEditor/Private/SLevelEditor.cpp
--- a/Engine/Source/Editor/LevelEditor/Private/SLevelEditor.cpp
+++ b/Engine/Source/Editor/LevelEditor/Private/SLevelEditor.cpp
@@ -350,7 +350,7 @@ TSharedPtr<SLevelViewport> SLevelEditor::GetActiveViewport()
 					{
 						const TSharedPtr< SLevelViewport >& Viewport = (*LevelViewports)[ ViewportIndex ];
 
-						if( Viewport->IsVisible() )
+						if( Viewport->IsInForegroundTab() )
 						{
 							if( &Viewport->GetLevelViewportClient() == GCurrentLevelEditingViewportClient )
 							{
diff --git a/Engine/Source/Editor/LevelEditor/Private/SLevelViewport.cpp b/Engine/Source/Editor/LevelEditor/Private/SLevelViewport.cpp
--- a/Engine/Source/Editor/LevelEditor/Private/SLevelViewport.cpp
+++ b/Engine/Source/Editor/LevelEditor/Private/SLevelViewport.cpp
@@ -106,7 +106,12 @@ void SLevelViewport::HandleViewportSettingChanged(FName PropertyName)
 bool SLevelViewport::IsVisible() const
 {
 	// The viewport is visible if we don't have a parent layout (likely a floating window) or this viewport is visible in the parent layout
-	return ViewportWidget.IsValid() && (!ParentLayout.IsValid() || ParentLayout.Pin()->IsLevelViewportVisible(*this)) && SEditorViewport::IsVisible();
+	return IsInForegroundTab() && SEditorViewport::IsVisible();
+}
+
+bool SLevelViewport::IsInForegroundTab() const
+{
+	return ViewportWidget.IsValid() && (!ParentLayout.IsValid() || ParentLayout.Pin()->IsLevelViewportVisible(*this));
 }
 
 void SLevelViewport::Construct(const FArguments& InArgs)
diff --git a/Engine/Source/Editor/LevelEditor/Public/SLevelViewport.h b/Engine/Source/Editor/LevelEditor/Public/SLevelViewport.h
--- a/Engine/Source/Editor/LevelEditor/Public/SLevelViewport.h
+++ b/Engine/Source/Editor/LevelEditor/Public/SLevelViewport.h
@@ -49,7 +49,10 @@ public:
 	/**
 	 * @return true if the viewport is visible. false otherwise                  
 	 */
-	bool IsVisible() const;
+	virtual bool IsVisible() const;
+
+	/** @return true if this viewport is in a foregrounded tab */
+	bool IsInForegroundTab() const;
 
 	/**
 	 * @return The editor client for this viewport
