More widget metadata for tutorials

diff --git a/Engine/Source/Editor/LandscapeEditor/Private/LandscapeEditorDetailWidgets.cpp b/Engine/Source/Editor/LandscapeEditor/Private/LandscapeEditorDetailWidgets.cpp
--- a/Engine/Source/Editor/LandscapeEditor/Private/LandscapeEditorDetailWidgets.cpp
+++ b/Engine/Source/Editor/LandscapeEditor/Private/LandscapeEditorDetailWidgets.cpp
@@ -153,7 +153,7 @@ void SToolSelector::BuildMultiBlockWidget(const ISlateStyle* StyleSet, const FNa
 	// Create the content for our button
 	TSharedRef< SWidget > ButtonContent =
 		SNew( SVerticalBox )
-
+		.AddMetaData<FTagMetaData>(FTagMetaData(TEXT("LandscapeToolButton")))
 		// Icon image
 		+ SVerticalBox::Slot()
 		.AutoHeight()
diff --git a/Engine/Source/Editor/LevelEditor/Private/SSurfaceProperties.cpp b/Engine/Source/Editor/LevelEditor/Private/SSurfaceProperties.cpp
--- a/Engine/Source/Editor/LevelEditor/Private/SSurfaceProperties.cpp
+++ b/Engine/Source/Editor/LevelEditor/Private/SSurfaceProperties.cpp
@@ -46,6 +46,7 @@ void SSurfaceProperties::Construct( const FArguments& InArgs )
 			.BorderBackgroundColor(BorderColor)
 			.BorderImage(BorderStyle)
 			.Padding(10)
+			.AddMetaData<FTagMetaData>(TEXT("DetailsView.SelectedSurfaces"))
 			[
 				SNew(STextBlock)
 				.Text(this, &SSurfaceProperties::GetSelectedSurfacesDescription)
@@ -60,6 +61,7 @@ void SSurfaceProperties::Construct( const FArguments& InArgs )
 			.BorderBackgroundColor( BorderColor )
 			.BorderImage( BorderStyle )
 			.Padding(10)
+			.AddMetaData<FTagMetaData>(TEXT("DetailsView.TexturePan"))
 			[
 				ConstructTexturePan()
 			]
@@ -78,6 +80,7 @@ void SSurfaceProperties::Construct( const FArguments& InArgs )
 				.BorderBackgroundColor( BorderColor )
 				.BorderImage( BorderStyle )
 				.Padding(10)
+				.AddMetaData<FTagMetaData>(TEXT("DetailsView.TextureRotate"))
 				[
 					ConstructTextureRotate()
 				]
@@ -89,6 +92,7 @@ void SSurfaceProperties::Construct( const FArguments& InArgs )
 				.BorderBackgroundColor( BorderColor )
 				.BorderImage( BorderStyle )
 				.Padding(10)
+				.AddMetaData<FTagMetaData>(TEXT("DetailsView.TextureFlip"))
 				[
 					ConstructTextureFlip()
 				]
@@ -103,6 +107,7 @@ void SSurfaceProperties::Construct( const FArguments& InArgs )
 			.BorderBackgroundColor( BorderColor )
 			.BorderImage( BorderStyle )
 			.Padding(10)
+			.AddMetaData<FTagMetaData>(TEXT("DetailsView.TextureScale"))
 			[
 				ConstructTextureScale()
 			]
@@ -116,6 +121,7 @@ void SSurfaceProperties::Construct( const FArguments& InArgs )
 			.BorderBackgroundColor( BorderColor )
 			.BorderImage( BorderStyle )
 			.Padding(10)
+			.AddMetaData<FTagMetaData>(TEXT("DetailsView.ConstructLighting"))
 			[
 				ConstructLighting()
 			]
diff --git a/Engine/Source/Editor/PropertyEditor/Private/DetailItemNode.cpp b/Engine/Source/Editor/PropertyEditor/Private/DetailItemNode.cpp
--- a/Engine/Source/Editor/PropertyEditor/Private/DetailItemNode.cpp
+++ b/Engine/Source/Editor/PropertyEditor/Private/DetailItemNode.cpp
@@ -141,7 +141,13 @@ TSharedRef< ITableRow > FDetailItemNode::GenerateNodeWidget( const TSharedRef<ST
 	FGraphNodeMetaData TagMeta(TEXT("DetailRowItem"));
 	if (Customization.IsValidCustomization() && Customization.GetPropertyNode().IsValid() )
 	{
-		TagMeta.FriendlyName = Customization.GetPropertyNode()->GetDisplayName(); 
+		TagMeta.FriendlyName = Customization.GetPropertyNode()->GetDisplayName();
+	}
+	if (ParentCategory.IsValid())
+	{
+		FString Tag = FString::Printf(TEXT("%s_%s"), *TagMeta.FriendlyName, *ParentCategory.Pin()->GetCategoryPathName());
+		TagMeta.Tag = *Tag;
+		TagMeta.TabTypeToOpen = TEXT("LevelEditorSelectionDetails");
 	}
 	if( Customization.HasPropertyNode() && Customization.GetPropertyNode()->AsCategoryNode() )
 	{
