Reset show flags and some views when a new level is opened

diff --git a/Engine/Source/Editor/LevelEditor/Private/SLevelViewport.cpp b/Engine/Source/Editor/LevelEditor/Private/SLevelViewport.cpp
--- a/Engine/Source/Editor/LevelEditor/Private/SLevelViewport.cpp
+++ b/Engine/Source/Editor/LevelEditor/Private/SLevelViewport.cpp
@@ -920,6 +920,8 @@ void SLevelViewport::OnMapChanged( UWorld* World, EMapChangeType::Type MapChange
 			{
 				World->EditorViews[LevelViewportClient->ViewportType].CamOrthoZoom = DEFAULT_ORTHOZOOM;
 			}
+	
+			ResetNewLevelViewFlags();
 
 			LevelViewportClient->SetInitialViewTransform(
 				World->EditorViews[LevelViewportClient->ViewportType].CamPosition,
@@ -937,6 +939,9 @@ void SLevelViewport::OnMapChanged( UWorld* World, EMapChangeType::Type MapChange
 		}
 		else if( MapChangeType == EMapChangeType::NewMap )
 		{
+		
+			ResetNewLevelViewFlags();
+
 			LevelViewportClient->ResetViewForNewMap();
 		}
 		World->EditorViews[LevelViewportClient->ViewportType].CamUpdated = false;
@@ -3369,6 +3374,15 @@ void SLevelViewport::HideMouseCaptureLabel()
 	PIEOverlaySlotIndex = 0;
 }
 
+void SLevelViewport::ResetNewLevelViewFlags()
+{
+	ChangeExposureSetting(FEditorViewportCommands::AutoExposureRadioID);
+
+	LevelViewportClient->SetViewMode(VMI_Lit);
+
+	OnUseDefaultShowFlags();
+}
+
 void SLevelViewport::EndPlayInEditorSession()
 {
 	check( HasPlayInEditorViewport() );
diff --git a/Engine/Source/Editor/LevelEditor/Public/SLevelViewport.h b/Engine/Source/Editor/LevelEditor/Public/SLevelViewport.h
--- a/Engine/Source/Editor/LevelEditor/Public/SLevelViewport.h
+++ b/Engine/Source/Editor/LevelEditor/Public/SLevelViewport.h
@@ -648,6 +648,8 @@ private:
 	/** Hide the mouse capture label */
 	void HideMouseCaptureLabel();
 
+	/** Resets view flags when a new level is created or opened */
+	void ResetNewLevelViewFlags();
 private:
 	/** Returns the DPI scaler that should be used for game UI in the viewport */
 	float GetGameViewportDPIScale() const;
