Fixed build error cause by missing include

diff --git a/Engine/Source/Editor/LevelEditor/Private/SLevelViewport.cpp b/Engine/Source/Editor/LevelEditor/Private/SLevelViewport.cpp
--- a/Engine/Source/Editor/LevelEditor/Private/SLevelViewport.cpp
+++ b/Engine/Source/Editor/LevelEditor/Private/SLevelViewport.cpp
@@ -28,6 +28,7 @@
 #include "Settings.h"
 #include "SceneOutlinerTreeItems.h"
 #include "BufferVisualizationData.h"
+#include "EditorViewportCommands.h"
 
 static const FName LevelEditorName("LevelEditor");
 
