Parent the high res screenshot ui to the main frame

diff --git a/Engine/Source/Editor/LevelEditor/Private/HighResScreenshotUI.cpp b/Engine/Source/Editor/LevelEditor/Private/HighResScreenshotUI.cpp
--- a/Engine/Source/Editor/LevelEditor/Private/HighResScreenshotUI.cpp
+++ b/Engine/Source/Editor/LevelEditor/Private/HighResScreenshotUI.cpp
@@ -292,7 +292,11 @@ TWeakPtr<class SWindow> SHighResScreenshotDialog::OpenDialog(const TSharedPtr<FS
 
 		Window->SetOnWindowClosed(FOnWindowClosed::CreateStatic(&WindowClosedHandler));
 		Dialog->SetWindow(Window);
-		FSlateApplication::Get().AddWindow(Window);
+
+
+		TSharedPtr<SWindow> ParentWindow = FGlobalTabmanager::Get()->GetRootWindow();
+
+		FSlateApplication::Get().AddWindowAsNativeChild(Window, ParentWindow.ToSharedRef());
 
 		CurrentWindow = TWeakPtr<SWindow>(Window);
 		CurrentDialog = TWeakPtr<SHighResScreenshotDialog>(Dialog);
