Fixed crash when hot-reloading the Level Editor

This addresses TTP#344926 - CRITICAL: Regression: EDITOR: Crash: Reload Level Editor

diff --git a/Engine/Source/Editor/LevelEditor/Private/LevelEditor.cpp b/Engine/Source/Editor/LevelEditor/Private/LevelEditor.cpp
--- a/Engine/Source/Editor/LevelEditor/Private/LevelEditor.cpp
+++ b/Engine/Source/Editor/LevelEditor/Private/LevelEditor.cpp
@@ -202,7 +202,7 @@ void FLevelEditorModule::ShutdownModule()
 		TSharedPtr<SDockTab> LevelEditorTab = LevelEditorInstanceTabPtr.Pin();
 		if (LevelEditorTab.IsValid())
 		{
-			LevelEditorTab->RequestCloseTab();
+			LevelEditorTab->RemoveTabFromParent();
 		}
 		LevelEditorInstanceTabPtr.Reset();
 	}
