Can now modify the Game Mode for the Project Settings from the Blueprints menu on the editor toolbar as well as for the World Settings.

Will help the user checkout/force writable on required config file for editing.

diff --git a/Engine/Source/Editor/LevelEditor/LevelEditor.Build.cs b/Engine/Source/Editor/LevelEditor/LevelEditor.Build.cs
--- a/Engine/Source/Editor/LevelEditor/LevelEditor.Build.cs
+++ b/Engine/Source/Editor/LevelEditor/LevelEditor.Build.cs
@@ -50,6 +50,7 @@ public class LevelEditor : ModuleRules
                 "RHI",
 				"Projects",
 				"TargetPlatform",
+				"EngineSettings",
 			}
 		);
 
diff --git a/Engine/Source/Editor/LevelEditor/Private/LevelEditorActions.cpp b/Engine/Source/Editor/LevelEditor/Private/LevelEditorActions.cpp
--- a/Engine/Source/Editor/LevelEditor/Private/LevelEditorActions.cpp
+++ b/Engine/Source/Editor/LevelEditor/Private/LevelEditorActions.cpp
@@ -49,6 +49,7 @@
 #include "ActorPickerMode.h"
 #include "EngineBuildSettings.h"
 #include "HotReloadInterface.h"
+#include "ISourceControlModule.h"
 
 DEFINE_LOG_CATEGORY_STATIC(LevelEditorActions, Log, All);
 
@@ -1698,9 +1699,25 @@ void FLevelEditorActionCallbacks::OpenMarketplace()
 	}
 }
 
-bool FLevelEditorActionCallbacks::CanEditGameInfoBlueprints( TWeakPtr< SLevelEditor > LevelEditor)
+bool FLevelEditorActionCallbacks::CanSelectGameModeBlueprint()
 {
-	return LevelEditor.Pin()->GetWorld()->GetWorldSettings()->DefaultGameMode && LevelEditor.Pin()->GetWorld()->GetWorldSettings()->DefaultGameMode->ClassGeneratedBy;
+	bool bCheckOutNeeded = false;
+
+	FString ConfigFilePath = FPaths::ConvertRelativePathToFull(FString::Printf(TEXT("%sDefaultEngine.ini"), *FPaths::SourceConfigDir()));
+	if(ISourceControlModule::Get().IsEnabled())
+	{
+		// note: calling QueueStatusUpdate often does not spam status updates as an internal timer prevents this
+		//ISourceControlModule::Get().QueueStatusUpdate(ConfigFilePath);
+
+		ISourceControlProvider& SourceControlProvider = ISourceControlModule::Get().GetProvider();
+		FSourceControlStatePtr SourceControlState = SourceControlProvider.GetState(ConfigFilePath, EStateCacheUsage::Use);
+		bCheckOutNeeded = SourceControlState.IsValid() && SourceControlState->CanCheckout();
+	}
+	else
+	{
+		bCheckOutNeeded = (FPaths::FileExists(ConfigFilePath) && IFileManager::Get().IsReadOnly(*ConfigFilePath));
+	}
+	return !bCheckOutNeeded;
 }
 
 void FLevelEditorActionCallbacks::OpenLevelBlueprint( TWeakPtr< SLevelEditor > LevelEditor )
@@ -1747,132 +1764,16 @@ void FLevelEditorActionCallbacks::CreateClassBlueprint()
 	}
 }
 
-void FLevelEditorActionCallbacks::OpenGameModeBlueprint( TWeakPtr< SLevelEditor > LevelEditor )
+void FLevelEditorActionCallbacks::CheckOutProjectSettingsConfig( )
 {
-	if( LevelEditor.Pin()->GetWorld()->GetCurrentLevel() )
+	FString ConfigFilePath = FPaths::ConvertRelativePathToFull(FString::Printf(TEXT("%sDefaultEngine.ini"), *FPaths::SourceConfigDir()));
+	if(ISourceControlModule::Get().IsEnabled())
 	{
-		TSubclassOf<class AGameMode> ActiveGameMode = LevelEditor.Pin()->GetWorld()->GetWorldSettings()->DefaultGameMode;
-		UClass* ParentBlueprintClass = NULL;
-
-		if(ActiveGameMode)
-		{
-			ParentBlueprintClass = ActiveGameMode;
-		}
-		else
-		{
-			ParentBlueprintClass = AGameMode::StaticClass();
-		}
-
-		const FString NewBPName(TEXT("NewGameMode"));
-		UBlueprint* Blueprint = LevelEditorActionsHelpers::OpenOrCreateBlueprintFromClass(NSLOCTEXT("LevelEditorCommands", "CreateGameModeBlueprint_Title", "Create GameMode Blueprint"), ParentBlueprintClass, LevelEditor, NewBPName);
-		if(Blueprint)
-		{
-			LevelEditor.Pin()->GetWorld()->GetWorldSettings()->DefaultGameMode = Cast<UClass>(Blueprint->GeneratedClass);
-		}
+		SourceControlHelpers::CheckOutFile(ConfigFilePath);
 	}
-}
-
-void FLevelEditorActionCallbacks::OpenGameStateBlueprint( TWeakPtr< SLevelEditor > LevelEditor )
-{
-	if( LevelEditor.Pin()->GetWorld()->GetCurrentLevel() )
-	{
-		TSubclassOf<class AGameMode> ActiveGameMode = LevelEditor.Pin()->GetWorld()->GetWorldSettings()->DefaultGameMode;
-		AGameMode* GameMode = Cast<AGameMode>(ActiveGameMode->GetDefaultObject());
-
-		UClass* ParentBlueprintClass = NULL;
-		if(GameMode->GameStateClass)
-		{
-			ParentBlueprintClass = GameMode->GameStateClass;
-		}
-		else
-		{
-			ParentBlueprintClass = AGameState::StaticClass();
-		}
-
-		const FString NewBPName(TEXT("NewGameState"));
-		UBlueprint* Blueprint = LevelEditorActionsHelpers::OpenOrCreateBlueprintFromClass(NSLOCTEXT("LevelEditorCommands", "CreateGameStateBlueprint_Title", "Create GameState Blueprint"), ParentBlueprintClass, LevelEditor, NewBPName);
-		if(Blueprint)
-		{
-			GameMode->GameStateClass = Cast<UClass>(Blueprint->GeneratedClass);
-		}
-	}
-}
-
-void FLevelEditorActionCallbacks::OpenDefaultPawnBlueprint( TWeakPtr< SLevelEditor > LevelEditor )
-{
-	if( LevelEditor.Pin()->GetWorld()->GetCurrentLevel() )
-	{
-		TSubclassOf<class AGameMode> ActiveGameMode = LevelEditor.Pin()->GetWorld()->GetWorldSettings()->DefaultGameMode;
-		AGameMode* GameMode = Cast<AGameMode>(ActiveGameMode->GetDefaultObject());
-
-		UClass* ParentBlueprintClass = NULL;
-		if(GameMode->DefaultPawnClass)
-		{
-			ParentBlueprintClass = GameMode->DefaultPawnClass;
-		}
-		else
-		{
-			ParentBlueprintClass = APawn::StaticClass();
-		}
-
-		const FString NewBPName(TEXT("NewPawn"));
-		UBlueprint* Blueprint = LevelEditorActionsHelpers::OpenOrCreateBlueprintFromClass(NSLOCTEXT("LevelEditorCommands", "CreateDefaultPawnBlueprint_Title", "Create DefaultPawn Blueprint"), ParentBlueprintClass, LevelEditor, NewBPName);
-		if(Blueprint)
-		{
-			GameMode->DefaultPawnClass = Cast<UClass>(Blueprint->GeneratedClass);
-		}
-	}
-}
-
-void FLevelEditorActionCallbacks::OpenHUDBlueprint( TWeakPtr< SLevelEditor > LevelEditor )
-{
-	if( LevelEditor.Pin()->GetWorld()->GetCurrentLevel() )
-	{
-		TSubclassOf<class AGameMode> ActiveGameMode = LevelEditor.Pin()->GetWorld()->GetWorldSettings()->DefaultGameMode;
-		AGameMode* GameMode = Cast<AGameMode>(ActiveGameMode->GetDefaultObject());
-
-		UClass* ParentBlueprintClass = NULL;
-		if(GameMode->HUDClass)
-		{
-			ParentBlueprintClass = GameMode->HUDClass;
-		}
-		else
-		{
-			ParentBlueprintClass = AHUD::StaticClass();
-		}
-
-		const FString NewBPName(TEXT("NewHUD"));
-		UBlueprint* Blueprint = LevelEditorActionsHelpers::OpenOrCreateBlueprintFromClass(NSLOCTEXT("LevelEditorCommands", "CreateHUDBlueprint_Title", "Create HUD Blueprint"), ParentBlueprintClass, LevelEditor, NewBPName);
-		if(Blueprint)
-		{
-			GameMode->HUDClass = Cast<UClass>(Blueprint->GeneratedClass);
-		}
-	}
-}
-
-void FLevelEditorActionCallbacks::OpenPlayerControllerBlueprint( TWeakPtr< SLevelEditor > LevelEditor )
-{
-	if( LevelEditor.Pin()->GetWorld()->GetCurrentLevel() )
+	else
 	{
-		TSubclassOf<class AGameMode> ActiveGameMode = LevelEditor.Pin()->GetWorld()->GetWorldSettings()->DefaultGameMode;
-		AGameMode* GameMode = Cast<AGameMode>(ActiveGameMode->GetDefaultObject());
-
-		UClass* ParentBlueprintClass = NULL;
-		if(GameMode->PlayerControllerClass)
-		{
-			ParentBlueprintClass = GameMode->PlayerControllerClass;
-		}
-		else
-		{
-			ParentBlueprintClass = APlayerController::StaticClass();
-		}
-
-		const FString NewBPName(TEXT("NewPC"));
-		UBlueprint* Blueprint = LevelEditorActionsHelpers::OpenOrCreateBlueprintFromClass(NSLOCTEXT("LevelEditorCommands", "CreatePlayerControllerBlueprint_Title", "Create PlayerController Blueprint"), ParentBlueprintClass, LevelEditor, NewBPName);
-		if(Blueprint)
-		{
-			GameMode->PlayerControllerClass = Cast<UClass>(Blueprint->GeneratedClass);
-		}
+		FPlatformFileManager::Get().GetPlatformFile().SetReadOnly(*ConfigFilePath, false);
 	}
 }
 
@@ -2757,11 +2658,7 @@ void FLevelEditorCommands::RegisterCommands()
 	UI_COMMAND( EditMatinee, "Edit Matinee", "Selects a Matinee to edit", EUserInterfaceActionType::Button, FInputGesture() );
 
 	UI_COMMAND( OpenLevelBlueprint, "Open Level Blueprint", "Edit the Level Blueprint for the current level", EUserInterfaceActionType::Button, FInputGesture() );
-	UI_COMMAND( OpenGameModeBlueprint, "Open/Create Game Mode Blueprint", "Opens or create the active GameMode blueprint", EUserInterfaceActionType::Button, FInputGesture() );
-	UI_COMMAND( OpenGameStateBlueprint, "Open/Create Game State Blueprint", "Open or create the GameState blueprint and auto assign it to the active GameMode", EUserInterfaceActionType::Button, FInputGesture() );
-	UI_COMMAND( OpenDefaultPawnBlueprint, "Open/Create Default Pawn Blueprint", "Open or create the Pawn blueprint and auto assign it to the active GameMode", EUserInterfaceActionType::Button, FInputGesture() );
-	UI_COMMAND( OpenHUDBlueprint, "Open/Create HUD Blueprint", "Open or create the HUD blueprint and auto assign it to the active GameMode", EUserInterfaceActionType::Button, FInputGesture() );
-	UI_COMMAND( OpenPlayerControllerBlueprint, "Open/Create Player Controller Blueprint", "Open or create the Player Controller blueprint and auto assign it to the active GameMode", EUserInterfaceActionType::Button, FInputGesture() );
+	UI_COMMAND( CheckOutProjectSettingsConfig, "Check Out", "Checks out the project settings config file so the game mode can be set.", EUserInterfaceActionType::Button, FInputGesture() );
 	UI_COMMAND( CreateClassBlueprint, "New Class Blueprint...", "Create a new Class Blueprint", EUserInterfaceActionType::Button, FInputGesture());
 
 	UI_COMMAND( ShowTransformWidget, "Show Transform Widget", "Toggles the visibility of the transform widgets", EUserInterfaceActionType::ToggleButton, FInputGesture() );
diff --git a/Engine/Source/Editor/LevelEditor/Private/LevelEditorToolBar.cpp b/Engine/Source/Editor/LevelEditor/Private/LevelEditorToolBar.cpp
--- a/Engine/Source/Editor/LevelEditor/Private/LevelEditorToolBar.cpp
+++ b/Engine/Source/Editor/LevelEditor/Private/LevelEditorToolBar.cpp
@@ -21,6 +21,7 @@
 #include "Editor/ClassViewer/Private/SClassViewer.h"
 #include "Editor/ClassViewer/Public/ClassViewerFilter.h"
 #include "KismetEditorUtilities.h"
+#include "ISourceControlModule.h"
 
 namespace LevelEditorActionHelpers
 {
@@ -43,5 +44,182 @@
 	};
 
 	/**
+	 * Retrieves the GameMode class
+	 *
+	 * @param InLevelEditor				The editor to extract the world from
+	 * @param bInIsProjectSettings		TRUE if retrieving the game mode from the project settings
+	 * @return							The GameMode class in the Project Settings or World Settings
+	 */
+	static UClass* GetGameModeClass(TWeakPtr< SLevelEditor > InLevelEditor, bool bInIsProjectSettings);
+
+	/** Callback for the label to display for the GameMode menu selection */
+	static FText GetOpenGameModeBlueprintLabel(TWeakPtr< SLevelEditor > InLevelEditor, bool bInIsProjectSettings);
+
+	/** Callback when selecting a GameMode class, assigns it to the world */
+	static void OnSelectGameModeClassPicked(UClass* InChosenClass, TWeakPtr< SLevelEditor > InLevelEditor, bool bInIsProjectSettings);
+
+	/** Callback when creating a new GameMode class, creates the Blueprint and assigns it to the world */
+	static void OnCreateGameModeClassPicked(UClass* InChosenClass, TWeakPtr< SLevelEditor > InLevelEditor, bool bInIsProjectSettings);
+
+	/**
+	 * Retrieves the active GameState class from
+	 *
+	 * @param InLevelEditor		The editor to extract the world from
+	 * @param bInIsProjectSettings		TRUE if retrieving the game mode from the project settings
+	 * @return					The active GameState class in the World Settings
+	 */
+	static UClass* GetGameStateClass(TWeakPtr< SLevelEditor > InLevelEditor, bool bInIsProjectSettings);
+
+	/** Callback for the label to display for the GameState menu selection */
+	static FText GetOpenGameStateBlueprintLabel(TWeakPtr< SLevelEditor > InLevelEditor, bool bInIsProjectSettings);
+	
+	/** Callback when selecting a GameState class, assigns it to the world */
+	static void OnSelectGameStateClassPicked(UClass* InChosenClass, TWeakPtr< SLevelEditor > InLevelEditor, bool bInIsProjectSettings);
+
+	/** Callback when creating a new GameState class, creates the Blueprint and assigns it to the world */
+	static void OnCreateGameStateClassPicked(UClass* InChosenClass, TWeakPtr< SLevelEditor > InLevelEditor, bool bInIsProjectSettings);
+
+	/**
+	 * Retrieves the active Pawn class from
+	 *
+	 * @param InLevelEditor		The editor to extract the world from
+	 * @return					The active Pawn class in the World Settings
+	 */
+	static UClass* GetPawnClass(TWeakPtr< SLevelEditor > InLevelEditor, bool bInIsProjectSettings);
+
+	/** Callback for the label to display for the Pawn menu selection */
+	static FText GetOpenPawnBlueprintLabel(TWeakPtr< SLevelEditor > InLevelEditor, bool bInIsProjectSettings);
+
+	/** Callback for the tooltip to display for the Pawn menu selection */
+	static FText GetOpenPawnBlueprintTooltip(TWeakPtr< SLevelEditor > InLevelEditor);
+
+	/** Callback when selecting a Pawn class, assigns it to the world */
+	static void OnSelectPawnClassPicked(UClass* InChosenClass, TWeakPtr< SLevelEditor > InLevelEditor, bool bInIsProjectSettings);
+
+	/** Callback when creating a new Pawn class, creates the Blueprint and assigns it to the world */
+	static void OnCreatePawnClassPicked(UClass* InChosenClass, TWeakPtr< SLevelEditor > InLevelEditor, bool bInIsProjectSettings);
+
+	/**
+	 * Retrieves the active HUD class from
+	 *
+	 * @param InLevelEditor		The editor to extract the world from
+	 * @param bInIsProjectSettings		TRUE if retrieving the game mode from the project settings
+	 * @return					The active HUD class in the World Settings
+	 */
+	static UClass* GetHUDClass(TWeakPtr< SLevelEditor > InLevelEditor, bool bInIsProjectSettings);
+
+	/** Callback for the label to display for the HUD menu selection */
+	static FText GetOpenHUDBlueprintLabel(TWeakPtr< SLevelEditor > InLevelEditor, bool bInIsProjectSettings);
+	
+	/** Callback when selecting a HUD class, assigns it to the world */
+	static void OnSelectHUDClassPicked(UClass* InChosenClass, TWeakPtr< SLevelEditor > InLevelEditor, bool bInIsProjectSettings);
+
+	/** Callback when creating a new HUD class, creates the Blueprint and assigns it to the world */
+	static void OnCreateHUDClassPicked(UClass* InChosenClass, TWeakPtr< SLevelEditor > InLevelEditor, bool bInIsProjectSettings);
+
+	/**
+	 * Retrieves the active PlayerController class from
+	 *
+	 * @param InLevelEditor		The editor to extract the world from
+	 * @param bInIsProjectSettings		TRUE if retrieving the game mode from the project settings
+	 * @return					The active PlayerController class in the World Settings
+	 */
+	static UClass* GetPlayerControllerClass(TWeakPtr< SLevelEditor > InLevelEditor, bool bInIsProjectSettings);
+
+	/** Callback for the label to display for the PlayerController menu selection */
+	static FText GetOpenPlayerControllerBlueprintLabel(TWeakPtr< SLevelEditor > InLevelEditor, bool bInIsProjectSettings);
+
+	/** Callback when selecting a PlayerController class, assigns it to the world */
+	static void OnSelectPlayerControllerClassPicked(UClass* InChosenClass, TWeakPtr< SLevelEditor > InLevelEditor, bool bInIsProjectSettings);
+
+	/** Callback when creating a new PlayerController class, creates the Blueprint and assigns it to the world */
+	static void OnCreatePlayerControllerClassPicked(UClass* InChosenClass, TWeakPtr< SLevelEditor > InLevelEditor, bool bInIsProjectSettings);
+
+	/** Open the game mode blueprint, in the project settings or world settings */
+	static void OpenGameModeBlueprint( TWeakPtr< SLevelEditor > InLevelEditor, bool bInIsProjectSettings )
+	{
+		if(UClass* GameModeClass = GetGameModeClass(InLevelEditor, bInIsProjectSettings))
+		{
+			if(UBlueprint* BlueprintClass = Cast<UBlueprint>(GameModeClass->ClassGeneratedBy))
+			{
+				// @todo Re-enable once world centric works
+				const bool bOpenWorldCentric = false;
+				FAssetEditorManager::Get().OpenEditorForAsset(
+					BlueprintClass,
+					bOpenWorldCentric ? EToolkitMode::WorldCentric : EToolkitMode::Standalone,
+					InLevelEditor.Pin() );
+			}
+		}
+	}
+
+	/** Open the game state blueprint, in the project settings or world settings */
+	static void OpenGameStateBlueprint( TWeakPtr< SLevelEditor > InLevelEditor, bool bInIsProjectSettings )
+	{
+		if(UClass* GameStateClass = GetGameStateClass(InLevelEditor, bInIsProjectSettings))
+		{
+			if(UBlueprint* BlueprintClass = Cast<UBlueprint>(GameStateClass->ClassGeneratedBy))
+			{
+				// @todo Re-enable once world centric works
+				const bool bOpenWorldCentric = false;
+				FAssetEditorManager::Get().OpenEditorForAsset(
+					BlueprintClass,
+					bOpenWorldCentric ? EToolkitMode::WorldCentric : EToolkitMode::Standalone,
+					InLevelEditor.Pin() );
+			}
+		}
+	}
+
+	/** Open the default pawn blueprint, in the project settings or world settings */
+	static void OpenDefaultPawnBlueprint( TWeakPtr< SLevelEditor > InLevelEditor, bool bInIsProjectSettings )
+	{
+		if(UClass* DefaultPawnClass = GetPawnClass(InLevelEditor, bInIsProjectSettings))
+		{
+			if(UBlueprint* BlueprintClass = Cast<UBlueprint>(DefaultPawnClass->ClassGeneratedBy))
+			{
+				// @todo Re-enable once world centric works
+				const bool bOpenWorldCentric = false;
+				FAssetEditorManager::Get().OpenEditorForAsset(
+					BlueprintClass,
+					bOpenWorldCentric ? EToolkitMode::WorldCentric : EToolkitMode::Standalone,
+					InLevelEditor.Pin() );
+			}
+		}
+	}
+
+	/** Open the HUD blueprint, in the project settings or world settings */
+	static void OpenHUDBlueprint( TWeakPtr< SLevelEditor > InLevelEditor, bool bInIsProjectSettings )
+	{
+		if(UClass* DefaultHUDClass = GetHUDClass(InLevelEditor, bInIsProjectSettings))
+		{
+			if(UBlueprint* BlueprintClass = Cast<UBlueprint>(DefaultHUDClass->ClassGeneratedBy))
+			{
+				// @todo Re-enable once world centric works
+				const bool bOpenWorldCentric = false;
+				FAssetEditorManager::Get().OpenEditorForAsset(
+					BlueprintClass,
+					bOpenWorldCentric ? EToolkitMode::WorldCentric : EToolkitMode::Standalone,
+					InLevelEditor.Pin() );
+			}
+		}
+	}
+
+	/** Open the player controller blueprint, in the project settings or world settings */
+	static void OpenPlayerControllerBlueprint( TWeakPtr< SLevelEditor > InLevelEditor, bool bInIsProjectSettings )
+	{
+		if(UClass* PlayerControllerClass = GetPlayerControllerClass(InLevelEditor, bInIsProjectSettings))
+		{
+			if(UBlueprint* BlueprintClass = Cast<UBlueprint>(PlayerControllerClass->ClassGeneratedBy))
+			{
+				// @todo Re-enable once world centric works
+				const bool bOpenWorldCentric = false;
+				FAssetEditorManager::Get().OpenEditorForAsset(
+					BlueprintClass,
+					bOpenWorldCentric ? EToolkitMode::WorldCentric : EToolkitMode::Standalone,
+					InLevelEditor.Pin() );
+			}
+		}
+	}
+
+	/**
 	 * Builds a sub-menu for selecting a class
 	 *
@@ -109,7 +287,7 @@
 	struct FBlueprintMenuSettings
 	{
 		/** The UI command for editing the Blueprint class associated with the menu */
-		TSharedPtr< FUICommandInfo > EditCommand;
+		FUIAction EditCommand;
 
 		/** Current class associated with the menu */
 		UClass* CurrentClass;
@@ -122,39 +300,567 @@
 
 		/** Callback when a class is picked, to create a new child class of and assign */
 		FOnClassPicked OnCreateClassPicked;
+
+		/** Level Editor these menu settings are for */
+		TWeakPtr< SLevelEditor > LevelEditor;
+
+		/** TRUE if these represent Project Settings, FALSE if they represent World Settings */
+		bool bIsProjectSettings;
 	};
 
+	/** Returns the label of the "Check Out" option based on if source control is present or not */
+	FText GetCheckOutLabel()
+	{
+#define LOCTEXT_NAMESPACE "LevelToolBarViewMenu"
+		if(ISourceControlModule::Get().IsEnabled())
+		{
+			return LOCTEXT("CheckoutMenuLabel", "Check Out");
+		}
+		else
+		{
+			return LOCTEXT("MakeWritableLabel", "Make Writable");
+		}
+#undef LOCTEXT_NAMESPACE
+	}
+
+	/** Returns the tooltip of the "Check Out" option based on if source control is present or not */
+	FText GetCheckOutTooltip()
+	{
+#define LOCTEXT_NAMESPACE "LevelToolBarViewMenu"
+		if(ISourceControlModule::Get().IsEnabled())
+		{
+			return LOCTEXT("CheckoutMenuTooltip", "Checks out the project settings config file so the game mode can be set.");
+		}
+		else
+		{
+			return LOCTEXT("MakeWritableTooltip", "Forces the project settings config file to be writable so the game mode can be set.");
+		}
+#undef LOCTEXT_NAMESPACE
+	}
+
 	/**
 	 * A sub-menu for the Blueprints dropdown, facilitates all the sub-menu actions such as creating, editing, and selecting classes for the world settings game mode.
 	 *
 	 * @param InMenuBuilder		Object to append menu items/widgets to
+	 * @param InCommandList		Commandlist for menu items
 	 * @param InSettingsData	All the data needed to create the menu actions
 	 */
-	void GetBlueprintSettingsSubMenu(FMenuBuilder& InMenuBuilder, FBlueprintMenuSettings InSettingsData)
+	void GetBlueprintSettingsSubMenu(FMenuBuilder& InMenuBuilder, const TSharedRef<FUICommandList> InCommandList, FBlueprintMenuSettings InSettingsData);
+
+	void CreateGameModeSubMenu(FMenuBuilder& InMenuBuilder, TSharedRef<FUICommandList> InCommandList, TWeakPtr< SLevelEditor > InLevelEditor, bool bInProjectSettings)
 	{
 #define LOCTEXT_NAMESPACE "LevelToolBarViewMenu"
+		LevelEditorActionHelpers::FBlueprintMenuSettings GameModeMenuSettings;
+		GameModeMenuSettings.EditCommand = 
+			FUIAction(
+				FExecuteAction::CreateStatic< TWeakPtr< SLevelEditor > >( &OpenGameModeBlueprint, InLevelEditor, bInProjectSettings )
+			);
+		GameModeMenuSettings.OnCreateClassPicked = FOnClassPicked::CreateStatic( &LevelEditorActionHelpers::OnCreateGameModeClassPicked, InLevelEditor, bInProjectSettings );
+		GameModeMenuSettings.OnSelectClassPicked = FOnClassPicked::CreateStatic( &LevelEditorActionHelpers::OnSelectGameModeClassPicked, InLevelEditor, bInProjectSettings );
+		GameModeMenuSettings.CurrentClass = LevelEditorActionHelpers::GetGameModeClass(InLevelEditor, bInProjectSettings);
+		GameModeMenuSettings.RootClass = AGameMode::StaticClass();
+		GameModeMenuSettings.LevelEditor = InLevelEditor;
+		GameModeMenuSettings.bIsProjectSettings = bInProjectSettings;
 
-		FSlateIcon EditBPIcon(FEditorStyle::Get().GetStyleSetName(), TEXT("PropertyWindow.Button_Edit"));
-		FSlateIcon NewBPIcon(FEditorStyle::Get().GetStyleSetName(), TEXT("PropertyWindow.Button_AddToArray"));
-		FText RootClassName = FText::FromString(InSettingsData.RootClass->GetName());
+		auto IsGameModeActive = [](TWeakPtr< SLevelEditor > InLevelEditor, bool bInProjectSettings)->bool
+		{
+			UClass* WorldSettingsGameMode = LevelEditorActionHelpers::GetGameModeClass(InLevelEditor, false);
+			if((WorldSettingsGameMode == nullptr) ^ bInProjectSettings ) //(WorldSettingsGameMode && !bInProjectSettings) || (!WorldSettingsGameMode && bInProjectSettings) )
+			{
+				return false;
+			}
+			return true;
+		};
 
-		// If there is currently a valid GameMode Blueprint, offer to edit the Blueprint
-		if(InSettingsData.CurrentClass && InSettingsData.CurrentClass->ClassGeneratedBy)
-		{
-			FText BlueprintName = FText::FromString(InSettingsData.CurrentClass->ClassGeneratedBy->GetName());
-			InMenuBuilder.AddMenuEntry( InSettingsData.EditCommand, NAME_None, FText::Format( LOCTEXT("EditBlueprint", "Edit {Blueprint}"), BlueprintName), FText::Format( LOCTEXT("EditBlueprint_Tooltip", "Open the world's assigned {RootClass} blueprint"), RootClassName), EditBPIcon );
-		}
+		InMenuBuilder.AddSubMenu( LevelEditorActionHelpers::GetOpenGameModeBlueprintLabel(InLevelEditor, bInProjectSettings), LOCTEXT("GameMode_Tooltip", "Select, edit, or create a new GameMode blueprint for the world"), FNewMenuDelegate::CreateStatic( &LevelEditorActionHelpers::GetBlueprintSettingsSubMenu, InCommandList, GameModeMenuSettings), FUIAction(FExecuteAction(), FCanExecuteAction(), FIsActionChecked::CreateStatic(IsGameModeActive, InLevelEditor, bInProjectSettings)), NAME_None, EUserInterfaceActionType::RadioButton );
+#undef LOCTEXT_NAMESPACE
 
+	}
+
+	/**
+	 * Builds the game mode's sub menu objects
+	 *
+	 * @param InMenuBuilder		Object to append menu items/widgets to
+	 * @param InCommandList		Commandlist for menu items
+	 * @param InSettingsData	All the data needed to create the menu actions
+	 */
+	void GetGameModeSubMenu(FMenuBuilder& InMenuBuilder, const TSharedRef<FUICommandList> InCommandList, const FBlueprintMenuSettings& InSettingsData)
+	{
+#define LOCTEXT_NAMESPACE "LevelToolBarViewMenu"
+		// Game State
+		LevelEditorActionHelpers::FBlueprintMenuSettings GameStateMenuSettings;
+		GameStateMenuSettings.EditCommand = 
+			FUIAction(
+				FExecuteAction::CreateStatic< TWeakPtr< SLevelEditor > >( &OpenGameStateBlueprint, InSettingsData.LevelEditor, InSettingsData.bIsProjectSettings )
+			);
+		GameStateMenuSettings.OnCreateClassPicked = FOnClassPicked::CreateStatic( &LevelEditorActionHelpers::OnCreateGameStateClassPicked, InSettingsData.LevelEditor, InSettingsData.bIsProjectSettings );
+		GameStateMenuSettings.OnSelectClassPicked = FOnClassPicked::CreateStatic( &LevelEditorActionHelpers::OnSelectGameStateClassPicked, InSettingsData.LevelEditor, InSettingsData.bIsProjectSettings );
+		GameStateMenuSettings.CurrentClass = LevelEditorActionHelpers::GetGameStateClass(InSettingsData.LevelEditor, InSettingsData.bIsProjectSettings);
+		GameStateMenuSettings.RootClass = AGameState::StaticClass();
+		GameStateMenuSettings.LevelEditor = InSettingsData.LevelEditor;
+		GameStateMenuSettings.bIsProjectSettings = InSettingsData.bIsProjectSettings;
+
+		InMenuBuilder.AddSubMenu( LevelEditorActionHelpers::GetOpenGameStateBlueprintLabel(InSettingsData.LevelEditor, InSettingsData.bIsProjectSettings), LOCTEXT("GameState_Tooltip", "Select, edit, or create a new GameState blueprint for the world"), FNewMenuDelegate::CreateStatic( &LevelEditorActionHelpers::GetBlueprintSettingsSubMenu, InCommandList, GameStateMenuSettings ) );
+
+		// Pawn
+		LevelEditorActionHelpers::FBlueprintMenuSettings PawnMenuSettings;
+		PawnMenuSettings.EditCommand = 
+			FUIAction(
+				FExecuteAction::CreateStatic< TWeakPtr< SLevelEditor > >( &OpenDefaultPawnBlueprint, InSettingsData.LevelEditor, InSettingsData.bIsProjectSettings )
+			);
+		PawnMenuSettings.OnCreateClassPicked = FOnClassPicked::CreateStatic( &LevelEditorActionHelpers::OnCreatePawnClassPicked, InSettingsData.LevelEditor, InSettingsData.bIsProjectSettings );
+		PawnMenuSettings.OnSelectClassPicked = FOnClassPicked::CreateStatic( &LevelEditorActionHelpers::OnSelectPawnClassPicked, InSettingsData.LevelEditor, InSettingsData.bIsProjectSettings );
+		PawnMenuSettings.CurrentClass = LevelEditorActionHelpers::GetPawnClass(InSettingsData.LevelEditor, InSettingsData.bIsProjectSettings);
+		PawnMenuSettings.RootClass = APawn::StaticClass();
+		PawnMenuSettings.LevelEditor = InSettingsData.LevelEditor;
+		PawnMenuSettings.bIsProjectSettings = InSettingsData.bIsProjectSettings;
+
+		InMenuBuilder.AddSubMenu( LevelEditorActionHelpers::GetOpenPawnBlueprintLabel(InSettingsData.LevelEditor, InSettingsData.bIsProjectSettings), LOCTEXT("Pawn_Tooltip", "Select, edit, or create a new Pawn blueprint for the world"), FNewMenuDelegate::CreateStatic( &LevelEditorActionHelpers::GetBlueprintSettingsSubMenu, InCommandList, PawnMenuSettings ) );
+
+		// HUD
+		LevelEditorActionHelpers::FBlueprintMenuSettings HUDMenuSettings;
+		HUDMenuSettings.EditCommand = 
+			FUIAction(
+				FExecuteAction::CreateStatic< TWeakPtr< SLevelEditor > >( &OpenHUDBlueprint, InSettingsData.LevelEditor, InSettingsData.bIsProjectSettings )
+			);
+		HUDMenuSettings.OnCreateClassPicked = FOnClassPicked::CreateStatic( &LevelEditorActionHelpers::OnCreateHUDClassPicked, InSettingsData.LevelEditor, InSettingsData.bIsProjectSettings );
+		HUDMenuSettings.OnSelectClassPicked = FOnClassPicked::CreateStatic( &LevelEditorActionHelpers::OnSelectHUDClassPicked, InSettingsData.LevelEditor, InSettingsData.bIsProjectSettings );
+		HUDMenuSettings.CurrentClass = LevelEditorActionHelpers::GetHUDClass(InSettingsData.LevelEditor, InSettingsData.bIsProjectSettings);
+		HUDMenuSettings.RootClass = AHUD::StaticClass();
+		HUDMenuSettings.LevelEditor = InSettingsData.LevelEditor;
+		HUDMenuSettings.bIsProjectSettings = InSettingsData.bIsProjectSettings;
+
+		InMenuBuilder.AddSubMenu( LevelEditorActionHelpers::GetOpenHUDBlueprintLabel(InSettingsData.LevelEditor, InSettingsData.bIsProjectSettings), LOCTEXT("HUD_Tooltip", "Select, edit, or create a new HUD blueprint for the world"), FNewMenuDelegate::CreateStatic( &LevelEditorActionHelpers::GetBlueprintSettingsSubMenu, InCommandList, HUDMenuSettings ) );
+
+		// Player Controller
+		LevelEditorActionHelpers::FBlueprintMenuSettings PlayerControllerMenuSettings;
+		PlayerControllerMenuSettings.EditCommand = 
+			FUIAction(
+				FExecuteAction::CreateStatic< TWeakPtr< SLevelEditor > >( &OpenPlayerControllerBlueprint, InSettingsData.LevelEditor, InSettingsData.bIsProjectSettings )
+			);
+		PlayerControllerMenuSettings.OnCreateClassPicked = FOnClassPicked::CreateStatic( &LevelEditorActionHelpers::OnCreatePlayerControllerClassPicked, InSettingsData.LevelEditor, InSettingsData.bIsProjectSettings );
+		PlayerControllerMenuSettings.OnSelectClassPicked = FOnClassPicked::CreateStatic( &LevelEditorActionHelpers::OnSelectPlayerControllerClassPicked, InSettingsData.LevelEditor, InSettingsData.bIsProjectSettings );
+		PlayerControllerMenuSettings.CurrentClass = LevelEditorActionHelpers::GetPlayerControllerClass(InSettingsData.LevelEditor, InSettingsData.bIsProjectSettings);
+		PlayerControllerMenuSettings.RootClass = APlayerController::StaticClass();
+		PlayerControllerMenuSettings.LevelEditor = InSettingsData.LevelEditor;
+		PlayerControllerMenuSettings.bIsProjectSettings = InSettingsData.bIsProjectSettings;
+
+		InMenuBuilder.AddSubMenu( LevelEditorActionHelpers::GetOpenPlayerControllerBlueprintLabel(InSettingsData.LevelEditor, InSettingsData.bIsProjectSettings), LOCTEXT("PlayerController_Tooltip", "Select, edit, or create a new Player controller blueprint for the world"), FNewMenuDelegate::CreateStatic( &LevelEditorActionHelpers::GetBlueprintSettingsSubMenu, InCommandList, PlayerControllerMenuSettings ) );
+#undef LOCTEXT_NAMESPACE
+	}
+}
+
+void LevelEditorActionHelpers::GetBlueprintSettingsSubMenu(FMenuBuilder& InMenuBuilder, const TSharedRef<FUICommandList> InCommandList, FBlueprintMenuSettings InSettingsData)
+{
+#define LOCTEXT_NAMESPACE "LevelToolBarViewMenu"
+
+	InMenuBuilder.PushCommandList(InCommandList);
+
+	FSlateIcon EditBPIcon(FEditorStyle::Get().GetStyleSetName(), TEXT("PropertyWindow.Button_Edit"));
+	FSlateIcon NewBPIcon(FEditorStyle::Get().GetStyleSetName(), TEXT("PropertyWindow.Button_AddToArray"));
+	FText RootClassName = FText::FromString(InSettingsData.RootClass->GetName());
+
+	// If there is currently a valid GameMode Blueprint, offer to edit the Blueprint
+	if(InSettingsData.CurrentClass && InSettingsData.CurrentClass->ClassGeneratedBy)
+	{
+		FText BlueprintName = FText::FromString(InSettingsData.CurrentClass->ClassGeneratedBy->GetName());
+		FFormatNamedArguments Args;
+		Args.Add(TEXT("Blueprint"), BlueprintName);
+		Args.Add(TEXT("RootClass"), RootClassName);
+
+		InMenuBuilder.AddMenuEntry( FText::Format( LOCTEXT("EditBlueprint", "Edit {Blueprint}"), Args), FText::Format( LOCTEXT("EditBlueprint_Tooltip", "Open the world's assigned {RootClass} blueprint"), Args), EditBPIcon, InSettingsData.EditCommand );
+	}
+
+	if(InSettingsData.bIsProjectSettings && InSettingsData.CurrentClass && InSettingsData.CurrentClass->IsChildOf(AGameMode::StaticClass()) && !FLevelEditorActionCallbacks::CanSelectGameModeBlueprint())
+	{
+		InMenuBuilder.BeginSection("CheckoutSection", LOCTEXT("CheckoutSection","Check Out Project Settings") );
+		TAttribute<FText> CheckOutLabel;
+		CheckOutLabel.BindStatic(&GetCheckOutLabel);
+
+		TAttribute<FText> CheckOutTooltip;
+		CheckOutTooltip.BindStatic(&GetCheckOutTooltip);
+		InMenuBuilder.AddMenuEntry(FLevelEditorCommands::Get().CheckOutProjectSettingsConfig, NAME_None, CheckOutLabel, CheckOutTooltip, FSlateIcon(FEditorStyle::Get().GetStyleSetName(), TEXT("Icons.Error")));
+		InMenuBuilder.EndSection();
+	}
+	else
+	{
 		// Create a new GameMode, this is always available so the user can easily create a new one
 		InMenuBuilder.AddSubMenu( LOCTEXT("CreateBlueprint", "Create..."), FText::Format( LOCTEXT("CreateClass_Tooltip", "Create a new {RootClass} based on a selected class and auto-assign it to the world"), RootClassName ), FNewMenuDelegate::CreateStatic( &LevelEditorActionHelpers::GetCreateSettingsClassSubMenu, InSettingsData.RootClass, InSettingsData.OnCreateClassPicked ), false, NewBPIcon );
 
 		// Select a game mode, this is always available so the user can switch his selection
 		InMenuBuilder.AddSubMenu( FText::Format( LOCTEXT("SelectGameModeClass", "Select {RootClass} Class"), RootClassName ), FText::Format( LOCTEXT("SelectGameModeClass_Tooltip", "Select a class to be the active {RootClass} in the world"), RootClassName ), FNewMenuDelegate::CreateStatic( &LevelEditorActionHelpers::GetSelectSettingsClassSubMenu, InSettingsData.RootClass, InSettingsData.OnSelectClassPicked ) );
 
+		if(InSettingsData.RootClass == AGameMode::StaticClass())
+		{
+			InMenuBuilder.BeginSection(NAME_None, LOCTEXT("GameModeClasses", "Game Mode Classes"));
+
+			if(InSettingsData.CurrentClass &&InSettingsData.CurrentClass->ClassGeneratedBy)
+			{
+				GetGameModeSubMenu(InMenuBuilder, InCommandList, InSettingsData);
+			}
+
+			InMenuBuilder.EndSection();
+		}
+	}
+
+	InMenuBuilder.PopCommandList();
+
 #undef LOCTEXT_NAMESPACE
-	}
 }
 
+UClass* LevelEditorActionHelpers::GetGameModeClass(TWeakPtr< SLevelEditor > InLevelEditor, bool bInIsProjectSettings)
+{
+	UClass* GameModeClass = nullptr;
+	if(bInIsProjectSettings)
+	{
+		UObject* GameModeObject = LoadObject<UObject>(nullptr, *UGameMapsSettings::GetGlobalDefaultGameMode());
+		if(UBlueprint* GameModeAsBlueprint = Cast<UBlueprint>(GameModeObject))
+		{
+			GameModeClass = GameModeAsBlueprint->GeneratedClass;
+		}
+		else
+		{
+			GameModeClass = FindObject<UClass>(nullptr, *UGameMapsSettings::GetGlobalDefaultGameMode());
+		}
+	}
+	else
+	{
+		AWorldSettings* WorldSettings = InLevelEditor.Pin()->GetWorld()->GetWorldSettings();
+		if(WorldSettings->DefaultGameMode)
+		{
+			GameModeClass = WorldSettings->DefaultGameMode;
+		}
+	}
+	return GameModeClass;
+}
+
+FText LevelEditorActionHelpers::GetOpenGameModeBlueprintLabel(TWeakPtr< SLevelEditor > InLevelEditor, bool bInIsProjectSettings)
+{
+#define LOCTEXT_NAMESPACE "LevelToolBarViewMenu"
+	if(UClass* GameModeClass = GetGameModeClass(InLevelEditor, bInIsProjectSettings))
+	{
+		if(GameModeClass->ClassGeneratedBy)
+		{
+			return FText::Format( LOCTEXT("GameModeEditBlueprint", "GameMode: Edit {GameModeName}"), FText::FromString(GameModeClass->ClassGeneratedBy->GetName()));
+		}
+
+		return FText::Format( LOCTEXT("GameModeBlueprint", "GameMode: {GameModeName}"), FText::FromString(GameModeClass->GetName()));
+	}
+
+	if(bInIsProjectSettings)
+	{
+		return LOCTEXT("GameModeCreateBlueprint", "GameMode: New...");
+	}
+
+	// For World Settings, we want to inform the user that they are not overridding the Project Settings
+	return LOCTEXT("GameModeNotOverridden", "GameMode: Not overridden!");
+
+#undef LOCTEXT_NAMESPACE
+}
+
+void LevelEditorActionHelpers::OnCreateGameModeClassPicked(UClass* InChosenClass, TWeakPtr< SLevelEditor > InLevelEditor, bool bInIsProjectSettings)
+{
+	if(InChosenClass)
+	{
+		const FString NewBPName(TEXT("NewGameMode"));
+		UBlueprint* Blueprint = FKismetEditorUtilities::CreateBlueprintFromClass(NSLOCTEXT("LevelEditorCommands", "CreateGameModeBlueprint_Title", "Create GameMode Blueprint"), InChosenClass, NewBPName);
+
+		if(Blueprint)
+		{
+			// @todo Re-enable once world centric works
+			const bool bOpenWorldCentric = false;
+			FAssetEditorManager::Get().OpenEditorForAsset(
+				Blueprint,
+				bOpenWorldCentric ? EToolkitMode::WorldCentric : EToolkitMode::Standalone,
+				InLevelEditor.Pin()  );
+			OnSelectGameModeClassPicked(Blueprint->GeneratedClass, InLevelEditor, bInIsProjectSettings);
+		}
+	}
+	FSlateApplication::Get().DismissAllMenus();
+}
+
+void LevelEditorActionHelpers::OnSelectGameModeClassPicked(UClass* InChosenClass, TWeakPtr< SLevelEditor > InLevelEditor, bool bInIsProjectSettings)
+{
+	if(bInIsProjectSettings)
+	{
+		UGameMapsSettings::SetGlobalDefaultGameMode(InChosenClass? InChosenClass->GetPathName() : FString());
+
+		ISettingsModule* SettingsModule = ISettingsModule::Get();
+		if (SettingsModule != nullptr)
+		{
+			ISettingsContainerPtr SettingsContainer = SettingsModule->GetContainer("Project");
+
+			if (SettingsContainer.IsValid())
+			{
+				ISettingsCategoryPtr SettingsCategory = SettingsContainer->GetCategory("Project");
+
+				if(SettingsCategory.IsValid())
+				{
+					SettingsCategory->GetSection("Maps")->Save();
+				}
+			}
+		}
+	}
+	else
+	{
+		InLevelEditor.Pin()->GetWorld()->GetWorldSettings()->DefaultGameMode = InChosenClass;
+	}
+	FSlateApplication::Get().DismissAllMenus();
+}
+
+UClass* LevelEditorActionHelpers::GetGameStateClass(TWeakPtr< SLevelEditor > InLevelEditor, bool bInIsProjectSettings)
+{
+	if(UClass* GameModeClass = GetGameModeClass(InLevelEditor, bInIsProjectSettings))
+	{
+		AGameMode* ActiveGameMode = Cast<AGameMode>(GameModeClass->GetDefaultObject());
+		if(ActiveGameMode)
+		{
+			return ActiveGameMode->GameStateClass;
+		}
+	}
+	return NULL;
+}
+
+FText LevelEditorActionHelpers::GetOpenGameStateBlueprintLabel(TWeakPtr< SLevelEditor > InLevelEditor, bool bInIsProjectSettings)
+{
+#define LOCTEXT_NAMESPACE "LevelToolBarViewMenu"
+	if(UClass* GameStateClass = GetGameStateClass(InLevelEditor, bInIsProjectSettings))
+	{
+		if(GameStateClass->ClassGeneratedBy)
+		{
+			return FText::Format( LOCTEXT("GameStateEditBlueprint", "GameState: Edit {GameStateName}"), FText::FromString(GameStateClass->ClassGeneratedBy->GetName()));
+		}
+
+		return FText::Format( LOCTEXT("GameStateBlueprint", "GameState: {GameStateName}"), FText::FromString(GameStateClass->GetName()));
+	}
+
+	return LOCTEXT("GameStateCreateBlueprint", "GameState: New...");
+#undef LOCTEXT_NAMESPACE
+}
+
+void LevelEditorActionHelpers::OnCreateGameStateClassPicked(UClass* InChosenClass, TWeakPtr< SLevelEditor > InLevelEditor, bool bInIsProjectSettings)
+{
+	if(InChosenClass)
+	{
+		const FString NewBPName(TEXT("NewGameState"));
+		UBlueprint* Blueprint = FKismetEditorUtilities::CreateBlueprintFromClass(NSLOCTEXT("LevelEditorCommands", "CreateGameStateBlueprint_Title", "Create GameState Blueprint"), InChosenClass, NewBPName);
+
+		if(Blueprint)
+		{
+			// @todo Re-enable once world centric works
+			const bool bOpenWorldCentric = false;
+			FAssetEditorManager::Get().OpenEditorForAsset(
+				Blueprint,
+				bOpenWorldCentric ? EToolkitMode::WorldCentric : EToolkitMode::Standalone,
+				InLevelEditor.Pin()  );
+
+			OnSelectGameStateClassPicked(Blueprint->GeneratedClass, InLevelEditor, bInIsProjectSettings);
+		}
+	}
+	FSlateApplication::Get().DismissAllMenus();
+}
+
+void LevelEditorActionHelpers::OnSelectGameStateClassPicked(UClass* InChosenClass, TWeakPtr< SLevelEditor > InLevelEditor, bool bInIsProjectSettings)
+{
+	if(UClass* GameModeClass = GetGameModeClass(InLevelEditor, bInIsProjectSettings))
+	{
+		AGameMode* ActiveGameMode = Cast<AGameMode>(GameModeClass->GetDefaultObject());
+		ActiveGameMode->GameStateClass = InChosenClass;
+	}
+
+	FSlateApplication::Get().DismissAllMenus();
+}
+
+UClass* LevelEditorActionHelpers::GetPawnClass(TWeakPtr< SLevelEditor > InLevelEditor, bool bInIsProjectSettings)
+{
+	if(UClass* GameModeClass = GetGameModeClass(InLevelEditor, bInIsProjectSettings))
+	{
+		AGameMode* ActiveGameMode = Cast<AGameMode>(GameModeClass->GetDefaultObject());
+
+		if(ActiveGameMode)
+		{
+			return ActiveGameMode->DefaultPawnClass;
+		}
+	}
+	return NULL;
+}
+
+FText LevelEditorActionHelpers::GetOpenPawnBlueprintLabel(TWeakPtr< SLevelEditor > InLevelEditor, bool bInIsProjectSettings)
+{
+#define LOCTEXT_NAMESPACE "LevelToolBarViewMenu"
+	if(UClass* PawnClass = GetPawnClass(InLevelEditor, bInIsProjectSettings))
+	{
+		if(PawnClass->ClassGeneratedBy)
+		{
+			return FText::Format( LOCTEXT("PawnEditBlueprint", "Pawn: Edit {PawnName}"), FText::FromString(PawnClass->ClassGeneratedBy->GetName()));
+		}
+
+		return FText::Format( LOCTEXT("PawnBlueprint", "Pawn: {PawnName}"), FText::FromString(PawnClass->GetName()));
+	}
+
+	return LOCTEXT("PawnCreateBlueprint", "Pawn: New...");
+#undef LOCTEXT_NAMESPACE
+}
+
+void LevelEditorActionHelpers::OnCreatePawnClassPicked(UClass* InChosenClass, TWeakPtr< SLevelEditor > InLevelEditor, bool bInIsProjectSettings)
+{
+	if(InChosenClass)
+	{
+		const FString NewBPName(TEXT("NewPawn"));
+		UBlueprint* Blueprint = FKismetEditorUtilities::CreateBlueprintFromClass(NSLOCTEXT("LevelEditorCommands", "CreatePawnBlueprint_Title", "Create Pawn Blueprint"), InChosenClass, NewBPName);
+
+		if(Blueprint)
+		{
+			// @todo Re-enable once world centric works
+			const bool bOpenWorldCentric = false;
+			FAssetEditorManager::Get().OpenEditorForAsset(
+				Blueprint,
+				bOpenWorldCentric ? EToolkitMode::WorldCentric : EToolkitMode::Standalone,
+				InLevelEditor.Pin()  );
+
+			OnSelectPawnClassPicked(Blueprint->GeneratedClass, InLevelEditor, bInIsProjectSettings);
+		}
+	}
+	FSlateApplication::Get().DismissAllMenus();
+}
+
+void LevelEditorActionHelpers::OnSelectPawnClassPicked(UClass* InChosenClass, TWeakPtr< SLevelEditor > InLevelEditor, bool bInIsProjectSettings)
+{
+	if(UClass* GameModeClass = GetGameModeClass(InLevelEditor, bInIsProjectSettings))
+	{
+		AGameMode* ActiveGameMode = Cast<AGameMode>(GameModeClass->GetDefaultObject());
+		ActiveGameMode->DefaultPawnClass = InChosenClass;
+	}
+
+	FSlateApplication::Get().DismissAllMenus();
+}
+
+UClass* LevelEditorActionHelpers::GetHUDClass(TWeakPtr< SLevelEditor > InLevelEditor, bool bInIsProjectSettings)
+{
+	if(UClass* GameModeClass = GetGameModeClass(InLevelEditor, bInIsProjectSettings))
+	{
+		AGameMode* ActiveGameMode = Cast<AGameMode>(GameModeClass->GetDefaultObject());
+		if(ActiveGameMode)
+		{
+			return ActiveGameMode->HUDClass;
+		}
+	}
+	return NULL;
+}
+
+FText LevelEditorActionHelpers::GetOpenHUDBlueprintLabel(TWeakPtr< SLevelEditor > InLevelEditor, bool bInIsProjectSettings)
+{
+#define LOCTEXT_NAMESPACE "LevelToolBarViewMenu"
+	if(UClass* HUDClass = GetHUDClass(InLevelEditor, bInIsProjectSettings))
+	{
+		if(HUDClass->ClassGeneratedBy)
+		{
+			return FText::Format( LOCTEXT("HUDEditBlueprint", "HUD: Edit {HUDName}"), FText::FromString(HUDClass->ClassGeneratedBy->GetName()));
+		}
+
+		return FText::Format( LOCTEXT("HUDBlueprint", "HUD: {HUDName}"), FText::FromString(HUDClass->GetName()));
+	}
+
+	return LOCTEXT("HUDCreateBlueprint", "HUD: New...");
+#undef LOCTEXT_NAMESPACE
+}
+
+void LevelEditorActionHelpers::OnCreateHUDClassPicked(UClass* InChosenClass, TWeakPtr< SLevelEditor > InLevelEditor, bool bInIsProjectSettings)
+{
+	if(InChosenClass)
+	{
+		const FString NewBPName(TEXT("NewHUD"));
+		UBlueprint* Blueprint = FKismetEditorUtilities::CreateBlueprintFromClass(NSLOCTEXT("LevelEditorCommands", "CreateHUDBlueprint_Title", "Create HUD Blueprint"), InChosenClass, NewBPName);
+
+		if(Blueprint)
+		{
+			// @todo Re-enable once world centric works
+			const bool bOpenWorldCentric = false;
+			FAssetEditorManager::Get().OpenEditorForAsset(
+				Blueprint,
+				bOpenWorldCentric ? EToolkitMode::WorldCentric : EToolkitMode::Standalone,
+				InLevelEditor.Pin()  );
+
+			OnSelectHUDClassPicked(Blueprint->GeneratedClass, InLevelEditor, bInIsProjectSettings);
+		}
+	}
+	FSlateApplication::Get().DismissAllMenus();
+}
+
+void LevelEditorActionHelpers::OnSelectHUDClassPicked(UClass* InChosenClass, TWeakPtr< SLevelEditor > InLevelEditor, bool bInIsProjectSettings)
+{
+	if(UClass* GameModeClass = GetGameModeClass(InLevelEditor, bInIsProjectSettings))
+	{
+		AGameMode* ActiveGameMode = Cast<AGameMode>(GameModeClass->GetDefaultObject());
+		ActiveGameMode->HUDClass = InChosenClass;
+	}
+
+	FSlateApplication::Get().DismissAllMenus();
+}
+
+UClass* LevelEditorActionHelpers::GetPlayerControllerClass(TWeakPtr< SLevelEditor > InLevelEditor, bool bInIsProjectSettings)
+{
+	if(UClass* GameModeClass = GetGameModeClass(InLevelEditor, bInIsProjectSettings))
+	{
+		AGameMode* ActiveGameMode = Cast<AGameMode>(GameModeClass->GetDefaultObject());
+		if(ActiveGameMode)
+		{
+			return ActiveGameMode->PlayerControllerClass;
+		}
+	}
+	return NULL;
+}
+
+FText LevelEditorActionHelpers::GetOpenPlayerControllerBlueprintLabel(TWeakPtr< SLevelEditor > InLevelEditor, bool bInIsProjectSettings)
+{
+#define LOCTEXT_NAMESPACE "LevelToolBarViewMenu"
+	if(UClass* PlayerControllerClass = GetPlayerControllerClass(InLevelEditor, bInIsProjectSettings))
+	{
+		if(PlayerControllerClass->ClassGeneratedBy)
+		{
+			return FText::Format( LOCTEXT("PlayerControllerEditBlueprint", "PlayerController: Edit {PlayerControllerName}"), FText::FromString(PlayerControllerClass->ClassGeneratedBy->GetName()));
+		}
+
+		return FText::Format( LOCTEXT("PlayerControllerBlueprint", "PlayerController: {PlayerControllerName}"), FText::FromString(PlayerControllerClass->GetName()));
+	}
+
+	return LOCTEXT("PlayerControllerCreateBlueprint", "PlayerController: New...");
+#undef LOCTEXT_NAMESPACE
+}
+
+void LevelEditorActionHelpers::OnCreatePlayerControllerClassPicked(UClass* InChosenClass, TWeakPtr< SLevelEditor > InLevelEditor, bool bInIsProjectSettings)
+{
+	if(InChosenClass)
+	{
+		const FString NewBPName(TEXT("NewPlayerController"));
+		UBlueprint* Blueprint = FKismetEditorUtilities::CreateBlueprintFromClass(NSLOCTEXT("LevelEditorCommands", "CreatePlayerControllerBlueprint_Title", "Create PlayerController Blueprint"), InChosenClass, NewBPName);
+
+		if(Blueprint)
+		{
+			// @todo Re-enable once world centric works
+			const bool bOpenWorldCentric = false;
+			FAssetEditorManager::Get().OpenEditorForAsset(
+				Blueprint,
+				bOpenWorldCentric ? EToolkitMode::WorldCentric : EToolkitMode::Standalone,
+				InLevelEditor.Pin()  );
+
+			OnSelectPlayerControllerClassPicked(Blueprint->GeneratedClass, InLevelEditor, bInIsProjectSettings);
+		}
+	}
+	FSlateApplication::Get().DismissAllMenus();
+}
+
+void LevelEditorActionHelpers::OnSelectPlayerControllerClassPicked(UClass* InChosenClass, TWeakPtr< SLevelEditor > InLevelEditor, bool bInIsProjectSettings)
+{
+	if(UClass* GameModeClass = GetGameModeClass(InLevelEditor, bInIsProjectSettings))
+	{
+		AGameMode* ActiveGameMode = Cast<AGameMode>(GameModeClass->GetDefaultObject());
+		ActiveGameMode->PlayerControllerClass = InChosenClass;
+	}
+	FSlateApplication::Get().DismissAllMenus();
+}
+
+
 /**
  * Static: Creates a widget for the level editor tool bar
  *
@@ -289,488 +995,3 @@
 
 	TArray<TSharedPtr<FExtender>> Extenders;
 	for (int32 i = 0; i < MenuExtenderDelegates.Num(); ++i)
-	{
-		if (MenuExtenderDelegates[i].IsBound())
-		{
-			Extenders.Add(MenuExtenderDelegates[i].Execute(InCommandList));
-		}
-	}
-	TSharedPtr<FExtender> MenuExtender = FExtender::Combine(Extenders);
-
-	const bool bShouldCloseWindowAfterMenuSelection = true;
-	FMenuBuilder MenuBuilder( bShouldCloseWindowAfterMenuSelection, InCommandList, MenuExtender );
-
-	struct FLightingMenus
-	{
-
-		/** Generates a lighting quality sub-menu */
-		static void MakeLightingQualityMenu( FMenuBuilder& InMenuBuilder )
-		{
-			InMenuBuilder.BeginSection("LevelEditorBuildLightingQuality", LOCTEXT( "LightingQualityHeading", "Quality Level" ) );
-			{
-				InMenuBuilder.AddMenuEntry( FLevelEditorCommands::Get().LightingQuality_Production );
-				InMenuBuilder.AddMenuEntry( FLevelEditorCommands::Get().LightingQuality_High );
-				InMenuBuilder.AddMenuEntry( FLevelEditorCommands::Get().LightingQuality_Medium );
-				InMenuBuilder.AddMenuEntry( FLevelEditorCommands::Get().LightingQuality_Preview );
-			}
-			InMenuBuilder.EndSection();
-		}
-
-		/** Generates a lighting tools sub-menu */
-		static void MakeLightingToolsMenu( FMenuBuilder& InMenuBuilder )
-		{
-			InMenuBuilder.BeginSection("LevelEditorBuildLightingTools", LOCTEXT( "LightingToolsHeading", "Light Environments" ) );
-			{
-				InMenuBuilder.AddMenuEntry( FLevelEditorCommands::Get().LightingTools_ShowBounds );
-				InMenuBuilder.AddMenuEntry( FLevelEditorCommands::Get().LightingTools_ShowTraces );
-				InMenuBuilder.AddMenuEntry( FLevelEditorCommands::Get().LightingTools_ShowDirectOnly );
-				InMenuBuilder.AddMenuEntry( FLevelEditorCommands::Get().LightingTools_ShowIndirectOnly );
-				InMenuBuilder.AddMenuEntry( FLevelEditorCommands::Get().LightingTools_ShowIndirectSamples );
-			}
-			InMenuBuilder.EndSection();
-		}
-
-		/** Generates a lighting density sub-menu */
-		static void MakeLightingDensityMenu( FMenuBuilder& InMenuBuilder )
-		{
-			InMenuBuilder.BeginSection("LevelEditorBuildLightingDensity", LOCTEXT( "LightingDensityHeading", "Density Rendering" ) );
-			{
-				TSharedRef<SWidget> Ideal =		SNew(SHorizontalBox)
-												+SHorizontalBox::Slot()
-												.Padding( FMargin( 27.0f, 0.0f, 0.0f, 0.0f ) )
-												.FillWidth(1.0f)
-												[
-													SNew(SSpinBox<float>)
-													.MinValue(0.f)
-													.MaxValue(100.f)
-													.Value(FLevelEditorActionCallbacks::GetLightingDensityIdeal())
-													.OnValueChanged_Static(&FLevelEditorActionCallbacks::SetLightingDensityIdeal)
-												];
-				InMenuBuilder.AddWidget(Ideal, LOCTEXT("LightingDensity_Ideal","Ideal Density"));
-				
-				TSharedRef<SWidget> Maximum =	SNew(SHorizontalBox)
-												+SHorizontalBox::Slot()
-												.FillWidth(1.0f)
-												[
-													SNew(SSpinBox<float>)
-													.MinValue(0.01f)
-													.MaxValue(100.01f)
-													.Value(FLevelEditorActionCallbacks::GetLightingDensityMaximum())
-													.OnValueChanged_Static(&FLevelEditorActionCallbacks::SetLightingDensityMaximum)
-												];
-				InMenuBuilder.AddWidget(Maximum, LOCTEXT("LightingDensity_Maximum","Maximum Density"));
-
-				TSharedRef<SWidget> ClrScale =	SNew(SHorizontalBox)
-												+SHorizontalBox::Slot()
-												.Padding( FMargin( 35.0f, 0.0f, 0.0f, 0.0f ) )
-												.FillWidth(1.0f)
-												[
-													SNew(SSpinBox<float>)
-													.MinValue(0.f)
-													.MaxValue(10.f)
-													.Value(FLevelEditorActionCallbacks::GetLightingDensityColorScale())
-													.OnValueChanged_Static(&FLevelEditorActionCallbacks::SetLightingDensityColorScale)
-												];
-				InMenuBuilder.AddWidget(ClrScale, LOCTEXT("LightingDensity_ColorScale","Color Scale"));
-
-				TSharedRef<SWidget> GrayScale =	SNew(SHorizontalBox)
-												+SHorizontalBox::Slot()
-												.Padding( FMargin( 11.0f, 0.0f, 0.0f, 0.0f ) )
-												.FillWidth(1.0f)
-												[
-													SNew(SSpinBox<float>)
-													.MinValue(0.f)
-													.MaxValue(10.f)
-													.Value(FLevelEditorActionCallbacks::GetLightingDensityGrayscaleScale())
-													.OnValueChanged_Static(&FLevelEditorActionCallbacks::SetLightingDensityGrayscaleScale)
-												];
-				InMenuBuilder.AddWidget(GrayScale, LOCTEXT("LightingDensity_GrayscaleScale","Grayscale Scale"));
-
-				InMenuBuilder.AddMenuEntry( FLevelEditorCommands::Get().LightingDensity_RenderGrayscale );
-			}
-			InMenuBuilder.EndSection();
-		}
-
-		/** Generates a lighting resolution sub-menu */
-		static void MakeLightingResolutionMenu( FMenuBuilder& InMenuBuilder )
-		{
-			InMenuBuilder.BeginSection("LevelEditorBuildLightingResolution1", LOCTEXT( "LightingResolutionHeading1", "Primitive Types" ) );
-			{
-				TSharedRef<SWidget> Meshes =	SNew(SHorizontalBox)
-												+SHorizontalBox::Slot()
-												.AutoWidth()
-												[
-													SNew( SCheckBox )
-													.Style( FEditorStyle::Get(), "Menu.CheckBox" )
-													.ToolTipText(LOCTEXT( "StaticMeshesToolTip", "Static Meshes will be adjusted if checked." ))
-													.IsChecked_Static(&FLevelEditorActionCallbacks::IsLightingResolutionStaticMeshesChecked)
-													.OnCheckStateChanged_Static(&FLevelEditorActionCallbacks::SetLightingResolutionStaticMeshes)
-													.Content()
-													[
-														SNew( STextBlock )
-														.Text( LOCTEXT("StaticMeshes", "Static Meshes") )
-													]
-												]
-												+SHorizontalBox::Slot()
-												.AutoWidth()
-												.Padding( FMargin( 4.0f, 0.0f, 11.0f, 0.0f ) )
-												[
-													SNew(SSpinBox<float>)
-													.MinValue(4.f)
-													.MaxValue(4096.f)
-													.ToolTipText(LOCTEXT( "LightingResolutionStaticMeshesMinToolTip", "The minimum lightmap resolution for static mesh adjustments. Anything outside of Min/Max range will not be touched when adjusting." ))
-													.Value(FLevelEditorActionCallbacks::GetLightingResolutionMinSMs())
-													.OnValueChanged_Static(&FLevelEditorActionCallbacks::SetLightingResolutionMinSMs)
-												]
-												+SHorizontalBox::Slot()
-												.AutoWidth()
-												[
-													SNew(SSpinBox<float>)
-													.MinValue(4.f)
-													.MaxValue(4096.f)
-													.ToolTipText(LOCTEXT( "LightingResolutionStaticMeshesMaxToolTip", "The maximum lightmap resolution for static mesh adjustments. Anything outside of Min/Max range will not be touched when adjusting." ))
-													.Value(FLevelEditorActionCallbacks::GetLightingResolutionMaxSMs())
-													.OnValueChanged_Static(&FLevelEditorActionCallbacks::SetLightingResolutionMaxSMs)
-												];
-				InMenuBuilder.AddWidget(Meshes, FText::GetEmpty(), true);
-				
-				TSharedRef<SWidget> BSPs =		SNew(SHorizontalBox)
-												+SHorizontalBox::Slot()
-												.AutoWidth()
-												[
-													SNew( SCheckBox )
-													.Style(FEditorStyle::Get(), "Menu.CheckBox")
-													.ToolTipText(LOCTEXT( "BSPSurfacesToolTip", "BSP Surfaces will be adjusted if checked." ))
-													.IsChecked_Static(&FLevelEditorActionCallbacks::IsLightingResolutionBSPSurfacesChecked)
-													.OnCheckStateChanged_Static(&FLevelEditorActionCallbacks::SetLightingResolutionBSPSurfaces)
-													.Content()
-													[
-														SNew( STextBlock )
-														.Text( LOCTEXT("BSPSurfaces", "BSP Surfaces") )
-													]
-												]
-												+SHorizontalBox::Slot()
-												.AutoWidth()
-												.Padding( FMargin( 6.0f, 0.0f, 4.0f, 0.0f ) )
-												[
-													SNew(SSpinBox<float>)
-													.MinValue(1.f)
-													.MaxValue(63556.f)
-													.ToolTipText(LOCTEXT( "LightingResolutionBSPsMinToolTip", "The minimum lightmap resolution of a BSP surface to adjust. When outside of the Min/Max range, the BSP surface will no be altered." ))
-													.Value(FLevelEditorActionCallbacks::GetLightingResolutionMinBSPs())
-													.OnValueChanged_Static(&FLevelEditorActionCallbacks::SetLightingResolutionMinBSPs)
-												]
-												+SHorizontalBox::Slot()
-												.AutoWidth()
-												[
-													SNew(SSpinBox<float>)
-													.MinValue(1.f)
-													.MaxValue(63556.f)
-													.ToolTipText(LOCTEXT( "LightingResolutionBSPsMaxToolTip", "The maximum lightmap resolution of a BSP surface to adjust. When outside of the Min/Max range, the BSP surface will no be altered." ))
-													.Value(FLevelEditorActionCallbacks::GetLightingResolutionMaxBSPs())
-													.OnValueChanged_Static(&FLevelEditorActionCallbacks::SetLightingResolutionMaxBSPs)
-												];
-				InMenuBuilder.AddWidget(BSPs, FText::GetEmpty(), true);
-			}
-			InMenuBuilder.EndSection(); //LevelEditorBuildLightingResolution1
-
-			InMenuBuilder.BeginSection("LevelEditorBuildLightingResolution2", LOCTEXT( "LightingResolutionHeading2", "Select Options" ) );
-			{
-				InMenuBuilder.AddMenuEntry( FLevelEditorCommands::Get().LightingResolution_CurrentLevel );
-				InMenuBuilder.AddMenuEntry( FLevelEditorCommands::Get().LightingResolution_SelectedLevels );
-				InMenuBuilder.AddMenuEntry( FLevelEditorCommands::Get().LightingResolution_AllLoadedLevels );
-				InMenuBuilder.AddMenuEntry( FLevelEditorCommands::Get().LightingResolution_SelectedObjectsOnly );
-			}
-			InMenuBuilder.EndSection();
-
-			InMenuBuilder.BeginSection("LevelEditorBuildLightingResolution3", LOCTEXT( "LightingResolutionHeading3", "Ratio" ) );
-			{
-				TSharedRef<SWidget> Ratio =		SNew(SSpinBox<int32>)
-												.MinValue(0)
-												.MaxValue(400)
-												.ToolTipText(LOCTEXT( "LightingResolutionRatioToolTip", "Ratio to apply (New Resolution = Ratio / 100.0f * CurrentResolution)." ))
-												.Value(FLevelEditorActionCallbacks::GetLightingResolutionRatio())
-												.OnEndSliderMovement_Static(&FLevelEditorActionCallbacks::SetLightingResolutionRatio)
-												.OnValueCommitted_Static(&FLevelEditorActionCallbacks::SetLightingResolutionRatioCommit);
-				InMenuBuilder.AddWidget(Ratio, LOCTEXT( "LightingResolutionRatio", "Ratio" ));
-			}
-			InMenuBuilder.EndSection();
-		}
-
-		/** Generates a lighting info dialogs sub-menu */
-		static void MakeLightingInfoMenu( FMenuBuilder& InMenuBuilder )
-		{
-			InMenuBuilder.BeginSection("LevelEditorBuildLightingInfo", LOCTEXT( "LightingInfoHeading", "Lighting Info Dialogs" ) );
-			{
-				InMenuBuilder.AddSubMenu(
-					LOCTEXT( "LightingToolsSubMenu", "Lighting Tools" ),
-					LOCTEXT( "LightingToolsSubMenu_ToolTip", "Shows the Lighting Tools options." ),
-					FNewMenuDelegate::CreateStatic( &FLightingMenus::MakeLightingToolsMenu ) );
-					
-				InMenuBuilder.AddSubMenu(
-					LOCTEXT( "LightingDensityRenderingSubMenu", "LightMap Density Rendering Options" ),
-					LOCTEXT( "LightingDensityRenderingSubMenu_ToolTip", "Shows the LightMap Density Rendering viewmode options." ),
-					FNewMenuDelegate::CreateStatic( &FLightingMenus::MakeLightingDensityMenu ) );
-
-				InMenuBuilder.AddSubMenu(
-					LOCTEXT( "LightingResolutionAdjustmentSubMenu", "LightMap Resolution Adjustment" ),
-					LOCTEXT( "LightingResolutionAdjustmentSubMenu_ToolTip", "Shows the LightMap Resolution Adjustment options." ),
-					FNewMenuDelegate::CreateStatic( &FLightingMenus::MakeLightingResolutionMenu ) );
-
-				InMenuBuilder.AddMenuEntry( FLevelEditorCommands::Get().LightingStaticMeshInfo, NAME_None, LOCTEXT( "BuildLightingInfo_LightingStaticMeshInfo", "Lighting StaticMesh Info..." ) );
-			}
-			InMenuBuilder.EndSection();
-		}
-	};
-
-	MenuBuilder.BeginSection("LevelEditorLighting", LOCTEXT( "LightingHeading", "Lighting" ) );
-	{
-		MenuBuilder.AddMenuEntry( FLevelEditorCommands::Get().BuildLightingOnly, NAME_None, LOCTEXT( "BuildLightingOnlyHeading", "Build Lighting Only" ) );
-
-		MenuBuilder.AddSubMenu(
-			LOCTEXT( "LightingQualitySubMenu", "Lighting Quality" ),
-			LOCTEXT( "LightingQualitySubMenu_ToolTip", "Allows you to select the quality level for precomputed lighting" ),
-			FNewMenuDelegate::CreateStatic( &FLightingMenus::MakeLightingQualityMenu ) );
-
-		MenuBuilder.AddSubMenu(
-			LOCTEXT( "BuildLightingInfoSubMenu", "Lighting Info" ),
-			LOCTEXT( "BuildLightingInfoSubMenu_ToolTip", "Access the lighting info dialogs" ),
-			FNewMenuDelegate::CreateStatic( &FLightingMenus::MakeLightingInfoMenu ) );
-
-		MenuBuilder.AddMenuEntry( FLevelEditorCommands::Get().LightingBuildOptions_UseErrorColoring );
-		MenuBuilder.AddMenuEntry( FLevelEditorCommands::Get().LightingBuildOptions_ShowLightingStats );
-	}
-	MenuBuilder.EndSection();
-
-	MenuBuilder.BeginSection("LevelEditorReflections", LOCTEXT( "ReflectionHeading", "Reflections" ) );
-	{
-		MenuBuilder.AddMenuEntry( FLevelEditorCommands::Get().BuildReflectionCapturesOnly );
-	}
-	MenuBuilder.EndSection();
-
-	MenuBuilder.BeginSection("LevelEditorVisibility", LOCTEXT( "VisibilityHeading", "Visibility" ) );
-	{
-		MenuBuilder.AddMenuEntry( FLevelEditorCommands::Get().BuildLightingOnly_VisibilityOnly );
-	}
-	MenuBuilder.EndSection();
-
-	MenuBuilder.BeginSection("LevelEditorGeometry", LOCTEXT( "GeometryHeading", "Geometry" ) );
-	{
-		MenuBuilder.AddMenuEntry( FLevelEditorCommands::Get().BuildGeometryOnly );
-		MenuBuilder.AddMenuEntry( FLevelEditorCommands::Get().BuildGeometryOnly_OnlyCurrentLevel );
-	}
-	MenuBuilder.EndSection();
-
-	MenuBuilder.BeginSection("LevelEditorNavigation", LOCTEXT( "NavigationHeading", "Navigation" ) );
-	{
-		MenuBuilder.AddMenuEntry( FLevelEditorCommands::Get().BuildPathsOnly );
-	}
-	MenuBuilder.EndSection();
-
-	MenuBuilder.BeginSection("LevelEditorAutomation", LOCTEXT( "AutomationHeading", "Automation" ) );
-	{
-		MenuBuilder.AddMenuEntry( FLevelEditorCommands::Get().BuildAndSubmitToSourceControl );
-	}
-	MenuBuilder.EndSection();
-
-	// Map Check
-	MenuBuilder.BeginSection("LevelEditorVerification", LOCTEXT( "VerificationHeading", "Verification" ) );
-	{
-		MenuBuilder.AddMenuEntry( FLevelEditorCommands::Get().MapCheck, NAME_None, LOCTEXT("OpenMapCheck", "Map Check") );
-	}
-	MenuBuilder.EndSection();
-
-#undef LOCTEXT_NAMESPACE
-
-	return MenuBuilder.MakeWidget();
-}
-
-static void MakeMaterialQualityLevelMenu( FMenuBuilder& MenuBuilder )
-{
-	MenuBuilder.BeginSection("LevelEditorMaterialQualityLevel", NSLOCTEXT( "LevelToolBarViewMenu", "MaterialQualityLevelHeading", "Material Quality Level" ) );
-	{
-		MenuBuilder.AddMenuEntry( FLevelEditorCommands::Get().MaterialQualityLevel_Low );
-		MenuBuilder.AddMenuEntry( FLevelEditorCommands::Get().MaterialQualityLevel_High );
-	}
-	MenuBuilder.EndSection();
-}
-
-static void MakeShaderModelPreviewMenu(FMenuBuilder& MenuBuilder)
-{
-	MenuBuilder.BeginSection("LevelEditorShaderModelPreview", NSLOCTEXT("LevelToolBarViewMenu", "FeatureLevelPreviewHeading", "Feature Level Preview"));
-	{
-		for (int32 i = GMaxRHIFeatureLevel; i >= 0; --i)
-		{
-			if (i != ERHIFeatureLevel::ES3_1)
-			{
-				MenuBuilder.AddMenuEntry(FLevelEditorCommands::Get().FeatureLevelPreview[i]);
-			}
-		}
-	}
-	MenuBuilder.EndSection();
-}
-
-static void MakeScalabilityMenu( FMenuBuilder& MenuBuilder )
-{
-	MenuBuilder.AddWidget(SNew(SScalabilitySettings), FText(), true);
-}
-
-static void MakePreviewSettingsMenu( FMenuBuilder& MenuBuilder )
-{
-#define LOCTEXT_NAMESPACE "LevelToolBarViewMenu"
-
-	MenuBuilder.BeginSection("LevelEditorPreview", LOCTEXT("PreviewHeading", "Previewing"));
-	{
-		MenuBuilder.AddMenuEntry(FLevelEditorCommands::Get().DrawBrushMarkerPolys);
-		MenuBuilder.AddMenuEntry(FLevelEditorCommands::Get().OnlyLoadVisibleInPIE);
-		MenuBuilder.AddMenuEntry(FLevelEditorCommands::Get().ToggleParticleSystemLOD);
-		MenuBuilder.AddMenuEntry(FLevelEditorCommands::Get().ToggleParticleSystemHelpers);
-		MenuBuilder.AddMenuEntry(FLevelEditorCommands::Get().ToggleFreezeParticleSimulation);
-		MenuBuilder.AddMenuEntry(FLevelEditorCommands::Get().ToggleLODViewLocking);
-		MenuBuilder.AddMenuEntry(FLevelEditorCommands::Get().LevelStreamingVolumePrevis);
-	}
-	MenuBuilder.EndSection();
-#undef LOCTEXT_NAMESPACE
-}
-
-TSharedRef< SWidget > FLevelEditorToolBar::GenerateQuickSettingsMenu( TSharedRef<FUICommandList> InCommandList )
-{
-#define LOCTEXT_NAMESPACE "LevelToolBarViewMenu"
-
-	// Get all menu extenders for this context menu from the level editor module
-	FLevelEditorModule& LevelEditorModule = FModuleManager::GetModuleChecked<FLevelEditorModule>( TEXT("LevelEditor") );
-	TArray<FLevelEditorModule::FLevelEditorMenuExtender> MenuExtenderDelegates = LevelEditorModule.GetAllLevelEditorToolbarViewMenuExtenders();
-
-	TArray<TSharedPtr<FExtender>> Extenders;
-	for (int32 i = 0; i < MenuExtenderDelegates.Num(); ++i)
-	{
-		if (MenuExtenderDelegates[i].IsBound())
-		{
-			Extenders.Add(MenuExtenderDelegates[i].Execute(InCommandList));
-		}
-	}
-	TSharedPtr<FExtender> MenuExtender = FExtender::Combine(Extenders);
-
-	const bool bShouldCloseWindowAfterMenuSelection = true;
-	FMenuBuilder MenuBuilder( bShouldCloseWindowAfterMenuSelection, InCommandList, MenuExtender );
-
-	struct Local
-	{
-		static void OpenSettings(FName ContainerName, FName CategoryName, FName SectionName)
-		{
-			FModuleManager::LoadModuleChecked<ISettingsModule>("Settings").ShowViewer(ContainerName, CategoryName, SectionName);
-		}
-	};
-
-	MenuBuilder.BeginSection("ProjectSettingsSection", LOCTEXT("ProjectSettings","Game Specific Settings") );
-	{
-
-		MenuBuilder.AddMenuEntry(FLevelEditorCommands::Get().WorldProperties);
-
-		MenuBuilder.AddMenuEntry(
-			LOCTEXT("ProjectSettingsMenuLabel", "Project Settings..."),
-			LOCTEXT("ProjectSettingsMenuToolTip", "Change the settings of the currently loaded project"),
-			FSlateIcon(),
-			FUIAction(FExecuteAction::CreateStatic(&Local::OpenSettings, FName("Project"), FName("Project"), FName("General")))
-			);
-	}
-	MenuBuilder.EndSection();
-
-	MenuBuilder.BeginSection("LevelEditorSelection", LOCTEXT("SelectionHeading","Selection") );
-	{
-		MenuBuilder.AddMenuEntry( FLevelEditorCommands::Get().AllowTranslucentSelection );
-		MenuBuilder.AddMenuEntry( FLevelEditorCommands::Get().AllowGroupSelection );
-		MenuBuilder.AddMenuEntry( FLevelEditorCommands::Get().StrictBoxSelect );
-		MenuBuilder.AddMenuEntry( FLevelEditorCommands::Get().TransparentBoxSelect );
-		MenuBuilder.AddMenuEntry( FLevelEditorCommands::Get().ShowTransformWidget );
-	}
-	MenuBuilder.EndSection();
-
-	MenuBuilder.BeginSection("LevelEditorScalability", LOCTEXT("ScalabilityHeading", "Scalability") );
-	{
-		MenuBuilder.AddSubMenu(
-			LOCTEXT( "ScalabilitySubMenu", "Engine Scalability Settings" ),
-			LOCTEXT( "ScalabilitySubMenu_ToolTip", "Open the engine scalability settings" ),
-			FNewMenuDelegate::CreateStatic( &MakeScalabilityMenu ) );
-
-		MenuBuilder.AddSubMenu(
-			LOCTEXT( "MaterialQualityLevelSubMenu", "Material Quality Level" ),
-			LOCTEXT( "MaterialQualityLevelSubMenu_ToolTip", "Sets the value of the CVar \"r.MaterialQualityLevel\" (low=0, high=1). This affects materials via the QualitySwitch material expression." ),
-			FNewMenuDelegate::CreateStatic( &MakeMaterialQualityLevelMenu ) );
-
-		if (GetDefault<UEditorExperimentalSettings>()->bFeatureLevelPreview)
-		{
-			MenuBuilder.AddSubMenu(
-				LOCTEXT("FeatureLevelPreviewSubMenu", "Feature Level Preview"),
-				LOCTEXT("FeatureLevelPreviewSubMenu_ToolTip", "Sets the feature level preview mode"),
-				FNewMenuDelegate::CreateStatic(&MakeShaderModelPreviewMenu));
-		}
-	}
-	MenuBuilder.EndSection();
-
-	MenuBuilder.BeginSection("LevelEditorAudio", LOCTEXT("AudioHeading", "Real Time Audio") );
-	{
-		TSharedRef<SWidget> VolumeItem = SNew(SHorizontalBox)
-											+SHorizontalBox::Slot()
-											.FillWidth(0.9f)
-											.Padding( FMargin(2.0f, 0.0f, 0.0f, 0.0f) )
-											[
-												SNew(SVolumeControl)
-												.ToolTipText_Static(&FLevelEditorActionCallbacks::GetAudioVolumeToolTip)
-												.Volume_Static(&FLevelEditorActionCallbacks::GetAudioVolume)
-												.OnVolumeChanged_Static(&FLevelEditorActionCallbacks::OnAudioVolumeChanged)
-												.Muted_Static(&FLevelEditorActionCallbacks::GetAudioMuted)
-												.OnMuteChanged_Static(&FLevelEditorActionCallbacks::OnAudioMutedChanged)
-											]
-											+SHorizontalBox::Slot()
-											.FillWidth(0.1f);
-
-		MenuBuilder.AddWidget(VolumeItem, LOCTEXT("VolumeControlLabel","Volume"));
-	}
-	MenuBuilder.EndSection();
-
-	MenuBuilder.BeginSection( "Snapping", LOCTEXT("SnappingHeading","Snapping") );
-	{
-		MenuBuilder.AddMenuEntry( FLevelEditorCommands::Get().EnableActorSnap );
-		TSharedRef<SWidget> SnapItem = 
-		SNew(SHorizontalBox)
-	          +SHorizontalBox::Slot()
-	          .FillWidth(0.9f)
-	          [
-		          SNew(SSlider)
-		          .ToolTipText_Static(&FLevelEditorActionCallbacks::GetActorSnapTooltip)
-		          .Value_Static(&FLevelEditorActionCallbacks::GetActorSnapSetting)
-		          .OnValueChanged_Static(&FLevelEditorActionCallbacks::SetActorSnapSetting)
-	          ]
-	          +SHorizontalBox::Slot()
-	          .FillWidth(0.1f);
-		MenuBuilder.AddWidget(SnapItem, LOCTEXT("ActorSnapLabel","Distance"));
-
-		MenuBuilder.AddMenuEntry( FLevelEditorCommands::Get().ToggleSocketSnapping );
-		MenuBuilder.AddMenuEntry( FLevelEditorCommands::Get().EnableVertexSnap );
-	}
-	MenuBuilder.EndSection();
-
-	MenuBuilder.BeginSection("LevelEditorViewport", LOCTEXT("ViewportHeading", "Viewport") );
-	{
-		MenuBuilder.AddMenuEntry( FLevelEditorCommands::Get().ToggleHideViewportUI );
-
-		MenuBuilder.AddSubMenu( LOCTEXT("PreviewMenu", "Previewing"), LOCTEXT("PreviewMenuTooltip","Game Preview Settings"), FNewMenuDelegate::CreateStatic( &MakePreviewSettingsMenu ) );
-	}
-	MenuBuilder.EndSection();
-
-#undef LOCTEXT_NAMESPACE
-
-	return MenuBuilder.MakeWidget();
-}
-
-TSharedRef< SWidget > FLevelEditorToolBar::GenerateCompileMenuContent( TSharedRef<FUICommandList> InCommandList )
-{
-#define LOCTEXT_NAMESPACE "LevelToolBarViewMenu"
-
-	// Get all menu extenders for this context menu from the level editor module
-	FLevelEditorModule& LevelEditorModule = FModuleManager::GetModuleChecked<FLevelEditorModule>( TEXT("LevelEditor") );
-	TArray<FLevelEditorModule::FLevelEditorMenuExtender> MenuExtenderDelegates = LevelEditorModule.GetAllLevelEditorToolbarCompileMenuExtenders();
-
-	TArray<TSharedPtr<FExtender>> Extenders;
-	for (int32 i = 0; i < MenuExtenderDelegates.Num(); ++i)
@@ -785,200 +1006,283 @@
 	const bool bShouldCloseWindowAfterMenuSelection = true;
 	FMenuBuilder MenuBuilder( bShouldCloseWindowAfterMenuSelection, InCommandList, MenuExtender );
 
-	if (LevelEditorModule.CanBeRecompiled())
+	struct FLightingMenus
 	{
-		MenuBuilder.AddMenuEntry( FLevelEditorCommands::Get().RecompileLevelEditor );
-		MenuBuilder.AddMenuEntry( FLevelEditorCommands::Get().ReloadLevelEditor );
-	}
 
-	MenuBuilder.AddMenuEntry( FLevelEditorCommands::Get().RecompileGameCode );
-
-#undef LOCTEXT_NAMESPACE
-
-	return MenuBuilder.MakeWidget();
-}
-
-TSharedRef< SWidget > FLevelEditorToolBar::GenerateOpenBlueprintMenuContent( TSharedRef<FUICommandList> InCommandList, TWeakPtr< SLevelEditor > InLevelEditor )
-{
-#define LOCTEXT_NAMESPACE "LevelToolBarViewMenu"
-
-	struct FBlueprintMenus
-	{
-		/** Generates a sub-level Blueprints sub-menu */
-		static void MakeSubLevelsMenu(FMenuBuilder& InMenuBuilder, TWeakPtr< SLevelEditor > InLvlEditor)
+		/** Generates a lighting quality sub-menu */
+		static void MakeLightingQualityMenu( FMenuBuilder& InMenuBuilder )
 		{
-			FSlateIcon EditBP(FEditorStyle::Get().GetStyleSetName(), TEXT("LevelEditor.OpenLevelBlueprint"));
-
-			InMenuBuilder.BeginSection(NAME_None, LOCTEXT("SubLevelsHeading", "Sub-Level Blueprints"));
+			InMenuBuilder.BeginSection("LevelEditorBuildLightingQuality", LOCTEXT( "LightingQualityHeading", "Quality Level" ) );
 			{
-				UWorld* World = InLvlEditor.Pin()->GetWorld();
-				for (int32 iLevel = 0; iLevel < World->GetNumLevels(); iLevel++)
-				{
-					ULevel* Level = World->GetLevel(iLevel);
-					if (Level != NULL && Level->GetOutermost() != NULL)
-					{
-						if (!Level->IsPersistentLevel())
-						{
-							FUIAction UIAction
-								(
-								FExecuteAction::CreateStatic(&FLevelEditorToolBar::OnOpenSubLevelBlueprint, Level)
-								);
-
-							FText DisplayName = FText::Format(LOCTEXT("SubLevelBlueprintItem", "Edit {LevelName}"), FText::FromString(FPaths::GetCleanFilename(Level->GetOutermost()->GetName())));
-							InMenuBuilder.AddMenuEntry(DisplayName, FText::GetEmpty(), EditBP, UIAction);
-						}
-					}
-				}
+				InMenuBuilder.AddMenuEntry( FLevelEditorCommands::Get().LightingQuality_Production );
+				InMenuBuilder.AddMenuEntry( FLevelEditorCommands::Get().LightingQuality_High );
+				InMenuBuilder.AddMenuEntry( FLevelEditorCommands::Get().LightingQuality_Medium );
+				InMenuBuilder.AddMenuEntry( FLevelEditorCommands::Get().LightingQuality_Preview );
 			}
 			InMenuBuilder.EndSection();
 		}
 
-		/** Handle BP being selected from popup picker */
-		static void OnBPSelected(const class FAssetData& AssetData)
+		/** Generates a lighting tools sub-menu */
+		static void MakeLightingToolsMenu( FMenuBuilder& InMenuBuilder )
 		{
-			UBlueprint* SelectedBP = Cast<UBlueprint>(AssetData.GetAsset());
-			if(SelectedBP)
+			InMenuBuilder.BeginSection("LevelEditorBuildLightingTools", LOCTEXT( "LightingToolsHeading", "Light Environments" ) );
 			{
-				FAssetEditorManager::Get().OpenEditorForAsset(SelectedBP);
-			}
-		}
-
-
-		/** Generates 'open blueprint' sub-menu */
-		static void MakeOpenClassBPMenu(FMenuBuilder& InMenuBuilder)
-		{
-			FContentBrowserModule& ContentBrowserModule = FModuleManager::Get().LoadModuleChecked<FContentBrowserModule>(TEXT("ContentBrowser"));
-
-			// Configure filter for asset picker
-			FAssetPickerConfig Config;
-			Config.Filter.ClassNames.Add(UBlueprint::StaticClass()->GetFName());
-			Config.InitialAssetViewType = EAssetViewType::List;
-			Config.ThumbnailScale = 0; // make thumbnails as small as possible
-			Config.OnAssetSelected = FOnAssetSelected::CreateStatic(&FBlueprintMenus::OnBPSelected);
-			Config.bAllowDragging = false;
-			// Don't show stuff in Engine
-			Config.Filter.PackagePaths.Add("/Game");
-			Config.Filter.bRecursivePaths = true;
-
-			TSharedRef<SWidget> Widget = 
-				SNew(SBox)
-				.WidthOverride(300.f)
-				.HeightOverride(300.f)
-				[
-					ContentBrowserModule.Get().CreateAssetPicker(Config)
-				];
-		
-
-			InMenuBuilder.BeginSection(NAME_None, LOCTEXT("BrowseHeader", "Browse"));
-			{
-				InMenuBuilder.AddWidget(Widget, FText::GetEmpty());
+				InMenuBuilder.AddMenuEntry( FLevelEditorCommands::Get().LightingTools_ShowBounds );
+				InMenuBuilder.AddMenuEntry( FLevelEditorCommands::Get().LightingTools_ShowTraces );
+				InMenuBuilder.AddMenuEntry( FLevelEditorCommands::Get().LightingTools_ShowDirectOnly );
+				InMenuBuilder.AddMenuEntry( FLevelEditorCommands::Get().LightingTools_ShowIndirectOnly );
+				InMenuBuilder.AddMenuEntry( FLevelEditorCommands::Get().LightingTools_ShowIndirectSamples );
 			}
 			InMenuBuilder.EndSection();
 		}
+
+		/** Generates a lighting density sub-menu */
+		static void MakeLightingDensityMenu( FMenuBuilder& InMenuBuilder )
+		{
+			InMenuBuilder.BeginSection("LevelEditorBuildLightingDensity", LOCTEXT( "LightingDensityHeading", "Density Rendering" ) );
+			{
+				TSharedRef<SWidget> Ideal =		SNew(SHorizontalBox)
+												+SHorizontalBox::Slot()
+												.Padding( FMargin( 27.0f, 0.0f, 0.0f, 0.0f ) )
+												.FillWidth(1.0f)
+												[
+													SNew(SSpinBox<float>)
+													.MinValue(0.f)
+													.MaxValue(100.f)
+													.Value(FLevelEditorActionCallbacks::GetLightingDensityIdeal())
+													.OnValueChanged_Static(&FLevelEditorActionCallbacks::SetLightingDensityIdeal)
+												];
+				InMenuBuilder.AddWidget(Ideal, LOCTEXT("LightingDensity_Ideal","Ideal Density"));
+				
+				TSharedRef<SWidget> Maximum =	SNew(SHorizontalBox)
+												+SHorizontalBox::Slot()
+												.FillWidth(1.0f)
+												[
+													SNew(SSpinBox<float>)
+													.MinValue(0.01f)
+													.MaxValue(100.01f)
+													.Value(FLevelEditorActionCallbacks::GetLightingDensityMaximum())
+													.OnValueChanged_Static(&FLevelEditorActionCallbacks::SetLightingDensityMaximum)
+												];
+				InMenuBuilder.AddWidget(Maximum, LOCTEXT("LightingDensity_Maximum","Maximum Density"));
+
+				TSharedRef<SWidget> ClrScale =	SNew(SHorizontalBox)
+												+SHorizontalBox::Slot()
+												.Padding( FMargin( 35.0f, 0.0f, 0.0f, 0.0f ) )
+												.FillWidth(1.0f)
+												[
+													SNew(SSpinBox<float>)
+													.MinValue(0.f)
+													.MaxValue(10.f)
+													.Value(FLevelEditorActionCallbacks::GetLightingDensityColorScale())
+													.OnValueChanged_Static(&FLevelEditorActionCallbacks::SetLightingDensityColorScale)
+												];
+				InMenuBuilder.AddWidget(ClrScale, LOCTEXT("LightingDensity_ColorScale","Color Scale"));
+
+				TSharedRef<SWidget> GrayScale =	SNew(SHorizontalBox)
+												+SHorizontalBox::Slot()
+												.Padding( FMargin( 11.0f, 0.0f, 0.0f, 0.0f ) )
+												.FillWidth(1.0f)
+												[
+													SNew(SSpinBox<float>)
+													.MinValue(0.f)
+													.MaxValue(10.f)
+													.Value(FLevelEditorActionCallbacks::GetLightingDensityGrayscaleScale())
+													.OnValueChanged_Static(&FLevelEditorActionCallbacks::SetLightingDensityGrayscaleScale)
+												];
+				InMenuBuilder.AddWidget(GrayScale, LOCTEXT("LightingDensity_GrayscaleScale","Grayscale Scale"));
+
+				InMenuBuilder.AddMenuEntry( FLevelEditorCommands::Get().LightingDensity_RenderGrayscale );
+			}
+			InMenuBuilder.EndSection();
+		}
+
+		/** Generates a lighting resolution sub-menu */
+		static void MakeLightingResolutionMenu( FMenuBuilder& InMenuBuilder )
+		{
+			InMenuBuilder.BeginSection("LevelEditorBuildLightingResolution1", LOCTEXT( "LightingResolutionHeading1", "Primitive Types" ) );
+			{
+				TSharedRef<SWidget> Meshes =	SNew(SHorizontalBox)
+												+SHorizontalBox::Slot()
+												.AutoWidth()
+												[
+													SNew( SCheckBox )
+													.Style( FEditorStyle::Get(), "Menu.CheckBox" )
+													.ToolTipText(LOCTEXT( "StaticMeshesToolTip", "Static Meshes will be adjusted if checked." ))
+													.IsChecked_Static(&FLevelEditorActionCallbacks::IsLightingResolutionStaticMeshesChecked)
+													.OnCheckStateChanged_Static(&FLevelEditorActionCallbacks::SetLightingResolutionStaticMeshes)
+													.Content()
+													[
+														SNew( STextBlock )
+														.Text( LOCTEXT("StaticMeshes", "Static Meshes") )
+													]
+												]
+												+SHorizontalBox::Slot()
+												.AutoWidth()
+												.Padding( FMargin( 4.0f, 0.0f, 11.0f, 0.0f ) )
+												[
+													SNew(SSpinBox<float>)
+													.MinValue(4.f)
+													.MaxValue(4096.f)
+													.ToolTipText(LOCTEXT( "LightingResolutionStaticMeshesMinToolTip", "The minimum lightmap resolution for static mesh adjustments. Anything outside of Min/Max range will not be touched when adjusting." ))
+													.Value(FLevelEditorActionCallbacks::GetLightingResolutionMinSMs())
+													.OnValueChanged_Static(&FLevelEditorActionCallbacks::SetLightingResolutionMinSMs)
+												]
+												+SHorizontalBox::Slot()
+												.AutoWidth()
+												[
+													SNew(SSpinBox<float>)
+													.MinValue(4.f)
+													.MaxValue(4096.f)
+													.ToolTipText(LOCTEXT( "LightingResolutionStaticMeshesMaxToolTip", "The maximum lightmap resolution for static mesh adjustments. Anything outside of Min/Max range will not be touched when adjusting." ))
+													.Value(FLevelEditorActionCallbacks::GetLightingResolutionMaxSMs())
+													.OnValueChanged_Static(&FLevelEditorActionCallbacks::SetLightingResolutionMaxSMs)
+												];
+				InMenuBuilder.AddWidget(Meshes, FText::GetEmpty(), true);
+				
+				TSharedRef<SWidget> BSPs =		SNew(SHorizontalBox)
+												+SHorizontalBox::Slot()
+												.AutoWidth()
+												[
+													SNew( SCheckBox )
+													.Style(FEditorStyle::Get(), "Menu.CheckBox")
+													.ToolTipText(LOCTEXT( "BSPSurfacesToolTip", "BSP Surfaces will be adjusted if checked." ))
+													.IsChecked_Static(&FLevelEditorActionCallbacks::IsLightingResolutionBSPSurfacesChecked)
+													.OnCheckStateChanged_Static(&FLevelEditorActionCallbacks::SetLightingResolutionBSPSurfaces)
+													.Content()
+													[
+														SNew( STextBlock )
+														.Text( LOCTEXT("BSPSurfaces", "BSP Surfaces") )
+													]
+												]
+												+SHorizontalBox::Slot()
+												.AutoWidth()
+												.Padding( FMargin( 6.0f, 0.0f, 4.0f, 0.0f ) )
+												[
+													SNew(SSpinBox<float>)
+													.MinValue(1.f)
+													.MaxValue(63556.f)
+													.ToolTipText(LOCTEXT( "LightingResolutionBSPsMinToolTip", "The minimum lightmap resolution of a BSP surface to adjust. When outside of the Min/Max range, the BSP surface will no be altered." ))
+													.Value(FLevelEditorActionCallbacks::GetLightingResolutionMinBSPs())
+													.OnValueChanged_Static(&FLevelEditorActionCallbacks::SetLightingResolutionMinBSPs)
+												]
+												+SHorizontalBox::Slot()
+												.AutoWidth()
+												[
+													SNew(SSpinBox<float>)
+													.MinValue(1.f)
+													.MaxValue(63556.f)
+													.ToolTipText(LOCTEXT( "LightingResolutionBSPsMaxToolTip", "The maximum lightmap resolution of a BSP surface to adjust. When outside of the Min/Max range, the BSP surface will no be altered." ))
+													.Value(FLevelEditorActionCallbacks::GetLightingResolutionMaxBSPs())
+													.OnValueChanged_Static(&FLevelEditorActionCallbacks::SetLightingResolutionMaxBSPs)
+												];
+				InMenuBuilder.AddWidget(BSPs, FText::GetEmpty(), true);
+			}
+			InMenuBuilder.EndSection(); //LevelEditorBuildLightingResolution1
+
+			InMenuBuilder.BeginSection("LevelEditorBuildLightingResolution2", LOCTEXT( "LightingResolutionHeading2", "Select Options" ) );
+			{
+				InMenuBuilder.AddMenuEntry( FLevelEditorCommands::Get().LightingResolution_CurrentLevel );
+				InMenuBuilder.AddMenuEntry( FLevelEditorCommands::Get().LightingResolution_SelectedLevels );
+				InMenuBuilder.AddMenuEntry( FLevelEditorCommands::Get().LightingResolution_AllLoadedLevels );
+				InMenuBuilder.AddMenuEntry( FLevelEditorCommands::Get().LightingResolution_SelectedObjectsOnly );
+			}
+			InMenuBuilder.EndSection();
+
+			InMenuBuilder.BeginSection("LevelEditorBuildLightingResolution3", LOCTEXT( "LightingResolutionHeading3", "Ratio" ) );
+			{
+				TSharedRef<SWidget> Ratio =		SNew(SSpinBox<int32>)
+												.MinValue(0)
+												.MaxValue(400)
+												.ToolTipText(LOCTEXT( "LightingResolutionRatioToolTip", "Ratio to apply (New Resolution = Ratio / 100.0f * CurrentResolution)." ))
+												.Value(FLevelEditorActionCallbacks::GetLightingResolutionRatio())
+												.OnEndSliderMovement_Static(&FLevelEditorActionCallbacks::SetLightingResolutionRatio)
+												.OnValueCommitted_Static(&FLevelEditorActionCallbacks::SetLightingResolutionRatioCommit);
+				InMenuBuilder.AddWidget(Ratio, LOCTEXT( "LightingResolutionRatio", "Ratio" ));
+			}
+			InMenuBuilder.EndSection();
+		}
+
+		/** Generates a lighting info dialogs sub-menu */
+		static void MakeLightingInfoMenu( FMenuBuilder& InMenuBuilder )
+		{
+			InMenuBuilder.BeginSection("LevelEditorBuildLightingInfo", LOCTEXT( "LightingInfoHeading", "Lighting Info Dialogs" ) );
+			{
+				InMenuBuilder.AddSubMenu(
+					LOCTEXT( "LightingToolsSubMenu", "Lighting Tools" ),
+					LOCTEXT( "LightingToolsSubMenu_ToolTip", "Shows the Lighting Tools options." ),
+					FNewMenuDelegate::CreateStatic( &FLightingMenus::MakeLightingToolsMenu ) );
+					
+				InMenuBuilder.AddSubMenu(
+					LOCTEXT( "LightingDensityRenderingSubMenu", "LightMap Density Rendering Options" ),
+					LOCTEXT( "LightingDensityRenderingSubMenu_ToolTip", "Shows the LightMap Density Rendering viewmode options." ),
+					FNewMenuDelegate::CreateStatic( &FLightingMenus::MakeLightingDensityMenu ) );
+
+				InMenuBuilder.AddSubMenu(
+					LOCTEXT( "LightingResolutionAdjustmentSubMenu", "LightMap Resolution Adjustment" ),
+					LOCTEXT( "LightingResolutionAdjustmentSubMenu_ToolTip", "Shows the LightMap Resolution Adjustment options." ),
+					FNewMenuDelegate::CreateStatic( &FLightingMenus::MakeLightingResolutionMenu ) );
+
+				InMenuBuilder.AddMenuEntry( FLevelEditorCommands::Get().LightingStaticMeshInfo, NAME_None, LOCTEXT( "BuildLightingInfo_LightingStaticMeshInfo", "Lighting StaticMesh Info..." ) );
+			}
+			InMenuBuilder.EndSection();
+		}
 	};
 
+	MenuBuilder.BeginSection("LevelEditorLighting", LOCTEXT( "LightingHeading", "Lighting" ) );
+	{
+		MenuBuilder.AddMenuEntry( FLevelEditorCommands::Get().BuildLightingOnly, NAME_None, LOCTEXT( "BuildLightingOnlyHeading", "Build Lighting Only" ) );
 
-	const bool bShouldCloseWindowAfterMenuSelection = true;
-	FMenuBuilder MenuBuilder( bShouldCloseWindowAfterMenuSelection, InCommandList );
+		MenuBuilder.AddSubMenu(
+			LOCTEXT( "LightingQualitySubMenu", "Lighting Quality" ),
+			LOCTEXT( "LightingQualitySubMenu_ToolTip", "Allows you to select the quality level for precomputed lighting" ),
+			FNewMenuDelegate::CreateStatic( &FLightingMenus::MakeLightingQualityMenu ) );
 
-	MenuBuilder.BeginSection(NAME_None, LOCTEXT("LevelScriptBlueprints", "Level Blueprints"));
-	{
-		MenuBuilder.AddMenuEntry( FLevelEditorCommands::Get().OpenLevelBlueprint );
+		MenuBuilder.AddSubMenu(
+			LOCTEXT( "BuildLightingInfoSubMenu", "Lighting Info" ),
+			LOCTEXT( "BuildLightingInfoSubMenu_ToolTip", "Access the lighting info dialogs" ),
+			FNewMenuDelegate::CreateStatic( &FLightingMenus::MakeLightingInfoMenu ) );
 
-		// If there are any sub-levels, display the sub-menu. A single level means there is only the persistent level
-		UWorld* World = InLevelEditor.Pin()->GetWorld();
-		if(World->GetNumLevels() > 1)
-		{
-			MenuBuilder.AddSubMenu(
-				LOCTEXT( "SubLevelsSubMenu", "Sub-Levels" ),
-				LOCTEXT( "SubLevelsSubMenu_ToolTip", "Shows available sub-level Blueprints that can be edited." ),
-				FNewMenuDelegate::CreateStatic( &FBlueprintMenus::MakeSubLevelsMenu, InLevelEditor ), 
-				FUIAction(), NAME_None, EUserInterfaceActionType::Button, false, FSlateIcon(FEditorStyle::Get().GetStyleSetName(), TEXT("LevelEditor.OpenLevelBlueprint")) );
-		}
+		MenuBuilder.AddMenuEntry( FLevelEditorCommands::Get().LightingBuildOptions_UseErrorColoring );
+		MenuBuilder.AddMenuEntry( FLevelEditorCommands::Get().LightingBuildOptions_ShowLightingStats );
 	}
 	MenuBuilder.EndSection();
 
-	MenuBuilder.BeginSection(NAME_None, LOCTEXT("WorldSettingsClasses", "World Settings Classes"));
+	MenuBuilder.BeginSection("LevelEditorReflections", LOCTEXT( "ReflectionHeading", "Reflections" ) );
 	{
-		AWorldSettings* WorldSettings = InLevelEditor.Pin()->GetWorld()->GetWorldSettings();
-
-		// Game Mode
-		AGameMode* ActiveGameMode = NULL;
-
-		// If the WorldSettings' default game mode is valid, pull out the default object so we can edit it
-		if(WorldSettings->DefaultGameMode)
-		{
-			ActiveGameMode = Cast<AGameMode>(InLevelEditor.Pin()->GetWorld()->GetWorldSettings()->DefaultGameMode->GetDefaultObject());
-		}
-
-		LevelEditorActionHelpers::FBlueprintMenuSettings GameModeMenuSettings;
-		GameModeMenuSettings.EditCommand = FLevelEditorCommands::Get().OpenGameModeBlueprint;
-		GameModeMenuSettings.OnCreateClassPicked = FOnClassPicked::CreateStatic( &FLevelEditorToolBar::OnCreateGameModeClassPicked, InLevelEditor );;
-		GameModeMenuSettings.OnSelectClassPicked = FOnClassPicked::CreateStatic( &FLevelEditorToolBar::OnSelectGameModeClassPicked, InLevelEditor );;
-		GameModeMenuSettings.CurrentClass = GetActiveGameModeClass(InLevelEditor);
-		GameModeMenuSettings.RootClass = AGameMode::StaticClass();
-
-		MenuBuilder.AddSubMenu( GetOpenGameModeBlueprintLabel(InLevelEditor), LOCTEXT("GameMode_Tooltip", "Select, edit, or create a new GameMode blueprint for the world"), FNewMenuDelegate::CreateStatic( &LevelEditorActionHelpers::GetBlueprintSettingsSubMenu, GameModeMenuSettings) );
-
-		if(ActiveGameMode && ActiveGameMode->GetClass()->ClassGeneratedBy)
-		{
-			// Game State
-			LevelEditorActionHelpers::FBlueprintMenuSettings GameStateMenuSettings;
-			GameStateMenuSettings.EditCommand = FLevelEditorCommands::Get().OpenGameStateBlueprint;
-			GameStateMenuSettings.OnCreateClassPicked = FOnClassPicked::CreateStatic( &FLevelEditorToolBar::OnCreateGameStateClassPicked, InLevelEditor );;
-			GameStateMenuSettings.OnSelectClassPicked = FOnClassPicked::CreateStatic( &FLevelEditorToolBar::OnSelectGameStateClassPicked, InLevelEditor );;
-			GameStateMenuSettings.CurrentClass = GetActiveGameStateClass(InLevelEditor);
-			GameStateMenuSettings.RootClass = AGameState::StaticClass();
-
-			MenuBuilder.AddSubMenu( GetOpenGameStateBlueprintLabel(InLevelEditor), LOCTEXT("GameState_Tooltip", "Select, edit, or create a new GameState blueprint for the world"), FNewMenuDelegate::CreateStatic( &LevelEditorActionHelpers::GetBlueprintSettingsSubMenu, GameStateMenuSettings ) );
-		
-			// Pawn
-			LevelEditorActionHelpers::FBlueprintMenuSettings PawnMenuSettings;
-			PawnMenuSettings.EditCommand = FLevelEditorCommands::Get().OpenDefaultPawnBlueprint;
-			PawnMenuSettings.OnCreateClassPicked = FOnClassPicked::CreateStatic( &FLevelEditorToolBar::OnCreatePawnClassPicked, InLevelEditor );;
-			PawnMenuSettings.OnSelectClassPicked = FOnClassPicked::CreateStatic( &FLevelEditorToolBar::OnSelectPawnClassPicked, InLevelEditor );;
-			PawnMenuSettings.CurrentClass = GetActivePawnClass(InLevelEditor);
-			PawnMenuSettings.RootClass = APawn::StaticClass();
-
-			MenuBuilder.AddSubMenu( GetOpenPawnBlueprintLabel(InLevelEditor), LOCTEXT("Pawn_Tooltip", "Select, edit, or create a new Pawn blueprint for the world"), FNewMenuDelegate::CreateStatic( &LevelEditorActionHelpers::GetBlueprintSettingsSubMenu, PawnMenuSettings ) );
-
-			// HUD
-			LevelEditorActionHelpers::FBlueprintMenuSettings HUDMenuSettings;
-			HUDMenuSettings.EditCommand = FLevelEditorCommands::Get().OpenHUDBlueprint;
-			HUDMenuSettings.OnCreateClassPicked = FOnClassPicked::CreateStatic( &FLevelEditorToolBar::OnCreateHUDClassPicked, InLevelEditor );;
-			HUDMenuSettings.OnSelectClassPicked = FOnClassPicked::CreateStatic( &FLevelEditorToolBar::OnSelectHUDClassPicked, InLevelEditor );;
-			HUDMenuSettings.CurrentClass = GetActiveHUDClass(InLevelEditor);
-			HUDMenuSettings.RootClass = AHUD::StaticClass();
-
-			MenuBuilder.AddSubMenu( GetOpenHUDBlueprintLabel(InLevelEditor), LOCTEXT("HUD_Tooltip", "Select, edit, or create a new HUD blueprint for the world"), FNewMenuDelegate::CreateStatic( &LevelEditorActionHelpers::GetBlueprintSettingsSubMenu, HUDMenuSettings ) );
-
-			// Player Controller
-			LevelEditorActionHelpers::FBlueprintMenuSettings PlayerControllerMenuSettings;
-			PlayerControllerMenuSettings.EditCommand = FLevelEditorCommands::Get().OpenPlayerControllerBlueprint;
-			PlayerControllerMenuSettings.OnCreateClassPicked = FOnClassPicked::CreateStatic( &FLevelEditorToolBar::OnCreatePlayerControllerClassPicked, InLevelEditor );;
-			PlayerControllerMenuSettings.OnSelectClassPicked = FOnClassPicked::CreateStatic( &FLevelEditorToolBar::OnSelectPlayerControllerClassPicked, InLevelEditor );;
-			PlayerControllerMenuSettings.CurrentClass = GetActivePlayerControllerClass(InLevelEditor);
-			PlayerControllerMenuSettings.RootClass = APlayerController::StaticClass();
-
-			MenuBuilder.AddSubMenu( GetOpenPlayerControllerBlueprintLabel(InLevelEditor), LOCTEXT("PlayerController_Tooltip", "Select, edit, or create a new Player controller blueprint for the world"), FNewMenuDelegate::CreateStatic( &LevelEditorActionHelpers::GetBlueprintSettingsSubMenu, PlayerControllerMenuSettings ) );
-
-		}
+		MenuBuilder.AddMenuEntry( FLevelEditorCommands::Get().BuildReflectionCapturesOnly );
 	}
 	MenuBuilder.EndSection();
 
-	MenuBuilder.BeginSection(NAME_None, LOCTEXT("ClassBlueprints", "Class Blueprints"));
+	MenuBuilder.BeginSection("LevelEditorVisibility", LOCTEXT( "VisibilityHeading", "Visibility" ) );
 	{
-		// New Class Blueprint...
-		MenuBuilder.AddMenuEntry(FLevelEditorCommands::Get().CreateClassBlueprint, NAME_None, LOCTEXT("NewClassBlueprint", "New Class Blueprint..."));
-
-		// Open Class Blueprint...
-		FSlateIcon OpenBPIcon(FEditorStyle::GetStyleSetName(), "LevelEditor.OpenClassBlueprint");
-		MenuBuilder.AddSubMenu(
-			LOCTEXT("OpenClassBlueprintSubMenu", "Open Class Blueprint..."),
-			LOCTEXT("OpenClassBlueprintSubMenu_ToolTip", "Open an existing Class Blueprint in this project"),
-			FNewMenuDelegate::CreateStatic(&FBlueprintMenus::MakeOpenClassBPMenu), 
-			false, 
-			OpenBPIcon );
+		MenuBuilder.AddMenuEntry( FLevelEditorCommands::Get().BuildLightingOnly_VisibilityOnly );
 	}
 	MenuBuilder.EndSection();
 
+	MenuBuilder.BeginSection("LevelEditorGeometry", LOCTEXT( "GeometryHeading", "Geometry" ) );
+	{
+		MenuBuilder.AddMenuEntry( FLevelEditorCommands::Get().BuildGeometryOnly );
+		MenuBuilder.AddMenuEntry( FLevelEditorCommands::Get().BuildGeometryOnly_OnlyCurrentLevel );
+	}
+	MenuBuilder.EndSection();
+
+	MenuBuilder.BeginSection("LevelEditorNavigation", LOCTEXT( "NavigationHeading", "Navigation" ) );
+	{
+		MenuBuilder.AddMenuEntry( FLevelEditorCommands::Get().BuildPathsOnly );
+	}
+	MenuBuilder.EndSection();
+
+	MenuBuilder.BeginSection("LevelEditorAutomation", LOCTEXT( "AutomationHeading", "Automation" ) );
+	{
+		MenuBuilder.AddMenuEntry( FLevelEditorCommands::Get().BuildAndSubmitToSourceControl );
+	}
+	MenuBuilder.EndSection();
+
+	// Map Check
+	MenuBuilder.BeginSection("LevelEditorVerification", LOCTEXT( "VerificationHeading", "Verification" ) );
+	{
+		MenuBuilder.AddMenuEntry( FLevelEditorCommands::Get().MapCheck, NAME_None, LOCTEXT("OpenMapCheck", "Map Check") );
+	}
+	MenuBuilder.EndSection();
+
@@ -987,6 +1291,360 @@
 	return MenuBuilder.MakeWidget();
 }
 
+static void MakeMaterialQualityLevelMenu( FMenuBuilder& MenuBuilder )
+{
+	MenuBuilder.BeginSection("LevelEditorMaterialQualityLevel", NSLOCTEXT( "LevelToolBarViewMenu", "MaterialQualityLevelHeading", "Material Quality Level" ) );
+	{
+		MenuBuilder.AddMenuEntry( FLevelEditorCommands::Get().MaterialQualityLevel_Low );
+		MenuBuilder.AddMenuEntry( FLevelEditorCommands::Get().MaterialQualityLevel_High );
+	}
+	MenuBuilder.EndSection();
+}
+
+static void MakeShaderModelPreviewMenu(FMenuBuilder& MenuBuilder)
+{
+	MenuBuilder.BeginSection("LevelEditorShaderModelPreview", NSLOCTEXT("LevelToolBarViewMenu", "FeatureLevelPreviewHeading", "Feature Level Preview"));
+	{
+		for (int32 i = GMaxRHIFeatureLevel; i >= 0; --i)
+		{
+			if (i != ERHIFeatureLevel::ES3_1)
+			{
+				MenuBuilder.AddMenuEntry(FLevelEditorCommands::Get().FeatureLevelPreview[i]);
+			}
+		}
+	}
+	MenuBuilder.EndSection();
+}
+
+static void MakeScalabilityMenu( FMenuBuilder& MenuBuilder )
+{
+	MenuBuilder.AddWidget(SNew(SScalabilitySettings), FText(), true);
+}
+
+static void MakePreviewSettingsMenu( FMenuBuilder& MenuBuilder )
+{
+#define LOCTEXT_NAMESPACE "LevelToolBarViewMenu"
+
+	MenuBuilder.BeginSection("LevelEditorPreview", LOCTEXT("PreviewHeading", "Previewing"));
+	{
+		MenuBuilder.AddMenuEntry(FLevelEditorCommands::Get().DrawBrushMarkerPolys);
+		MenuBuilder.AddMenuEntry(FLevelEditorCommands::Get().OnlyLoadVisibleInPIE);
+		MenuBuilder.AddMenuEntry(FLevelEditorCommands::Get().ToggleParticleSystemLOD);
+		MenuBuilder.AddMenuEntry(FLevelEditorCommands::Get().ToggleParticleSystemHelpers);
+		MenuBuilder.AddMenuEntry(FLevelEditorCommands::Get().ToggleFreezeParticleSimulation);
+		MenuBuilder.AddMenuEntry(FLevelEditorCommands::Get().ToggleLODViewLocking);
+		MenuBuilder.AddMenuEntry(FLevelEditorCommands::Get().LevelStreamingVolumePrevis);
+	}
+	MenuBuilder.EndSection();
+#undef LOCTEXT_NAMESPACE
+}
+
+TSharedRef< SWidget > FLevelEditorToolBar::GenerateQuickSettingsMenu( TSharedRef<FUICommandList> InCommandList )
+{
+#define LOCTEXT_NAMESPACE "LevelToolBarViewMenu"
+
+	// Get all menu extenders for this context menu from the level editor module
+	FLevelEditorModule& LevelEditorModule = FModuleManager::GetModuleChecked<FLevelEditorModule>( TEXT("LevelEditor") );
+	TArray<FLevelEditorModule::FLevelEditorMenuExtender> MenuExtenderDelegates = LevelEditorModule.GetAllLevelEditorToolbarViewMenuExtenders();
+
+	TArray<TSharedPtr<FExtender>> Extenders;
+	for (int32 i = 0; i < MenuExtenderDelegates.Num(); ++i)
+	{
+		if (MenuExtenderDelegates[i].IsBound())
+		{
+			Extenders.Add(MenuExtenderDelegates[i].Execute(InCommandList));
+		}
+	}
+	TSharedPtr<FExtender> MenuExtender = FExtender::Combine(Extenders);
+
+	const bool bShouldCloseWindowAfterMenuSelection = true;
+	FMenuBuilder MenuBuilder( bShouldCloseWindowAfterMenuSelection, InCommandList, MenuExtender );
+
+	struct Local
+	{
+		static void OpenSettings(FName ContainerName, FName CategoryName, FName SectionName)
+		{
+			FModuleManager::LoadModuleChecked<ISettingsModule>("Settings").ShowViewer(ContainerName, CategoryName, SectionName);
+		}
+	};
+
+	MenuBuilder.BeginSection("ProjectSettingsSection", LOCTEXT("ProjectSettings","Game Specific Settings") );
+	{
+
+		MenuBuilder.AddMenuEntry(FLevelEditorCommands::Get().WorldProperties);
+
+		MenuBuilder.AddMenuEntry(
+			LOCTEXT("ProjectSettingsMenuLabel", "Project Settings..."),
+			LOCTEXT("ProjectSettingsMenuToolTip", "Change the settings of the currently loaded project"),
+			FSlateIcon(),
+			FUIAction(FExecuteAction::CreateStatic(&Local::OpenSettings, FName("Project"), FName("Project"), FName("General")))
+			);
+	}
+	MenuBuilder.EndSection();
+
+	MenuBuilder.BeginSection("LevelEditorSelection", LOCTEXT("SelectionHeading","Selection") );
+	{
+		MenuBuilder.AddMenuEntry( FLevelEditorCommands::Get().AllowTranslucentSelection );
+		MenuBuilder.AddMenuEntry( FLevelEditorCommands::Get().AllowGroupSelection );
+		MenuBuilder.AddMenuEntry( FLevelEditorCommands::Get().StrictBoxSelect );
+		MenuBuilder.AddMenuEntry( FLevelEditorCommands::Get().TransparentBoxSelect );
+		MenuBuilder.AddMenuEntry( FLevelEditorCommands::Get().ShowTransformWidget );
+	}
+	MenuBuilder.EndSection();
+
+	MenuBuilder.BeginSection("LevelEditorScalability", LOCTEXT("ScalabilityHeading", "Scalability") );
+	{
+		MenuBuilder.AddSubMenu(
+			LOCTEXT( "ScalabilitySubMenu", "Engine Scalability Settings" ),
+			LOCTEXT( "ScalabilitySubMenu_ToolTip", "Open the engine scalability settings" ),
+			FNewMenuDelegate::CreateStatic( &MakeScalabilityMenu ) );
+
+		MenuBuilder.AddSubMenu(
+			LOCTEXT( "MaterialQualityLevelSubMenu", "Material Quality Level" ),
+			LOCTEXT( "MaterialQualityLevelSubMenu_ToolTip", "Sets the value of the CVar \"r.MaterialQualityLevel\" (low=0, high=1). This affects materials via the QualitySwitch material expression." ),
+			FNewMenuDelegate::CreateStatic( &MakeMaterialQualityLevelMenu ) );
+
+		if (GetDefault<UEditorExperimentalSettings>()->bFeatureLevelPreview)
+		{
+			MenuBuilder.AddSubMenu(
+				LOCTEXT("FeatureLevelPreviewSubMenu", "Feature Level Preview"),
+				LOCTEXT("FeatureLevelPreviewSubMenu_ToolTip", "Sets the feature level preview mode"),
+				FNewMenuDelegate::CreateStatic(&MakeShaderModelPreviewMenu));
+		}
+	}
+	MenuBuilder.EndSection();
+
+	MenuBuilder.BeginSection("LevelEditorAudio", LOCTEXT("AudioHeading", "Real Time Audio") );
+	{
+		TSharedRef<SWidget> VolumeItem = SNew(SHorizontalBox)
+											+SHorizontalBox::Slot()
+											.FillWidth(0.9f)
+											.Padding( FMargin(2.0f, 0.0f, 0.0f, 0.0f) )
+											[
+												SNew(SVolumeControl)
+												.ToolTipText_Static(&FLevelEditorActionCallbacks::GetAudioVolumeToolTip)
+												.Volume_Static(&FLevelEditorActionCallbacks::GetAudioVolume)
+												.OnVolumeChanged_Static(&FLevelEditorActionCallbacks::OnAudioVolumeChanged)
+												.Muted_Static(&FLevelEditorActionCallbacks::GetAudioMuted)
+												.OnMuteChanged_Static(&FLevelEditorActionCallbacks::OnAudioMutedChanged)
+											]
+											+SHorizontalBox::Slot()
+											.FillWidth(0.1f);
+
+		MenuBuilder.AddWidget(VolumeItem, LOCTEXT("VolumeControlLabel","Volume"));
+	}
+	MenuBuilder.EndSection();
+
+	MenuBuilder.BeginSection( "Snapping", LOCTEXT("SnappingHeading","Snapping") );
+	{
+		MenuBuilder.AddMenuEntry( FLevelEditorCommands::Get().EnableActorSnap );
+		TSharedRef<SWidget> SnapItem = 
+		SNew(SHorizontalBox)
+	          +SHorizontalBox::Slot()
+	          .FillWidth(0.9f)
+	          [
+		          SNew(SSlider)
+		          .ToolTipText_Static(&FLevelEditorActionCallbacks::GetActorSnapTooltip)
+		          .Value_Static(&FLevelEditorActionCallbacks::GetActorSnapSetting)
+		          .OnValueChanged_Static(&FLevelEditorActionCallbacks::SetActorSnapSetting)
+	          ]
+	          +SHorizontalBox::Slot()
+	          .FillWidth(0.1f);
+		MenuBuilder.AddWidget(SnapItem, LOCTEXT("ActorSnapLabel","Distance"));
+
+		MenuBuilder.AddMenuEntry( FLevelEditorCommands::Get().ToggleSocketSnapping );
+		MenuBuilder.AddMenuEntry( FLevelEditorCommands::Get().EnableVertexSnap );
+	}
+	MenuBuilder.EndSection();
+
+	MenuBuilder.BeginSection("LevelEditorViewport", LOCTEXT("ViewportHeading", "Viewport") );
+	{
+		MenuBuilder.AddMenuEntry( FLevelEditorCommands::Get().ToggleHideViewportUI );
+
+		MenuBuilder.AddSubMenu( LOCTEXT("PreviewMenu", "Previewing"), LOCTEXT("PreviewMenuTooltip","Game Preview Settings"), FNewMenuDelegate::CreateStatic( &MakePreviewSettingsMenu ) );
+	}
+	MenuBuilder.EndSection();
+
+#undef LOCTEXT_NAMESPACE
+
+	return MenuBuilder.MakeWidget();
+}
+
+TSharedRef< SWidget > FLevelEditorToolBar::GenerateCompileMenuContent( TSharedRef<FUICommandList> InCommandList )
+{
+#define LOCTEXT_NAMESPACE "LevelToolBarViewMenu"
+
+	// Get all menu extenders for this context menu from the level editor module
+	FLevelEditorModule& LevelEditorModule = FModuleManager::GetModuleChecked<FLevelEditorModule>( TEXT("LevelEditor") );
+	TArray<FLevelEditorModule::FLevelEditorMenuExtender> MenuExtenderDelegates = LevelEditorModule.GetAllLevelEditorToolbarCompileMenuExtenders();
+
+	TArray<TSharedPtr<FExtender>> Extenders;
+	for (int32 i = 0; i < MenuExtenderDelegates.Num(); ++i)
+	{
+		if (MenuExtenderDelegates[i].IsBound())
+		{
+			Extenders.Add(MenuExtenderDelegates[i].Execute(InCommandList));
+		}
+	}
+	TSharedPtr<FExtender> MenuExtender = FExtender::Combine(Extenders);
+
+	const bool bShouldCloseWindowAfterMenuSelection = true;
+	FMenuBuilder MenuBuilder( bShouldCloseWindowAfterMenuSelection, InCommandList, MenuExtender );
+
+	if (LevelEditorModule.CanBeRecompiled())
+	{
+		MenuBuilder.AddMenuEntry( FLevelEditorCommands::Get().RecompileLevelEditor );
+		MenuBuilder.AddMenuEntry( FLevelEditorCommands::Get().ReloadLevelEditor );
+	}
+
+	MenuBuilder.AddMenuEntry( FLevelEditorCommands::Get().RecompileGameCode );
+
+#undef LOCTEXT_NAMESPACE
+
+	return MenuBuilder.MakeWidget();
+}
+
+TSharedRef< SWidget > FLevelEditorToolBar::GenerateOpenBlueprintMenuContent( TSharedRef<FUICommandList> InCommandList, TWeakPtr< SLevelEditor > InLevelEditor )
+{
+#define LOCTEXT_NAMESPACE "LevelToolBarViewMenu"
+
+	struct FBlueprintMenus
+	{
+		/** Generates a sub-level Blueprints sub-menu */
+		static void MakeSubLevelsMenu(FMenuBuilder& InMenuBuilder, TWeakPtr< SLevelEditor > InLvlEditor)
+		{
+			FSlateIcon EditBP(FEditorStyle::Get().GetStyleSetName(), TEXT("LevelEditor.OpenLevelBlueprint"));
+
+			InMenuBuilder.BeginSection(NAME_None, LOCTEXT("SubLevelsHeading", "Sub-Level Blueprints"));
+			{
+				UWorld* World = InLvlEditor.Pin()->GetWorld();
+				for (int32 iLevel = 0; iLevel < World->GetNumLevels(); iLevel++)
+				{
+					ULevel* Level = World->GetLevel(iLevel);
+					if (Level != NULL && Level->GetOutermost() != NULL)
+					{
+						if (!Level->IsPersistentLevel())
+						{
+							FUIAction UIAction
+								(
+								FExecuteAction::CreateStatic(&FLevelEditorToolBar::OnOpenSubLevelBlueprint, Level)
+								);
+
+							FText DisplayName = FText::Format(LOCTEXT("SubLevelBlueprintItem", "Edit {LevelName}"), FText::FromString(FPaths::GetCleanFilename(Level->GetOutermost()->GetName())));
+							InMenuBuilder.AddMenuEntry(DisplayName, FText::GetEmpty(), EditBP, UIAction);
+						}
+					}
+				}
+			}
+			InMenuBuilder.EndSection();
+		}
+
+		/** Handle BP being selected from popup picker */
+		static void OnBPSelected(const class FAssetData& AssetData)
+		{
+			UBlueprint* SelectedBP = Cast<UBlueprint>(AssetData.GetAsset());
+			if(SelectedBP)
+			{
+				FAssetEditorManager::Get().OpenEditorForAsset(SelectedBP);
+			}
+		}
+
+
+		/** Generates 'open blueprint' sub-menu */
+		static void MakeOpenClassBPMenu(FMenuBuilder& InMenuBuilder)
+		{
+			FContentBrowserModule& ContentBrowserModule = FModuleManager::Get().LoadModuleChecked<FContentBrowserModule>(TEXT("ContentBrowser"));
+
+			// Configure filter for asset picker
+			FAssetPickerConfig Config;
+			Config.Filter.ClassNames.Add(UBlueprint::StaticClass()->GetFName());
+			Config.InitialAssetViewType = EAssetViewType::List;
+			Config.ThumbnailScale = 0; // make thumbnails as small as possible
+			Config.OnAssetSelected = FOnAssetSelected::CreateStatic(&FBlueprintMenus::OnBPSelected);
+			Config.bAllowDragging = false;
+			// Don't show stuff in Engine
+			Config.Filter.PackagePaths.Add("/Game");
+			Config.Filter.bRecursivePaths = true;
+
+			TSharedRef<SWidget> Widget = 
+				SNew(SBox)
+				.WidthOverride(300.f)
+				.HeightOverride(300.f)
+				[
+					ContentBrowserModule.Get().CreateAssetPicker(Config)
+				];
+		
+
+			InMenuBuilder.BeginSection(NAME_None, LOCTEXT("BrowseHeader", "Browse"));
+			{
+				InMenuBuilder.AddWidget(Widget, FText::GetEmpty());
+			}
+			InMenuBuilder.EndSection();
+		}
+	};
+
+
+	const bool bShouldCloseWindowAfterMenuSelection = true;
+	FMenuBuilder MenuBuilder( bShouldCloseWindowAfterMenuSelection, InCommandList );
+
+	MenuBuilder.BeginSection(NAME_None, LOCTEXT("LevelScriptBlueprints", "Level Blueprints"));
+	{
+		MenuBuilder.AddMenuEntry( FLevelEditorCommands::Get().OpenLevelBlueprint );
+
+		// If there are any sub-levels, display the sub-menu. A single level means there is only the persistent level
+		UWorld* World = InLevelEditor.Pin()->GetWorld();
+		if(World->GetNumLevels() > 1)
+		{
+			MenuBuilder.AddSubMenu(
+				LOCTEXT( "SubLevelsSubMenu", "Sub-Levels" ),
+				LOCTEXT( "SubLevelsSubMenu_ToolTip", "Shows available sub-level Blueprints that can be edited." ),
+				FNewMenuDelegate::CreateStatic( &FBlueprintMenus::MakeSubLevelsMenu, InLevelEditor ), 
+				FUIAction(), NAME_None, EUserInterfaceActionType::Button, false, FSlateIcon(FEditorStyle::Get().GetStyleSetName(), TEXT("LevelEditor.OpenLevelBlueprint")) );
+		}
+	}
+	MenuBuilder.EndSection();
+
+	MenuBuilder.BeginSection(NAME_None, LOCTEXT("ProjectSettingsClasses", "Project Settings"));
+	{
+		// If source control is enabled, queue up a query to the status of the config file so it is (hopefully) ready before we get to the sub-menu
+		if(ISourceControlModule::Get().IsEnabled())
+		{
+			FString ConfigFilePath = FPaths::ConvertRelativePathToFull(FString::Printf(TEXT("%sDefaultEngine.ini"), *FPaths::SourceConfigDir()));
+
+			// note: calling QueueStatusUpdate often does not spam status updates as an internal timer prevents this
+			ISourceControlModule::Get().QueueStatusUpdate(ConfigFilePath);
+		}
+		LevelEditorActionHelpers::CreateGameModeSubMenu(MenuBuilder, InCommandList, InLevelEditor, true);
+	}
+	MenuBuilder.EndSection();
+
+	MenuBuilder.BeginSection(NAME_None, LOCTEXT("WorldSettingsClasses", "World Override"));
+	{
+		LevelEditorActionHelpers::CreateGameModeSubMenu(MenuBuilder, InCommandList, InLevelEditor, false);
+	}
+	MenuBuilder.EndSection();
+
+	MenuBuilder.BeginSection(NAME_None, LOCTEXT("ClassBlueprints", "Class Blueprints"));
+	{
+		// New Class Blueprint...
+		MenuBuilder.AddMenuEntry(FLevelEditorCommands::Get().CreateClassBlueprint, NAME_None, LOCTEXT("NewClassBlueprint", "New Class Blueprint..."));
+
+		// Open Class Blueprint...
+		FSlateIcon OpenBPIcon(FEditorStyle::GetStyleSetName(), "LevelEditor.OpenClassBlueprint");
+		MenuBuilder.AddSubMenu(
+			LOCTEXT("OpenClassBlueprintSubMenu", "Open Class Blueprint..."),
+			LOCTEXT("OpenClassBlueprintSubMenu_ToolTip", "Open an existing Class Blueprint in this project"),
+			FNewMenuDelegate::CreateStatic(&FBlueprintMenus::MakeOpenClassBPMenu), 
+			false, 
+			OpenBPIcon );
+	}
+	MenuBuilder.EndSection();
+
+#undef LOCTEXT_NAMESPACE
+
+	return MenuBuilder.MakeWidget();
+}
+
 void FLevelEditorToolBar::OnOpenSubLevelBlueprint( ULevel* InLevel )
 {
 	ULevelScriptBlueprint* LevelScriptBlueprint = InLevel->GetLevelScriptBlueprint();
@@ -1001,312 +1659,4 @@
 	}
 }
 
-UClass* FLevelEditorToolBar::GetActiveGameModeClass(TWeakPtr< SLevelEditor > InLevelEditor)
-{
-	AWorldSettings* WorldSettings = InLevelEditor.Pin()->GetWorld()->GetWorldSettings();
-	if(WorldSettings->DefaultGameMode)
-	{
-		return WorldSettings->DefaultGameMode;
-	}
-	return NULL;
-}
-
-FText FLevelEditorToolBar::GetOpenGameModeBlueprintLabel(TWeakPtr< SLevelEditor > InLevelEditor)
-{
-#define LOCTEXT_NAMESPACE "LevelToolBarViewMenu"
-	if(UClass* GameModeClass = GetActiveGameModeClass(InLevelEditor))
-	{
-		if(GameModeClass->ClassGeneratedBy)
-		{
-			return FText::Format( LOCTEXT("GameModeEditBlueprint", "GameMode: Edit {GameModeName}"), FText::FromString(GameModeClass->ClassGeneratedBy->GetName()));
-		}
-
-		return FText::Format( LOCTEXT("GameModeBlueprint", "GameMode: {GameModeName}"), FText::FromString(GameModeClass->GetName()));
-	}
-
-	return LOCTEXT("GameModeCreateBlueprint", "GameMode: New...");
-#undef LOCTEXT_NAMESPACE
-}
-
-void FLevelEditorToolBar::OnCreateGameModeClassPicked(UClass* InChosenClass, TWeakPtr< SLevelEditor > InLevelEditor)
-{
-	if(InChosenClass)
-	{
-		const FString NewBPName(TEXT("NewGameMode"));
-		UBlueprint* Blueprint = FKismetEditorUtilities::CreateBlueprintFromClass(NSLOCTEXT("LevelEditorCommands", "CreateGameModeBlueprint_Title", "Create GameMode Blueprint"), InChosenClass, NewBPName);
-
-		if(Blueprint)
-		{
-			// @todo Re-enable once world centric works
-			const bool bOpenWorldCentric = false;
-			FAssetEditorManager::Get().OpenEditorForAsset(
-				Blueprint,
-				bOpenWorldCentric ? EToolkitMode::WorldCentric : EToolkitMode::Standalone,
-				InLevelEditor.Pin()  );
-
-			InLevelEditor.Pin()->GetWorld()->GetWorldSettings()->DefaultGameMode = Cast<UClass>(Blueprint->GeneratedClass);
-		}
-	}
-	FSlateApplication::Get().DismissAllMenus();
-}
-
-void FLevelEditorToolBar::OnSelectGameModeClassPicked(UClass* InChosenClass, TWeakPtr< SLevelEditor > InLevelEditor)
-{
-	InLevelEditor.Pin()->GetWorld()->GetWorldSettings()->DefaultGameMode = InChosenClass;
-
-	FSlateApplication::Get().DismissAllMenus();
-}
-
-UClass* FLevelEditorToolBar::GetActiveGameStateClass(TWeakPtr< SLevelEditor > InLevelEditor)
-{
-	AWorldSettings* WorldSettings = InLevelEditor.Pin()->GetWorld()->GetWorldSettings();
-	if(WorldSettings->DefaultGameMode)
-	{
-		AGameMode* ActiveGameMode = Cast<AGameMode>(WorldSettings->DefaultGameMode->GetDefaultObject());
-		if(ActiveGameMode)
-		{
-			return ActiveGameMode->GameStateClass;
-		}
-	}
-	return NULL;
-}
-
-FText FLevelEditorToolBar::GetOpenGameStateBlueprintLabel(TWeakPtr< SLevelEditor > InLevelEditor)
-{
-#define LOCTEXT_NAMESPACE "LevelToolBarViewMenu"
-	if(UClass* GameStateClass = GetActiveGameStateClass(InLevelEditor))
-	{
-		if(GameStateClass->ClassGeneratedBy)
-		{
-			return FText::Format( LOCTEXT("GameStateEditBlueprint", "GameState: Edit {GameStateName}"), FText::FromString(GameStateClass->ClassGeneratedBy->GetName()));
-		}
-
-		return FText::Format( LOCTEXT("GameStateBlueprint", "GameState: {GameStateName}"), FText::FromString(GameStateClass->GetName()));
-	}
-
-	return LOCTEXT("GameStateCreateBlueprint", "GameState: New...");
-#undef LOCTEXT_NAMESPACE
-}
-
-void FLevelEditorToolBar::OnCreateGameStateClassPicked(UClass* InChosenClass, TWeakPtr< SLevelEditor > InLevelEditor)
-{
-	if(InChosenClass)
-	{
-		const FString NewBPName(TEXT("NewGameState"));
-		UBlueprint* Blueprint = FKismetEditorUtilities::CreateBlueprintFromClass(NSLOCTEXT("LevelEditorCommands", "CreateGameStateBlueprint_Title", "Create GameState Blueprint"), InChosenClass, NewBPName);
-
-		if(Blueprint)
-		{
-			// @todo Re-enable once world centric works
-			const bool bOpenWorldCentric = false;
-			FAssetEditorManager::Get().OpenEditorForAsset(
-				Blueprint,
-				bOpenWorldCentric ? EToolkitMode::WorldCentric : EToolkitMode::Standalone,
-				InLevelEditor.Pin()  );
-
-			AGameMode* ActiveGameMode = Cast<AGameMode>(InLevelEditor.Pin()->GetWorld()->GetWorldSettings()->DefaultGameMode->GetDefaultObject());
-			ActiveGameMode->GameStateClass = Cast<UClass>(Blueprint->GeneratedClass);
-		}
-	}
-	FSlateApplication::Get().DismissAllMenus();
-}
-
-void FLevelEditorToolBar::OnSelectGameStateClassPicked(UClass* InChosenClass, TWeakPtr< SLevelEditor > InLevelEditor)
-{
-	AGameMode* ActiveGameMode = Cast<AGameMode>(InLevelEditor.Pin()->GetWorld()->GetWorldSettings()->DefaultGameMode->GetDefaultObject());
-
-	ActiveGameMode->GameStateClass = InChosenClass;
-
-	FSlateApplication::Get().DismissAllMenus();
-}
-
-UClass* FLevelEditorToolBar::GetActivePawnClass(TWeakPtr< SLevelEditor > InLevelEditor)
-{
-	AWorldSettings* WorldSettings = InLevelEditor.Pin()->GetWorld()->GetWorldSettings();
-	if(WorldSettings->DefaultGameMode)
-	{
-		AGameMode* ActiveGameMode = Cast<AGameMode>(WorldSettings->DefaultGameMode->GetDefaultObject());
-		if(ActiveGameMode)
-		{
-			return ActiveGameMode->DefaultPawnClass;
-		}
-	}
-	return NULL;
-}
-
-FText FLevelEditorToolBar::GetOpenPawnBlueprintLabel(TWeakPtr< SLevelEditor > InLevelEditor)
-{
-#define LOCTEXT_NAMESPACE "LevelToolBarViewMenu"
-	if(UClass* PawnClass = GetActivePawnClass(InLevelEditor))
-	{
-		if(PawnClass->ClassGeneratedBy)
-		{
-			return FText::Format( LOCTEXT("PawnEditBlueprint", "Pawn: Edit {PawnName}"), FText::FromString(PawnClass->ClassGeneratedBy->GetName()));
-		}
-
-		return FText::Format( LOCTEXT("PawnBlueprint", "Pawn: {PawnName}"), FText::FromString(PawnClass->GetName()));
-	}
-
-	return LOCTEXT("PawnCreateBlueprint", "Pawn: New...");
-#undef LOCTEXT_NAMESPACE
-}
-
-void FLevelEditorToolBar::OnCreatePawnClassPicked(UClass* InChosenClass, TWeakPtr< SLevelEditor > InLevelEditor)
-{
-	if(InChosenClass)
-	{
-		const FString NewBPName(TEXT("NewPawn"));
-		UBlueprint* Blueprint = FKismetEditorUtilities::CreateBlueprintFromClass(NSLOCTEXT("LevelEditorCommands", "CreatePawnBlueprint_Title", "Create Pawn Blueprint"), InChosenClass, NewBPName);
-
-		if(Blueprint)
-		{
-			// @todo Re-enable once world centric works
-			const bool bOpenWorldCentric = false;
-			FAssetEditorManager::Get().OpenEditorForAsset(
-				Blueprint,
-				bOpenWorldCentric ? EToolkitMode::WorldCentric : EToolkitMode::Standalone,
-				InLevelEditor.Pin()  );
-
-			AGameMode* ActiveGameMode = Cast<AGameMode>(InLevelEditor.Pin()->GetWorld()->GetWorldSettings()->DefaultGameMode->GetDefaultObject());
-			ActiveGameMode->DefaultPawnClass = Cast<UClass>(Blueprint->GeneratedClass);
-		}
-	}
-	FSlateApplication::Get().DismissAllMenus();
-}
-
-void FLevelEditorToolBar::OnSelectPawnClassPicked(UClass* InChosenClass, TWeakPtr< SLevelEditor > InLevelEditor)
-{
-	AGameMode* ActiveGameMode = Cast<AGameMode>(InLevelEditor.Pin()->GetWorld()->GetWorldSettings()->DefaultGameMode->GetDefaultObject());
-
-	ActiveGameMode->DefaultPawnClass = InChosenClass;
-
-	FSlateApplication::Get().DismissAllMenus();
-}
-
-UClass* FLevelEditorToolBar::GetActiveHUDClass(TWeakPtr< SLevelEditor > InLevelEditor)
-{
-	AWorldSettings* WorldSettings = InLevelEditor.Pin()->GetWorld()->GetWorldSettings();
-	if(WorldSettings->DefaultGameMode)
-	{
-		AGameMode* ActiveGameMode = Cast<AGameMode>(WorldSettings->DefaultGameMode->GetDefaultObject());
-		if(ActiveGameMode)
-		{
-			return ActiveGameMode->HUDClass;
-		}
-	}
-	return NULL;
-}
-
-FText FLevelEditorToolBar::GetOpenHUDBlueprintLabel(TWeakPtr< SLevelEditor > InLevelEditor)
-{
-#define LOCTEXT_NAMESPACE "LevelToolBarViewMenu"
-	if(UClass* HUDClass = GetActiveHUDClass(InLevelEditor))
-	{
-		if(HUDClass->ClassGeneratedBy)
-		{
-			return FText::Format( LOCTEXT("HUDEditBlueprint", "HUD: Edit {HUDName}"), FText::FromString(HUDClass->ClassGeneratedBy->GetName()));
-		}
-
-		return FText::Format( LOCTEXT("HUDBlueprint", "HUD: {HUDName}"), FText::FromString(HUDClass->GetName()));
-	}
-
-	return LOCTEXT("HUDCreateBlueprint", "HUD: New...");
-#undef LOCTEXT_NAMESPACE
-}
-
-void FLevelEditorToolBar::OnCreateHUDClassPicked(UClass* InChosenClass, TWeakPtr< SLevelEditor > InLevelEditor)
-{
-	if(InChosenClass)
-	{
-		const FString NewBPName(TEXT("NewHUD"));
-		UBlueprint* Blueprint = FKismetEditorUtilities::CreateBlueprintFromClass(NSLOCTEXT("LevelEditorCommands", "CreateHUDBlueprint_Title", "Create HUD Blueprint"), InChosenClass, NewBPName);
-
-		if(Blueprint)
-		{
-			// @todo Re-enable once world centric works
-			const bool bOpenWorldCentric = false;
-			FAssetEditorManager::Get().OpenEditorForAsset(
-				Blueprint,
-				bOpenWorldCentric ? EToolkitMode::WorldCentric : EToolkitMode::Standalone,
-				InLevelEditor.Pin()  );
-
-			AGameMode* ActiveGameMode = Cast<AGameMode>(InLevelEditor.Pin()->GetWorld()->GetWorldSettings()->DefaultGameMode->GetDefaultObject());
-			ActiveGameMode->HUDClass = Cast<UClass>(Blueprint->GeneratedClass);
-		}
-	}
-	FSlateApplication::Get().DismissAllMenus();
-}
-
-void FLevelEditorToolBar::OnSelectHUDClassPicked(UClass* InChosenClass, TWeakPtr< SLevelEditor > InLevelEditor)
-{
-	AGameMode* ActiveGameMode = Cast<AGameMode>(InLevelEditor.Pin()->GetWorld()->GetWorldSettings()->DefaultGameMode->GetDefaultObject());
-
-	ActiveGameMode->HUDClass = InChosenClass;
-
-	FSlateApplication::Get().DismissAllMenus();
-}
-
-UClass* FLevelEditorToolBar::GetActivePlayerControllerClass(TWeakPtr< SLevelEditor > InLevelEditor)
-{
-	AWorldSettings* WorldSettings = InLevelEditor.Pin()->GetWorld()->GetWorldSettings();
-	if(WorldSettings->DefaultGameMode)
-	{
-		AGameMode* ActiveGameMode = Cast<AGameMode>(WorldSettings->DefaultGameMode->GetDefaultObject());
-		if(ActiveGameMode)
-		{
-			return ActiveGameMode->PlayerControllerClass;
-		}
-	}
-	return NULL;
-}
-
-FText FLevelEditorToolBar::GetOpenPlayerControllerBlueprintLabel(TWeakPtr< SLevelEditor > InLevelEditor)
-{
-#define LOCTEXT_NAMESPACE "LevelToolBarViewMenu"
-	if(UClass* PlayerControllerClass = GetActivePlayerControllerClass(InLevelEditor))
-	{
-		if(PlayerControllerClass->ClassGeneratedBy)
-		{
-			return FText::Format( LOCTEXT("PlayerControllerEditBlueprint", "PlayerController: Edit {PlayerControllerName}"), FText::FromString(PlayerControllerClass->ClassGeneratedBy->GetName()));
-		}
-
-		return FText::Format( LOCTEXT("PlayerControllerBlueprint", "PlayerController: {PlayerControllerName}"), FText::FromString(PlayerControllerClass->GetName()));
-	}
-
-	return LOCTEXT("PlayerControllerCreateBlueprint", "PlayerController: New...");
-#undef LOCTEXT_NAMESPACE
-}
-
-void FLevelEditorToolBar::OnCreatePlayerControllerClassPicked(UClass* InChosenClass, TWeakPtr< SLevelEditor > InLevelEditor)
-{
-	if(InChosenClass)
-	{
-		const FString NewBPName(TEXT("NewPlayerController"));
-		UBlueprint* Blueprint = FKismetEditorUtilities::CreateBlueprintFromClass(NSLOCTEXT("LevelEditorCommands", "CreatePlayerControllerBlueprint_Title", "Create PlayerController Blueprint"), InChosenClass, NewBPName);
-
-		if(Blueprint)
-		{
-			// @todo Re-enable once world centric works
-			const bool bOpenWorldCentric = false;
-			FAssetEditorManager::Get().OpenEditorForAsset(
-				Blueprint,
-				bOpenWorldCentric ? EToolkitMode::WorldCentric : EToolkitMode::Standalone,
-				InLevelEditor.Pin()  );
-
-			AGameMode* ActiveGameMode = Cast<AGameMode>(InLevelEditor.Pin()->GetWorld()->GetWorldSettings()->DefaultGameMode->GetDefaultObject());
-			ActiveGameMode->PlayerControllerClass = Cast<UClass>(Blueprint->GeneratedClass);
-		}
-	}
-	FSlateApplication::Get().DismissAllMenus();
-}
-
-void FLevelEditorToolBar::OnSelectPlayerControllerClassPicked(UClass* InChosenClass, TWeakPtr< SLevelEditor > InLevelEditor)
-{
-	AGameMode* ActiveGameMode = Cast<AGameMode>(InLevelEditor.Pin()->GetWorld()->GetWorldSettings()->DefaultGameMode->GetDefaultObject());
-
-	ActiveGameMode->PlayerControllerClass = InChosenClass;
-
-	FSlateApplication::Get().DismissAllMenus();
-}
-
 TSharedRef< SWidget > FLevelEditorToolBar::GenerateMatineeMenuContent( TSharedRef<FUICommandList> InCommandList, TWeakPtr<SLevelEditor> LevelEditorWeakPtr )
diff --git a/Engine/Source/Editor/LevelEditor/Private/LevelEditorToolBar.h b/Engine/Source/Editor/LevelEditor/Private/LevelEditorToolBar.h
--- a/Engine/Source/Editor/LevelEditor/Private/LevelEditorToolBar.h
+++ b/Engine/Source/Editor/LevelEditor/Private/LevelEditorToolBar.h
@@ -82,94 +82,6 @@ protected:
 	 * @param InLevel	The level to open the Blueprint of (creates if needed)
 	 */
 	static void OnOpenSubLevelBlueprint( ULevel* InLevel );
-
-	/**
-	 * Retrieves the active GameMode class from
-	 *
-	 * @param InLevelEditor		The editor to extract the world from
-	 * @return					The active GameMode class in the World Settings
-	 */
-	static UClass* GetActiveGameModeClass(TWeakPtr< SLevelEditor > InLevelEditor);
-
-	/** Callback for the label to display for the GameMode menu selection */
-	static FText GetOpenGameModeBlueprintLabel(TWeakPtr< SLevelEditor > InLevelEditor);
-
-	/** Callback when selecting a GameMode class, assigns it to the world */
-	static void OnSelectGameModeClassPicked(UClass* InChosenClass, TWeakPtr< SLevelEditor > InLevelEditor);
-
-	/** Callback when creating a new GameMode class, creates the Blueprint and assigns it to the world */
-	static void OnCreateGameModeClassPicked(UClass* InChosenClass, TWeakPtr< SLevelEditor > InLevelEditor);
-
-	/**
-	 * Retrieves the active GameState class from
-	 *
-	 * @param InLevelEditor		The editor to extract the world from
-	 * @return					The active GameState class in the World Settings
-	 */
-	static UClass* GetActiveGameStateClass(TWeakPtr< SLevelEditor > InLevelEditor);
-
-	/** Callback for the label to display for the GameState menu selection */
-	static FText GetOpenGameStateBlueprintLabel(TWeakPtr< SLevelEditor > InLevelEditor);
-	
-	/** Callback when selecting a GameState class, assigns it to the world */
-	static void OnSelectGameStateClassPicked(UClass* InChosenClass, TWeakPtr< SLevelEditor > InLevelEditor);
-
-	/** Callback when creating a new GameState class, creates the Blueprint and assigns it to the world */
-	static void OnCreateGameStateClassPicked(UClass* InChosenClass, TWeakPtr< SLevelEditor > InLevelEditor);
-
-	/**
-	 * Retrieves the active Pawn class from
-	 *
-	 * @param InLevelEditor		The editor to extract the world from
-	 * @return					The active Pawn class in the World Settings
-	 */
-	static UClass* GetActivePawnClass(TWeakPtr< SLevelEditor > InLevelEditor);
-
-	/** Callback for the label to display for the Pawn menu selection */
-	static FText GetOpenPawnBlueprintLabel(TWeakPtr< SLevelEditor > InLevelEditor);
-
-	/** Callback for the tooltip to display for the Pawn menu selection */
-	static FText GetOpenPawnBlueprintTooltip(TWeakPtr< SLevelEditor > InLevelEditor);
-
-	/** Callback when selecting a Pawn class, assigns it to the world */
-	static void OnSelectPawnClassPicked(UClass* InChosenClass, TWeakPtr< SLevelEditor > InLevelEditor);
-
-	/** Callback when creating a new Pawn class, creates the Blueprint and assigns it to the world */
-	static void OnCreatePawnClassPicked(UClass* InChosenClass, TWeakPtr< SLevelEditor > InLevelEditor);
-
-	/**
-	 * Retrieves the active HUD class from
-	 *
-	 * @param InLevelEditor		The editor to extract the world from
-	 * @return					The active HUD class in the World Settings
-	 */
-	static UClass* GetActiveHUDClass(TWeakPtr< SLevelEditor > InLevelEditor);
-
-	/** Callback for the label to display for the HUD menu selection */
-	static FText GetOpenHUDBlueprintLabel(TWeakPtr< SLevelEditor > InLevelEditor);
-	
-	/** Callback when selecting a HUD class, assigns it to the world */
-	static void OnSelectHUDClassPicked(UClass* InChosenClass, TWeakPtr< SLevelEditor > InLevelEditor);
-
-	/** Callback when creating a new HUD class, creates the Blueprint and assigns it to the world */
-	static void OnCreateHUDClassPicked(UClass* InChosenClass, TWeakPtr< SLevelEditor > InLevelEditor);
-
-	/**
-	 * Retrieves the active PlayerController class from
-	 *
-	 * @param InLevelEditor		The editor to extract the world from
-	 * @return					The active PlayerController class in the World Settings
-	 */
-	static UClass* GetActivePlayerControllerClass(TWeakPtr< SLevelEditor > InLevelEditor);
-
-	/** Callback for the label to display for the PlayerController menu selection */
-	static FText GetOpenPlayerControllerBlueprintLabel(TWeakPtr< SLevelEditor > InLevelEditor);
-
-	/** Callback when selecting a PlayerController class, assigns it to the world */
-	static void OnSelectPlayerControllerClassPicked(UClass* InChosenClass, TWeakPtr< SLevelEditor > InLevelEditor);
-
-	/** Callback when creating a new PlayerController class, creates the Blueprint and assigns it to the world */
-	static void OnCreatePlayerControllerClassPicked(UClass* InChosenClass, TWeakPtr< SLevelEditor > InLevelEditor);
 };
 
 
diff --git a/Engine/Source/Editor/LevelEditor/Private/SLevelEditor.cpp b/Engine/Source/Editor/LevelEditor/Private/SLevelEditor.cpp
--- a/Engine/Source/Editor/LevelEditor/Private/SLevelEditor.cpp
+++ b/Engine/Source/Editor/LevelEditor/Private/SLevelEditor.cpp
@@ -85,46 +85,14 @@ void SLevelEditor::BindCommands()
 		Actions.EditAsset,
 		FExecuteAction::CreateStatic( &FLevelEditorActionCallbacks::EditAsset_Clicked, EToolkitMode::Standalone, TWeakPtr< SLevelEditor >( SharedThis( this ) ), true ) );
 
+	LevelEditorCommands->MapAction(
+		Actions.CheckOutProjectSettingsConfig,
+		FExecuteAction::CreateStatic( &FLevelEditorActionCallbacks::CheckOutProjectSettingsConfig ) );
+
 	LevelEditorCommands->MapAction(
 		Actions.OpenLevelBlueprint,
 		FExecuteAction::CreateStatic< TWeakPtr< SLevelEditor > >( &FLevelEditorActionCallbacks::OpenLevelBlueprint, SharedThis( this ) ) );
 	
-	LevelEditorCommands->MapAction(
-		Actions.OpenGameModeBlueprint,
-		FExecuteAction::CreateStatic< TWeakPtr< SLevelEditor > >( &FLevelEditorActionCallbacks::OpenGameModeBlueprint, SharedThis( this ) ) );
-
-	LevelEditorCommands->MapAction(
-		Actions.OpenGameStateBlueprint,
-		FExecuteAction::CreateStatic< TWeakPtr< SLevelEditor > >( &FLevelEditorActionCallbacks::OpenGameStateBlueprint, SharedThis( this ) ),
-		FCanExecuteAction::CreateStatic< TWeakPtr< SLevelEditor > >( &FLevelEditorActionCallbacks::CanEditGameInfoBlueprints, SharedThis( this ) ),
-		FIsActionChecked(),
-		FIsActionButtonVisible::CreateStatic< TWeakPtr< SLevelEditor > >( &FLevelEditorActionCallbacks::CanEditGameInfoBlueprints, SharedThis( this ) )
-		);
-
-	LevelEditorCommands->MapAction(
-		Actions.OpenDefaultPawnBlueprint,
-		FExecuteAction::CreateStatic< TWeakPtr< SLevelEditor > >( &FLevelEditorActionCallbacks::OpenDefaultPawnBlueprint, SharedThis( this ) ),
-		FCanExecuteAction::CreateStatic< TWeakPtr< SLevelEditor > >( &FLevelEditorActionCallbacks::CanEditGameInfoBlueprints, SharedThis( this ) ),
-		FIsActionChecked(),
-		FIsActionButtonVisible::CreateStatic< TWeakPtr< SLevelEditor > >( &FLevelEditorActionCallbacks::CanEditGameInfoBlueprints, SharedThis( this ) )
-		);
-
-	LevelEditorCommands->MapAction(
-		Actions.OpenHUDBlueprint,
-		FExecuteAction::CreateStatic< TWeakPtr< SLevelEditor > >( &FLevelEditorActionCallbacks::OpenHUDBlueprint, SharedThis( this ) ),
-		FCanExecuteAction::CreateStatic< TWeakPtr< SLevelEditor > >( &FLevelEditorActionCallbacks::CanEditGameInfoBlueprints, SharedThis( this ) ),
-		FIsActionChecked(),
-		FIsActionButtonVisible::CreateStatic< TWeakPtr< SLevelEditor > >( &FLevelEditorActionCallbacks::CanEditGameInfoBlueprints, SharedThis( this ) )
-		);
-
-	LevelEditorCommands->MapAction(
-		Actions.OpenPlayerControllerBlueprint,
-		FExecuteAction::CreateStatic< TWeakPtr< SLevelEditor > >( &FLevelEditorActionCallbacks::OpenPlayerControllerBlueprint, SharedThis( this ) ),
-		FCanExecuteAction::CreateStatic< TWeakPtr< SLevelEditor > >( &FLevelEditorActionCallbacks::CanEditGameInfoBlueprints, SharedThis( this ) ),
-		FIsActionChecked(),
-		FIsActionButtonVisible::CreateStatic< TWeakPtr< SLevelEditor > >( &FLevelEditorActionCallbacks::CanEditGameInfoBlueprints, SharedThis( this ) )
-		);
-
 	LevelEditorCommands->MapAction(
 		Actions.CreateClassBlueprint,
 		FExecuteAction::CreateStatic( &FLevelEditorActionCallbacks::CreateClassBlueprint ) );
diff --git a/Engine/Source/Editor/LevelEditor/Public/LevelEditorActions.h b/Engine/Source/Editor/LevelEditor/Public/LevelEditorActions.h
--- a/Engine/Source/Editor/LevelEditor/Public/LevelEditorActions.h
+++ b/Engine/Source/Editor/LevelEditor/Public/LevelEditorActions.h
@@ -502,11 +502,7 @@ public:
 	 * Blueprints commands
 	 */
 	TSharedPtr< FUICommandInfo > OpenLevelBlueprint;
-	TSharedPtr< FUICommandInfo > OpenGameModeBlueprint;
-	TSharedPtr< FUICommandInfo > OpenGameStateBlueprint;
-	TSharedPtr< FUICommandInfo > OpenDefaultPawnBlueprint;
-	TSharedPtr< FUICommandInfo > OpenHUDBlueprint;
-	TSharedPtr< FUICommandInfo > OpenPlayerControllerBlueprint;
+	TSharedPtr< FUICommandInfo > CheckOutProjectSettingsConfig;
 	TSharedPtr< FUICommandInfo > CreateClassBlueprint;
 
 	/** Editor mode commands */
@@ -1024,26 +1020,14 @@ public:
 	/** Open the Marketplace */
 	static void OpenMarketplace();
 
+	/** Checks out the Project Settings config */
+	static void CheckOutProjectSettingsConfig();
+
 	/** Open the level's blueprint in Kismet2 */
 	static void OpenLevelBlueprint( TWeakPtr< SLevelEditor > LevelEditor );
 
-	/** Open the world's game mode blueprint or help the user create one. */
-	static void OpenGameModeBlueprint( TWeakPtr< SLevelEditor > LevelEditor );
-
-	/** Open the current game state blueprint or help the user create one */
-	static void OpenGameStateBlueprint( TWeakPtr< SLevelEditor > LevelEditor );
-	
-	/** Open the current default pawn blueprint or help the user create one */
-	static void OpenDefaultPawnBlueprint( TWeakPtr< SLevelEditor > LevelEditor );
-	
-	/** Open the current HUD blueprint or help the user create one */
-	static void OpenHUDBlueprint( TWeakPtr< SLevelEditor > LevelEditor );
-	
-	/** Open the current player controller blueprint or help the user create one */
-	static void OpenPlayerControllerBlueprint( TWeakPtr< SLevelEditor > LevelEditor );
-
-	/** Returns TRUE if the user can edit game info Blueprints, this requires an active Blueprint based game mode to be set */
-	static bool CanEditGameInfoBlueprints( TWeakPtr< SLevelEditor > LevelEditor);
+	/** Returns TRUE if the user can edit the game mode Blueprint, this requires the DefaultEngine config file to be writable */
+	static bool CanSelectGameModeBlueprint();
 
 	/** Helps the user create a class Blueprint */
 	static void CreateClassBlueprint();
diff --git a/Engine/Source/Runtime/EngineSettings/Classes/GameMapsSettings.h b/Engine/Source/Runtime/EngineSettings/Classes/GameMapsSettings.h
--- a/Engine/Source/Runtime/EngineSettings/Classes/GameMapsSettings.h
+++ b/Engine/Source/Runtime/EngineSettings/Classes/GameMapsSettings.h
@@ -56,6 +56,13 @@ class ENGINESETTINGS_API UGameMapsSettings
 	 */
 	static void SetGameDefaultMap( const FString& NewMap );
 
+	/**
+	 * Set the default game type (see GlobalDefaultGameMode below)
+	 *
+	 * @param NewGameMode name of valid map to use
+	 */
+	static void SetGlobalDefaultGameMode( const FString& NewGameMode );
+
 public:
 
 	/** If set, this map will be loaded when the Editor starts up. */
diff --git a/Engine/Source/Runtime/EngineSettings/Private/EngineSettingsModule.cpp b/Engine/Source/Runtime/EngineSettings/Private/EngineSettingsModule.cpp
--- a/Engine/Source/Runtime/EngineSettings/Private/EngineSettingsModule.cpp
+++ b/Engine/Source/Runtime/EngineSettings/Private/EngineSettingsModule.cpp
@@ -99,5 +99,12 @@ void UGameMapsSettings::SetGameDefaultMap( const FString& NewMap )
 	}
 }
 
+void UGameMapsSettings::SetGlobalDefaultGameMode( const FString& NewGameMode )
+{
+	UGameMapsSettings* GameMapsSettings = Cast<UGameMapsSettings>(UGameMapsSettings::StaticClass()->GetDefaultObject());
+	
+	GameMapsSettings->GlobalDefaultGameMode = NewGameMode;
+}
+
 
 IMPLEMENT_MODULE(FEngineSettingsModule, EngineSettings);
