Fix world settings not having an icon in the window menu of the level editor

diff --git a/Engine/Source/Editor/LevelEditor/Private/SLevelEditor.cpp b/Engine/Source/Editor/LevelEditor/Private/SLevelEditor.cpp
--- a/Engine/Source/Editor/LevelEditor/Private/SLevelEditor.cpp
+++ b/Engine/Source/Editor/LevelEditor/Private/SLevelEditor.cpp
@@ -1058,7 +1058,7 @@ TSharedRef<SWidget> SLevelEditor::RestoreContentArea( const TSharedRef<SDockTab>
 		}
 
 		{
-			const FSlateIcon WorldPropertiesIcon(FEditorStyle::GetStyleSetName(), "LevelEditor.WorldProperties.Small");
+			const FSlateIcon WorldPropertiesIcon(FEditorStyle::GetStyleSetName(), "LevelEditor.WorldProperties.Tab");
 			LevelEditorTabManager->RegisterTabSpawner( "WorldSettingsTab", FOnSpawnTab::CreateSP<SLevelEditor, FName, FString>(this, &SLevelEditor::SpawnLevelEditorTab, FName("WorldSettingsTab"), FString()) )
 				.SetDisplayName(NSLOCTEXT("LevelEditorTabs", "WorldSettings", "World Settings"))
 				.SetTooltipText(NSLOCTEXT("LevelEditorTabs", "WorldSettingsTooltipText", "Open the World Settings tab, in which global properties of the level can be viewed and edited."))
