Removed private SClassViewer header include, not needed to use the ClassViewer outside of the module.

diff --git a/Engine/Source/Editor/LevelEditor/Private/LevelEditorToolBar.cpp b/Engine/Source/Editor/LevelEditor/Private/LevelEditorToolBar.cpp
--- a/Engine/Source/Editor/LevelEditor/Private/LevelEditorToolBar.cpp
+++ b/Engine/Source/Editor/LevelEditor/Private/LevelEditorToolBar.cpp
@@ -18,7 +18,6 @@
 #include "Engine/LevelScriptBlueprint.h"
 #include "Settings.h"
 #include "DesktopPlatformModule.h"
-#include "Editor/ClassViewer/Private/SClassViewer.h"
 #include "Editor/ClassViewer/Public/ClassViewerFilter.h"
 #include "KismetEditorUtilities.h"
 #include "ISourceControlModule.h"
@@ -246,7 +245,7 @@ namespace LevelEditorActionHelpers
 		Options.ClassFilter = Filter;
 
 		FText RootClassName = FText::FromString(InRootClass->GetName());
-		auto ClassViewer = StaticCastSharedRef<SClassViewer>(FModuleManager::LoadModuleChecked<FClassViewerModule>("ClassViewer").CreateClassViewer(Options, InOnClassPicked));
+		TSharedRef<SWidget> ClassViewer = FModuleManager::LoadModuleChecked<FClassViewerModule>("ClassViewer").CreateClassViewer(Options, InOnClassPicked);
 		InMenuBuilder.BeginSection(NAME_None, FText::Format( NSLOCTEXT("LevelToolBarViewMenu", "SelectGameModeLabel", "Select {RootClass} class"), RootClassName ));
 		InMenuBuilder.AddWidget(ClassViewer, FText::GetEmpty(), true);
 		InMenuBuilder.EndSection();
@@ -277,7 +276,7 @@ namespace LevelEditorActionHelpers
 		Options.ClassFilter = Filter;
 
 		FText RootClassName = FText::FromString(InRootClass->GetName());
-		auto ClassViewer = StaticCastSharedRef<SClassViewer>(FModuleManager::LoadModuleChecked<FClassViewerModule>("ClassViewer").CreateClassViewer(Options, InOnClassPicked));
+		TSharedRef<SWidget> ClassViewer = FModuleManager::LoadModuleChecked<FClassViewerModule>("ClassViewer").CreateClassViewer(Options, InOnClassPicked);
 		InMenuBuilder.BeginSection(NAME_None, FText::Format( NSLOCTEXT("LevelToolBarViewMenu", "CreateGameModeLabel", "Select {RootClass} parent class"), RootClassName ));
 		InMenuBuilder.AddWidget(ClassViewer, FText::GetEmpty(), true);
 		InMenuBuilder.EndSection();
