Fixing missing include for non-unity build.

diff --git a/Engine/Source/Editor/LevelEditor/Private/LevelEditorToolBar.cpp b/Engine/Source/Editor/LevelEditor/Private/LevelEditorToolBar.cpp
--- a/Engine/Source/Editor/LevelEditor/Private/LevelEditorToolBar.cpp
+++ b/Engine/Source/Editor/LevelEditor/Private/LevelEditorToolBar.cpp
@@ -18,6 +18,7 @@
 #include "Engine/LevelScriptBlueprint.h"
 #include "Settings.h"
 #include "DesktopPlatformModule.h"
+#include "Editor/ClassViewer/Public/ClassViewerModule.h"
 #include "Editor/ClassViewer/Public/ClassViewerFilter.h"
 #include "KismetEditorUtilities.h"
 #include "ISourceControlModule.h"
