Blueprints menu on the level editor toolbar will gray out the ability to create subclasses for a GameMode if the GameMode is a native class instead of hiding the functionality.

Did some polish on tooltips to better inform the user what is happening.

Can now edit native classes if a compiler is present (will open the header).

diff --git a/Engine/Source/Editor/LevelEditor/Private/LevelEditorToolBar.cpp b/Engine/Source/Editor/LevelEditor/Private/LevelEditorToolBar.cpp
--- a/Engine/Source/Editor/LevelEditor/Private/LevelEditorToolBar.cpp
+++ b/Engine/Source/Editor/LevelEditor/Private/LevelEditorToolBar.cpp
@@ -135,6 +135,22 @@ namespace LevelEditorActionHelpers
 	/** Callback when creating a new PlayerController class, creates the Blueprint and assigns it to the world */
 	static void OnCreatePlayerControllerClassPicked(UClass* InChosenClass, TWeakPtr< SLevelEditor > InLevelEditor, bool bInIsProjectSettings);
 
+	/** Opens a native class's header file if the compiler is available. */
+	static void OpenNativeClass(UClass* InClass)
+	{
+		if(InClass->HasAllClassFlags(CLASS_Native) && FSourceCodeNavigation::IsCompilerAvailable())
+		{
+			FString NativeParentClassHeaderPath;
+			const bool bFileFound = FSourceCodeNavigation::FindClassHeaderPath(InClass, NativeParentClassHeaderPath) 
+				&& (IFileManager::Get().FileSize(*NativeParentClassHeaderPath) != INDEX_NONE);
+			if (bFileFound)
+			{
+				const FString AbsoluteHeaderPath = IFileManager::Get().ConvertToAbsolutePathForExternalAppForRead(*NativeParentClassHeaderPath);
+				FSourceCodeNavigation::OpenSourceFile( AbsoluteHeaderPath );
+			}
+		}
+	}
+
 	/** Open the game mode blueprint, in the project settings or world settings */
 	static void OpenGameModeBlueprint( TWeakPtr< SLevelEditor > InLevelEditor, bool bInIsProjectSettings )
 	{
@@ -149,6 +165,10 @@ namespace LevelEditorActionHelpers
 					bOpenWorldCentric ? EToolkitMode::WorldCentric : EToolkitMode::Standalone,
 					InLevelEditor.Pin() );
 			}
+			else
+			{
+				OpenNativeClass(GameModeClass);
+			}
 		}
 	}
 
@@ -166,6 +186,10 @@ namespace LevelEditorActionHelpers
 					bOpenWorldCentric ? EToolkitMode::WorldCentric : EToolkitMode::Standalone,
 					InLevelEditor.Pin() );
 			}
+			else
+			{
+				OpenNativeClass(GameStateClass);
+			}
 		}
 	}
 
@@ -183,6 +207,10 @@ namespace LevelEditorActionHelpers
 					bOpenWorldCentric ? EToolkitMode::WorldCentric : EToolkitMode::Standalone,
 					InLevelEditor.Pin() );
 			}
+			else
+			{
+				OpenNativeClass(DefaultPawnClass);
+			}
 		}
 	}
 
@@ -200,6 +228,10 @@ namespace LevelEditorActionHelpers
 					bOpenWorldCentric ? EToolkitMode::WorldCentric : EToolkitMode::Standalone,
 					InLevelEditor.Pin() );
 			}
+			else
+			{
+				OpenNativeClass(DefaultHUDClass);
+			}
 		}
 	}
 
@@ -217,6 +249,10 @@ namespace LevelEditorActionHelpers
 					bOpenWorldCentric ? EToolkitMode::WorldCentric : EToolkitMode::Standalone,
 					InLevelEditor.Pin() );
 			}
+			else
+			{
+				OpenNativeClass(PlayerControllerClass);
+			}
 		}
 	}
 
@@ -347,6 +383,98 @@ namespace LevelEditorActionHelpers
 	 */
 	void GetBlueprintSettingsSubMenu(FMenuBuilder& InMenuBuilder, const TSharedRef<FUICommandList> InCommandList, FBlueprintMenuSettings InSettingsData);
 
+	/** Returns TRUE if the class can be edited, always TRUE for Blueprints and for native classes a compiler must be present */
+	bool CanEditClass(UClass* InClass)
+	{
+		// For native classes, we can only edit them if a compiler is available
+		if(InClass && InClass->HasAllClassFlags(CLASS_Native))
+		{
+			return FSourceCodeNavigation::IsCompilerAvailable();
+		}
+		return true;
+	}
+
+	/** Returns TRUE if the GameMode's sub-class can be created or selected */
+	bool CanCreateSelectSubClass(UClass* InGameModeClass, bool bInIsProjectSettings)
+	{
+		// Can never create or select project settings sub-classes if the config file is not checked out
+		if(bInIsProjectSettings && !FLevelEditorActionCallbacks::CanSelectGameModeBlueprint())
+		{
+			return false;
+		}
+
+		// If the game mode class is native, we cannot set the sub class
+		if(!InGameModeClass || InGameModeClass->HasAllClassFlags(CLASS_Native))
+		{
+			return false;
+		}
+		return true;
+	}
+
+	/** Creates a tooltip for a submenu */
+	FText GetSubMenuTooltip(UClass* InClass, UClass* InRootClass, bool bInIsProjectSettings)
+	{
+#define LOCTEXT_NAMESPACE "LevelToolBarViewMenu"
+		FFormatNamedArguments Args;
+		Args.Add(TEXT("Class"), FText::FromString(InRootClass->GetName()));
+		Args.Add(TEXT("TargetLocation"), bInIsProjectSettings? LOCTEXT("Project", "project") : LOCTEXT("World", "world"));
+		return FText::Format(LOCTEXT("ClassSubmenu_Tooltip", "Select, edit, or create a new {Class} blueprint for the {TargetLocation}"), Args);
+#undef LOCTEXT_NAMESPACE
+	}
+
+	/** Creates a tooltip for the create class submenu */
+	FText GetCreateMenuTooltip(UClass* InGameModeClass, UClass* InRootClass, bool bInIsProjectSettings)
+	{
+#define LOCTEXT_NAMESPACE "LevelToolBarViewMenu"
+		FText ResultText;
+
+		// Game modes can always be created and selected (providing the config is checked out, handled separately)
+		if(InRootClass != AGameMode::StaticClass() && InGameModeClass->HasAllClassFlags(CLASS_Native))
+		{
+			ResultText = LOCTEXT("CannotCreateClasses", "Cannot create classes when the game mode is a native class!");
+		}
+		else if(bInIsProjectSettings && !FLevelEditorActionCallbacks::CanSelectGameModeBlueprint())
+		{
+			ResultText = LOCTEXT("CannotCreateClasses_NeedsCheckOut", "Cannot create classes when the config file is not writable!");
+		}
+		else
+		{
+			FFormatNamedArguments Args;
+			Args.Add(TEXT("RootClass"), FText::FromString(InRootClass->GetName()));
+			Args.Add(TEXT("TargetLocation"), bInIsProjectSettings? LOCTEXT("Project", "project") : LOCTEXT("World", "world"));
+			ResultText = FText::Format( LOCTEXT("CreateClass_Tooltip", "Create a new {RootClass} based on a selected class and auto-assign it to the {TargetLocation}"), Args );
+		}
+
+		return ResultText;
+#undef LOCTEXT_NAMESPACE
+	}
+
+	/** Creates a tooltip for the select class submenu */
+	FText GetSelectMenuTooltip(UClass* InGameModeClass, UClass* InRootClass, bool bInIsProjectSettings)
+	{
+#define LOCTEXT_NAMESPACE "LevelToolBarViewMenu"
+		FText ResultText;
+
+		// Game modes can always be created and selected (providing the config is checked out, handled separately)
+		if(InRootClass != AGameMode::StaticClass() && InGameModeClass->HasAllClassFlags(CLASS_Native))
+		{
+			ResultText = LOCTEXT("CannotSelectClasses", "Cannot select classes when the game mode is a native class!");
+		}
+		else if(bInIsProjectSettings && !FLevelEditorActionCallbacks::CanSelectGameModeBlueprint())
+		{
+			ResultText = LOCTEXT("CannotSelectClasses_NeedsCheckOut", "Cannot select classes when the config file is not writable!");
+		}
+		else
+		{
+			FFormatNamedArguments Args;
+			Args.Add(TEXT("RootClass"), FText::FromString(InRootClass->GetName()));
+			Args.Add(TEXT("TargetLocation"), bInIsProjectSettings? LOCTEXT("Project", "project") : LOCTEXT("World", "world"));
+			ResultText = FText::Format( LOCTEXT("SelectClass_Tooltip", "Select a new {RootClass} based on a selected class and auto-assign it to the {TargetLocation}"), Args );
+		}
+		return ResultText;
+#undef LOCTEXT_NAMESPACE
+	}
+
 	void CreateGameModeSubMenu(FMenuBuilder& InMenuBuilder, TSharedRef<FUICommandList> InCommandList, TWeakPtr< SLevelEditor > InLevelEditor, bool bInProjectSettings)
 	{
 #define LOCTEXT_NAMESPACE "LevelToolBarViewMenu"
@@ -372,9 +500,12 @@ namespace LevelEditorActionHelpers
 			return true;
 		};
 
-		InMenuBuilder.AddSubMenu( LevelEditorActionHelpers::GetOpenGameModeBlueprintLabel(InLevelEditor, bInProjectSettings), LOCTEXT("GameMode_Tooltip", "Select, edit, or create a new GameMode blueprint for the world"), FNewMenuDelegate::CreateStatic( &LevelEditorActionHelpers::GetBlueprintSettingsSubMenu, InCommandList, GameModeMenuSettings), FUIAction(FExecuteAction(), FCanExecuteAction(), FIsActionChecked::CreateStatic(IsGameModeActive, InLevelEditor, bInProjectSettings)), NAME_None, EUserInterfaceActionType::RadioButton );
+		InMenuBuilder.AddSubMenu( LevelEditorActionHelpers::GetOpenGameModeBlueprintLabel(InLevelEditor, bInProjectSettings), 
+									GetSubMenuTooltip(GameModeMenuSettings.CurrentClass, GameModeMenuSettings.RootClass, bInProjectSettings), 
+									FNewMenuDelegate::CreateStatic( &LevelEditorActionHelpers::GetBlueprintSettingsSubMenu, InCommandList, GameModeMenuSettings), 
+									FUIAction(FExecuteAction(), FCanExecuteAction(), FIsActionChecked::CreateStatic(IsGameModeActive, InLevelEditor, bInProjectSettings)),
+									NAME_None, EUserInterfaceActionType::RadioButton );
 #undef LOCTEXT_NAMESPACE
-
 	}
 
 	/**
@@ -400,7 +531,10 @@ namespace LevelEditorActionHelpers
 		GameStateMenuSettings.LevelEditor = InSettingsData.LevelEditor;
 		GameStateMenuSettings.bIsProjectSettings = InSettingsData.bIsProjectSettings;
 
-		InMenuBuilder.AddSubMenu( LevelEditorActionHelpers::GetOpenGameStateBlueprintLabel(InSettingsData.LevelEditor, InSettingsData.bIsProjectSettings), LOCTEXT("GameState_Tooltip", "Select, edit, or create a new GameState blueprint for the world"), FNewMenuDelegate::CreateStatic( &LevelEditorActionHelpers::GetBlueprintSettingsSubMenu, InCommandList, GameStateMenuSettings ) );
+		InMenuBuilder.AddSubMenu( LevelEditorActionHelpers::GetOpenGameStateBlueprintLabel(InSettingsData.LevelEditor, InSettingsData.bIsProjectSettings),
+			GetSubMenuTooltip(GameStateMenuSettings.CurrentClass, GameStateMenuSettings.RootClass, InSettingsData.bIsProjectSettings),
+			FNewMenuDelegate::CreateStatic( &LevelEditorActionHelpers::GetBlueprintSettingsSubMenu, InCommandList, GameStateMenuSettings )
+		);
 
 		// Pawn
 		LevelEditorActionHelpers::FBlueprintMenuSettings PawnMenuSettings;
@@ -415,7 +549,10 @@ namespace LevelEditorActionHelpers
 		PawnMenuSettings.LevelEditor = InSettingsData.LevelEditor;
 		PawnMenuSettings.bIsProjectSettings = InSettingsData.bIsProjectSettings;
 
-		InMenuBuilder.AddSubMenu( LevelEditorActionHelpers::GetOpenPawnBlueprintLabel(InSettingsData.LevelEditor, InSettingsData.bIsProjectSettings), LOCTEXT("Pawn_Tooltip", "Select, edit, or create a new Pawn blueprint for the world"), FNewMenuDelegate::CreateStatic( &LevelEditorActionHelpers::GetBlueprintSettingsSubMenu, InCommandList, PawnMenuSettings ) );
+		InMenuBuilder.AddSubMenu( LevelEditorActionHelpers::GetOpenPawnBlueprintLabel(InSettingsData.LevelEditor, InSettingsData.bIsProjectSettings),
+			GetSubMenuTooltip(PawnMenuSettings.CurrentClass, PawnMenuSettings.RootClass, InSettingsData.bIsProjectSettings),
+			FNewMenuDelegate::CreateStatic( &LevelEditorActionHelpers::GetBlueprintSettingsSubMenu, InCommandList, PawnMenuSettings )
+		);
 
 		// HUD
 		LevelEditorActionHelpers::FBlueprintMenuSettings HUDMenuSettings;
@@ -430,7 +567,10 @@ namespace LevelEditorActionHelpers
 		HUDMenuSettings.LevelEditor = InSettingsData.LevelEditor;
 		HUDMenuSettings.bIsProjectSettings = InSettingsData.bIsProjectSettings;
 
-		InMenuBuilder.AddSubMenu( LevelEditorActionHelpers::GetOpenHUDBlueprintLabel(InSettingsData.LevelEditor, InSettingsData.bIsProjectSettings), LOCTEXT("HUD_Tooltip", "Select, edit, or create a new HUD blueprint for the world"), FNewMenuDelegate::CreateStatic( &LevelEditorActionHelpers::GetBlueprintSettingsSubMenu, InCommandList, HUDMenuSettings ) );
+		InMenuBuilder.AddSubMenu( LevelEditorActionHelpers::GetOpenHUDBlueprintLabel(InSettingsData.LevelEditor, InSettingsData.bIsProjectSettings),
+			GetSubMenuTooltip(HUDMenuSettings.CurrentClass, HUDMenuSettings.RootClass, InSettingsData.bIsProjectSettings),
+			FNewMenuDelegate::CreateStatic( &LevelEditorActionHelpers::GetBlueprintSettingsSubMenu, InCommandList, HUDMenuSettings )
+		);
 
 		// Player Controller
 		LevelEditorActionHelpers::FBlueprintMenuSettings PlayerControllerMenuSettings;
@@ -445,7 +585,10 @@ namespace LevelEditorActionHelpers
 		PlayerControllerMenuSettings.LevelEditor = InSettingsData.LevelEditor;
 		PlayerControllerMenuSettings.bIsProjectSettings = InSettingsData.bIsProjectSettings;
 
-		InMenuBuilder.AddSubMenu( LevelEditorActionHelpers::GetOpenPlayerControllerBlueprintLabel(InSettingsData.LevelEditor, InSettingsData.bIsProjectSettings), LOCTEXT("PlayerController_Tooltip", "Select, edit, or create a new Player controller blueprint for the world"), FNewMenuDelegate::CreateStatic( &LevelEditorActionHelpers::GetBlueprintSettingsSubMenu, InCommandList, PlayerControllerMenuSettings ) );
+		InMenuBuilder.AddSubMenu( LevelEditorActionHelpers::GetOpenPlayerControllerBlueprintLabel(InSettingsData.LevelEditor, InSettingsData.bIsProjectSettings),
+			GetSubMenuTooltip(PlayerControllerMenuSettings.CurrentClass, PlayerControllerMenuSettings.RootClass, InSettingsData.bIsProjectSettings),
+			FNewMenuDelegate::CreateStatic( &LevelEditorActionHelpers::GetBlueprintSettingsSubMenu, InCommandList, PlayerControllerMenuSettings )
+		);
 #undef LOCTEXT_NAMESPACE
 	}
 }
@@ -461,14 +604,39 @@ void LevelEditorActionHelpers::GetBlueprintSettingsSubMenu(FMenuBuilder& InMenuB
 	FText RootClassName = FText::FromString(InSettingsData.RootClass->GetName());
 
 	// If there is currently a valid GameMode Blueprint, offer to edit the Blueprint
-	if(InSettingsData.CurrentClass && InSettingsData.CurrentClass->ClassGeneratedBy)
+	if(InSettingsData.CurrentClass)
 	{
-		FText BlueprintName = FText::FromString(InSettingsData.CurrentClass->ClassGeneratedBy->GetName());
 		FFormatNamedArguments Args;
-		Args.Add(TEXT("Blueprint"), BlueprintName);
 		Args.Add(TEXT("RootClass"), RootClassName);
+		Args.Add(TEXT("TargetLocation"), InSettingsData.bIsProjectSettings? LOCTEXT("Project", "project") : LOCTEXT("World", "world"));
+
+		if(InSettingsData.CurrentClass->ClassGeneratedBy)
+		{
+			FText BlueprintName = FText::FromString(InSettingsData.CurrentClass->ClassGeneratedBy->GetName());
+			Args.Add(TEXT("Blueprint"), BlueprintName);
+			InMenuBuilder.AddMenuEntry( FText::Format( LOCTEXT("EditBlueprint", "Edit {Blueprint}"), Args), FText::Format( LOCTEXT("EditBlueprint_Tooltip", "Open the {TargetLocation}'s assigned {RootClass} blueprint"), Args), EditBPIcon, InSettingsData.EditCommand );
+		}
+		else
+		{
+			FText ClassName = FText::FromString(InSettingsData.CurrentClass->GetName());
+			Args.Add(TEXT("Class"), ClassName);
 
-		InMenuBuilder.AddMenuEntry( FText::Format( LOCTEXT("EditBlueprint", "Edit {Blueprint}"), Args), FText::Format( LOCTEXT("EditBlueprint_Tooltip", "Open the world's assigned {RootClass} blueprint"), Args), EditBPIcon, InSettingsData.EditCommand );
+			FText MenuDescription = FText::Format( LOCTEXT("EditNativeClass", "Edit {Class}.h"), Args);
+			if(FSourceCodeNavigation::IsCompilerAvailable())
+			{
+				InMenuBuilder.AddMenuEntry( MenuDescription, FText::Format( LOCTEXT("EditNativeClass_Tooltip", "Open the {TargetLocation}'s assigned {RootClass} header"), Args), EditBPIcon, InSettingsData.EditCommand );
+			}
+			else
+			{
+				auto CannotEditClass = []() -> bool
+				{
+					return false;
+				};
+
+				// There is no compiler present, this is always disabled with a tooltip to explain why
+				InMenuBuilder.AddMenuEntry( MenuDescription, FText::Format( LOCTEXT("CannotEditNativeClass_Tooltip", "Cannot edit the {TargetLocation}'s assigned {RootClass} header because no compiler is present!"), Args), EditBPIcon, FUIAction(FExecuteAction(), FCanExecuteAction::CreateStatic(CannotEditClass)) );
+			}
+		}
 	}
 
 	if(InSettingsData.bIsProjectSettings && InSettingsData.CurrentClass && InSettingsData.CurrentClass->IsChildOf(AGameMode::StaticClass()) && !FLevelEditorActionCallbacks::CanSelectGameModeBlueprint())
@@ -482,25 +650,54 @@ void LevelEditorActionHelpers::GetBlueprintSettingsSubMenu(FMenuBuilder& InMenuB
 		InMenuBuilder.AddMenuEntry(FLevelEditorCommands::Get().CheckOutProjectSettingsConfig, NAME_None, CheckOutLabel, CheckOutTooltip, FSlateIcon(FEditorStyle::Get().GetStyleSetName(), TEXT("Icons.Error")));
 		InMenuBuilder.EndSection();
 	}
-	else
-	{
-		// Create a new GameMode, this is always available so the user can easily create a new one
-		InMenuBuilder.AddSubMenu( LOCTEXT("CreateBlueprint", "Create..."), FText::Format( LOCTEXT("CreateClass_Tooltip", "Create a new {RootClass} based on a selected class and auto-assign it to the world"), RootClassName ), FNewMenuDelegate::CreateStatic( &LevelEditorActionHelpers::GetCreateSettingsClassSubMenu, InSettingsData.RootClass, InSettingsData.OnCreateClassPicked ), false, NewBPIcon );
 
-		// Select a game mode, this is always available so the user can switch his selection
-		InMenuBuilder.AddSubMenu( FText::Format( LOCTEXT("SelectGameModeClass", "Select {RootClass} Class"), RootClassName ), FText::Format( LOCTEXT("SelectGameModeClass_Tooltip", "Select a class to be the active {RootClass} in the world"), RootClassName ), FNewMenuDelegate::CreateStatic( &LevelEditorActionHelpers::GetSelectSettingsClassSubMenu, InSettingsData.RootClass, InSettingsData.OnSelectClassPicked ) );
-
-		if(InSettingsData.RootClass == AGameMode::StaticClass())
+	auto CannotCreateSelectNativeProjectGameMode = [](bool bInIsProjectSettings) -> bool
+	{
+		// For the project settings, we can only create/select the game mode class if the config is writable
+		if(bInIsProjectSettings)
 		{
-			InMenuBuilder.BeginSection(NAME_None, LOCTEXT("GameModeClasses", "Game Mode Classes"));
-
-			if(InSettingsData.CurrentClass &&InSettingsData.CurrentClass->ClassGeneratedBy)
-			{
-				GetGameModeSubMenu(InMenuBuilder, InCommandList, InSettingsData);
-			}
+			return FLevelEditorActionCallbacks::CanSelectGameModeBlueprint();
+		}
+		return true;
+	};
 
-			InMenuBuilder.EndSection();
+	// Create a new GameMode, this is always available so the user can easily create a new one
+	InMenuBuilder.AddSubMenu( LOCTEXT("CreateBlueprint", "Create..."),
+		GetCreateMenuTooltip(GetGameModeClass(InSettingsData.LevelEditor, InSettingsData.bIsProjectSettings), InSettingsData.RootClass, InSettingsData.bIsProjectSettings),
+		FNewMenuDelegate::CreateStatic( &LevelEditorActionHelpers::GetCreateSettingsClassSubMenu, InSettingsData.RootClass, InSettingsData.OnCreateClassPicked ),
+		FUIAction(
+			FExecuteAction(), 
+			InSettingsData.RootClass == AGameMode::StaticClass()? 
+				FCanExecuteAction::CreateStatic(CannotCreateSelectNativeProjectGameMode, InSettingsData.bIsProjectSettings) 
+				: FCanExecuteAction::CreateStatic( &CanCreateSelectSubClass, GetGameModeClass(InSettingsData.LevelEditor, InSettingsData.bIsProjectSettings), InSettingsData.bIsProjectSettings )
+		),
+		NAME_None, EUserInterfaceActionType::Button, false, NewBPIcon
+	);
+
+	// Select a game mode, this is always available so the user can switch his selection
+	InMenuBuilder.AddSubMenu( FText::Format( LOCTEXT("SelectGameModeClass", "Select {RootClass} Class"), RootClassName ),
+		GetSelectMenuTooltip(GetGameModeClass(InSettingsData.LevelEditor, InSettingsData.bIsProjectSettings), InSettingsData.RootClass, InSettingsData.bIsProjectSettings),
+		FNewMenuDelegate::CreateStatic( &LevelEditorActionHelpers::GetSelectSettingsClassSubMenu, InSettingsData.RootClass, InSettingsData.OnSelectClassPicked ),
+		FUIAction(
+			FExecuteAction(), 
+			InSettingsData.RootClass == AGameMode::StaticClass()?
+				FCanExecuteAction::CreateStatic(CannotCreateSelectNativeProjectGameMode, InSettingsData.bIsProjectSettings) 
+				: FCanExecuteAction::CreateStatic( &CanCreateSelectSubClass, GetGameModeClass(InSettingsData.LevelEditor, InSettingsData.bIsProjectSettings), InSettingsData.bIsProjectSettings )
+		),
+		NAME_None, EUserInterfaceActionType::Button
+	);
+
+	// For GameMode classes only, there are some sub-classes we need to add to the menu
+	if(InSettingsData.RootClass == AGameMode::StaticClass())
+	{
+		InMenuBuilder.BeginSection(NAME_None, LOCTEXT("GameModeClasses", "Game Mode Classes"));
+
+		if(InSettingsData.CurrentClass)
+		{
+			GetGameModeSubMenu(InMenuBuilder, InCommandList, InSettingsData);
 		}
+
+		InMenuBuilder.EndSection();
 	}
 
 	InMenuBuilder.PopCommandList();
