Added proper material picker to BSP surfaces

This addresses TTP#345459 - NUX:BSP: Changing or removing the material on a BSP surface was confusing

diff --git a/Engine/Source/Editor/LevelEditor/LevelEditor.Build.cs b/Engine/Source/Editor/LevelEditor/LevelEditor.Build.cs
--- a/Engine/Source/Editor/LevelEditor/LevelEditor.Build.cs
+++ b/Engine/Source/Editor/LevelEditor/LevelEditor.Build.cs
@@ -51,6 +51,7 @@ public class LevelEditor : ModuleRules
 				"Projects",
 				"TargetPlatform",
 				"EngineSettings",
+				"PropertyEditor",
 			}
 		);
 
diff --git a/Engine/Source/Editor/LevelEditor/Private/LevelEditorGenericDetails.cpp b/Engine/Source/Editor/LevelEditor/Private/LevelEditorGenericDetails.cpp
--- a/Engine/Source/Editor/LevelEditor/Private/LevelEditorGenericDetails.cpp
+++ b/Engine/Source/Editor/LevelEditor/Private/LevelEditorGenericDetails.cpp
@@ -7,6 +7,9 @@
 #include "AssetSelection.h"
 #include "SSurfaceProperties.h"
 #include "Runtime/AssetRegistry/Public/AssetRegistryModule.h"
+#include "SurfaceIterators.h"
+
+#include "PropertyCustomizationHelpers.h"
 
 #define LOCTEXT_NAMESPACE "FLevelEditorGenericDetails"
 
@@ -25,6 +28,43 @@ void FLevelEditorGenericDetails::CustomizeDetails( IDetailLayoutBuilder& DetailL
 	}
 }
 
+void FLevelEditorGenericDetails::GetSelectedSurfaceMaterials(IMaterialListBuilder& MaterialList) const
+{
+	for (TSelectedSurfaceIterator<> SurfaceIt(GWorld); SurfaceIt; ++SurfaceIt)
+	{
+        FBspSurf* SelectedSurface = *SurfaceIt;
+
+		const bool bCanBeReplaced = true;
+		MaterialList.AddMaterial( 0, SelectedSurface->Material, bCanBeReplaced );
+	}
+}
+
+void FLevelEditorGenericDetails::OnMaterialChanged( UMaterialInterface* NewMaterial, UMaterialInterface* PrevMaterial, int32 SlotIndex, bool bReplaceAll )
+{
+	bool bModelDirtied = false;
+	{
+		const FScopedTransaction Transaction( NSLOCTEXT("UnrealEd", "PolySetMaterial", "Set Material") );
+		GEditor->FlagModifyAllSelectedSurfacesInLevels( GWorld );
+
+		for ( TSelectedSurfaceIterator<> It(GWorld) ; It ; ++It )
+		{
+			UModel* Model = It.GetModel();
+			const int32 SurfaceIndex = It.GetSurfaceIndex();
+
+			Model->Surfs[SurfaceIndex].Material = NewMaterial;
+			GEditor->polyUpdateMaster( Model, SurfaceIndex, 0 );
+			Model->MarkPackageDirty();
+
+			bModelDirtied = true;
+		}
+	}
+	GEditor->RedrawLevelEditingViewports();
+	if ( bModelDirtied )
+	{
+		ULevel::LevelDirtiedEvent.Broadcast();
+	}
+}
+
 void FLevelEditorGenericDetails::AddSurfaceDetails( IDetailLayoutBuilder& DetailBuilder )
 {
 	FLevelEditorModule& LevelEditorModule = FModuleManager::GetModuleChecked<FLevelEditorModule>( TEXT("LevelEditor") );
@@ -81,57 +121,6 @@ void FLevelEditorGenericDetails::AddSurfaceDetails( IDetailLayoutBuilder& Detail
 
 	struct Local
 	{
-		static TSharedRef<SWidget> OnGetApplyMaterialMenu( const FLevelEditorCommands* InCommands, const TSharedPtr<const FUICommandList> InCommandBindings )
-		{
-			// Get the first material in the selection
-			// If it exists allow it to be applied to the selected surfaces
-			UMaterialInterface* Material = GEditor->GetSelectedObjects()->GetTop<UMaterialInterface>();
-
-			//Check for unloaded materials...
-			if (!Material)
-			{
-				TArray<FAssetData> Assets;
-				AssetSelectionUtils::GetSelectedAssets( Assets );				
-				for (int i=0; i<Assets.Num(); i++)
-				{
-					FAssetData* AssetData = &Assets[i];
-					if ( AssetData->GetClass()->IsChildOf(UMaterialInterface::StaticClass()) )
-					{
-						Material = Cast<UMaterialInterface>( AssetData->GetAsset() );
-						if (Material)
-						{
-							break;
-						}
-					}	
-				}
-			}
-
-			FText ApplyMaterialStr;
-			if( Material )
-			{
-				// Generate a string indicating what material will be applied
-				ApplyMaterialStr = FText::Format( LOCTEXT("ApplyMaterialToSurface", "Apply Material: {0}"), FText::FromString( Material->GetName() ) );
-			}
-
-			if( Material )
-			{
-				FMenuBuilder ApplyMaterialBuilder( true, InCommandBindings );
-				{
-					ApplyMaterialBuilder.AddMenuEntry( InCommands->ApplyMaterialToSurface, NAME_None, ApplyMaterialStr );
-				}
-
-				return ApplyMaterialBuilder.MakeWidget();
-			}
-			else
-			{
-				// No material is selected, show a generic message
-				return
-					SNew( STextBlock )
-					.Text( LOCTEXT("NoMaterialSelected", "No Material Selected") )
-					.IsEnabled( false );
-			}
-		}
-
 		static FReply ExecuteExecCommand(FString InCommand)
 		{
 			GUnrealEd->Exec( GWorld, *InCommand );
@@ -141,6 +130,17 @@ void FLevelEditorGenericDetails::AddSurfaceDetails( IDetailLayoutBuilder& Detail
 
 	const FSlateFontInfo& FontInfo = IDetailLayoutBuilder::GetDetailFont();
 
+	IDetailCategoryBuilder& MaterialsCategory = DetailBuilder.EditCategory( "Surface Materials", LOCTEXT("BSPSurfaceMaterials", "Surface Materials").ToString() );
+	{
+		FMaterialListDelegates MaterialListDelegates;
+		MaterialListDelegates.OnGetMaterials.BindSP(this, &FLevelEditorGenericDetails::GetSelectedSurfaceMaterials);
+		MaterialListDelegates.OnMaterialChanged.BindSP(this, &FLevelEditorGenericDetails::OnMaterialChanged);
+
+		TSharedRef<FMaterialList> MaterialList = MakeShareable(new FMaterialList(DetailBuilder, MaterialListDelegates));
+
+		MaterialsCategory.AddCustomBuilder( MaterialList );
+	}
+
 	// Add a new section for static meshes
 	IDetailCategoryBuilder& BSPCategory = DetailBuilder.EditCategory( "Geometry", LOCTEXT("BSPSurfacesTitle", "Geometry").ToString() );
 	BSPCategory.AddCustomRow( TEXT("") )
@@ -183,19 +183,6 @@ void FLevelEditorGenericDetails::AddSurfaceDetails( IDetailLayoutBuilder& Detail
 					AlignmentBuilder.MakeWidget()
 				]
 			]
-			+ SHorizontalBox::Slot()
-			.AutoWidth()
-			.Padding( 2.0f )
-			[
-				SNew( SComboButton )
-				.ButtonContent()
-				[
-					SNew( STextBlock )
-					.Text( LOCTEXT("ApplyMaterialMenu", "Apply Material") ) 
-					.Font( FontInfo )
-				]
-				.OnGetMenuContent_Static( &Local::OnGetApplyMaterialMenu, &Commands, CommandBindings )
-			]
 		]
 	];
 
diff --git a/Engine/Source/Editor/LevelEditor/Private/LevelEditorGenericDetails.h b/Engine/Source/Editor/LevelEditor/Private/LevelEditorGenericDetails.h
--- a/Engine/Source/Editor/LevelEditor/Private/LevelEditorGenericDetails.h
+++ b/Engine/Source/Editor/LevelEditor/Private/LevelEditorGenericDetails.h
@@ -2,6 +2,8 @@
 
 #pragma once
 
+class IMaterialListBuilder;
+
 /**
  * Creates details for the level editor details view that are not specific to any selected actor type
  */
@@ -19,5 +21,15 @@ private:
 	 * Adds a section for selected surface details
 	 */
 	void AddSurfaceDetails( class IDetailLayoutBuilder& DetailLayout );
+
+	/**
+	 * Populate the specified material list with the materials used on the currently selected BSP surfaces
+	 */
+	void GetSelectedSurfaceMaterials(IMaterialListBuilder& MaterialList) const;
+
+	/**
+	 * Called when the material should be changed on all selected BSP surfaces
+	 */
+	void OnMaterialChanged( UMaterialInterface* NewMaterial, UMaterialInterface* PrevMaterial, int32 SlotIndex, bool bReplaceAll );
 };
 
diff --git a/Engine/Source/Editor/LevelEditor/Private/SSurfaceProperties.cpp b/Engine/Source/Editor/LevelEditor/Private/SSurfaceProperties.cpp
--- a/Engine/Source/Editor/LevelEditor/Private/SSurfaceProperties.cpp
+++ b/Engine/Source/Editor/LevelEditor/Private/SSurfaceProperties.cpp
@@ -38,21 +38,6 @@ void SSurfaceProperties::Construct( const FArguments& InArgs )
 	[
 		SNew(SVerticalBox)
 
-		+ SVerticalBox::Slot()
-		.AutoHeight()
-		.Padding(0, 0, 0, 5)
-		[
-			SNew(SBorder)
-			.BorderBackgroundColor(BorderColor)
-			.BorderImage(BorderStyle)
-			.Padding(10)
-			.AddMetaData<FTagMetaData>(TEXT("DetailsView.SelectedSurfaces"))
-			[
-				SNew(STextBlock)
-				.Text(this, &SSurfaceProperties::GetSelectedSurfacesDescription)
-			]
-		]
-
 		+SVerticalBox::Slot()
 		.AutoHeight()
 		.Padding(0,0,0,5)
@@ -990,56 +975,4 @@ void SSurfaceProperties::OnToggleRotationDirection( ESlateCheckBoxState::Type Ne
 	bUseNegativeRotation = (NewState == ESlateCheckBoxState::Checked) ? true : false;
 }
 
-FText SSurfaceProperties::GetSelectedSurfacesDescription() const
-{
-	int SurfaceCount = 0;
-	UMaterialInterface* SurfaceMaterial = nullptr;
-	bool bMultipleMaterials = false;
-
-	// This would be easier to follow using a set to keep track of used materials
-	// but since it's called every frame a single pointer and boolean are used instead.
-	for (TSelectedSurfaceIterator<> SurfaceIt(GetWorld()); SurfaceIt; ++SurfaceIt)
-	{
-		FBspSurf* SelectedSurface = *SurfaceIt;
-		++SurfaceCount;
-		if (!bMultipleMaterials)
-		{
-			if (SurfaceMaterial != SelectedSurface->Material)
-			{
-				if (SurfaceCount == 1)
-				{
-					SurfaceMaterial = SelectedSurface->Material;
-				}
-				else
-				{
-					SurfaceMaterial = nullptr;
-					bMultipleMaterials = true;
-				}
-			}
-		}
-	}
-
-	FText MaterialText;
-	if (bMultipleMaterials)
-	{
-		MaterialText = LOCTEXT("MultipleMaterials", "Multiple materials");
-	}
-	else
-	{
-		if (SurfaceMaterial == nullptr)
-		{
-			MaterialText = LOCTEXT("DefaultMaterial", "Default material");
-		}
-		else
-		{
-			MaterialText = FText::FromName(SurfaceMaterial->GetFName());
-		}
-	}
-
-	return FText::Format(LOCTEXT("SurfaceDescriptionFormat", "{0} {1}: {2}"),
-		FText::AsNumber(SurfaceCount),
-		SurfaceCount == 1 ? LOCTEXT("SurfaceSingular", "Surface") : LOCTEXT("SurfacePlural", "Surfaces"),
-		MaterialText);
-}
-
 #undef LOCTEXT_NAMESPACE
diff --git a/Engine/Source/Editor/LevelEditor/Private/SSurfaceProperties.h b/Engine/Source/Editor/LevelEditor/Private/SSurfaceProperties.h
--- a/Engine/Source/Editor/LevelEditor/Private/SSurfaceProperties.h
+++ b/Engine/Source/Editor/LevelEditor/Private/SSurfaceProperties.h
@@ -235,12 +235,6 @@ private:
 	 */
 	void OnToggleRotationDirection( ESlateCheckBoxState::Type NewState );
 
-	/**
-	 * Returns a short summary to display to the user about the currently selected surfaces.
-	**/
-	FText GetSelectedSurfacesDescription() const;
-
-
 	/** typedef to make it easier to handle the lightmass objects */
 	typedef TArray<ULightmassPrimitiveSettingsObject*> TLightmassSettingsObjectArray;
 
diff --git a/Engine/Source/Editor/UnrealEd/Classes/Editor/EditorEngine.h b/Engine/Source/Editor/UnrealEd/Classes/Editor/EditorEngine.h
--- a/Engine/Source/Editor/UnrealEd/Classes/Editor/EditorEngine.h
+++ b/Engine/Source/Editor/UnrealEd/Classes/Editor/EditorEngine.h
@@ -2473,12 +2473,16 @@ private:
 	void ExecuteCommandForAllLevelModels( UWorld* InWorld, FSelectCommand InSelectCommand, const FText& TransDesription );
 	void ExecuteCommandForAllLevelModels( UWorld* InWorld, FSelectInWorldCommand InSelectCommand, const FText& TransDesription );
 	
+public:
+
 	/**
 	 * Utility method call ModifySelectedSurfs for each level model in the worlds level list.
 	 *
 	 * @param InWorld			World containing the levels
 	 */
 	 void FlagModifyAllSelectedSurfacesInLevels( UWorld* InWorld );
+	 
+private:
 
 	/**
 	 * This destroys the given world.
