CIS fix for non-unity builds

diff --git a/Engine/Source/Editor/LevelEditor/Private/LevelEditorGenericDetails.cpp b/Engine/Source/Editor/LevelEditor/Private/LevelEditorGenericDetails.cpp
--- a/Engine/Source/Editor/LevelEditor/Private/LevelEditorGenericDetails.cpp
+++ b/Engine/Source/Editor/LevelEditor/Private/LevelEditorGenericDetails.cpp
@@ -8,6 +8,7 @@
 #include "SSurfaceProperties.h"
 #include "Runtime/AssetRegistry/Public/AssetRegistryModule.h"
 #include "SurfaceIterators.h"
+#include "ScopedTransaction.h"
 
 #include "PropertyCustomizationHelpers.h"
 
