Fix crash on taking high-res screenshots on OS X - there's no root window on OS X.

diff --git a/Engine/Source/Editor/LevelEditor/Private/HighResScreenshotUI.cpp b/Engine/Source/Editor/LevelEditor/Private/HighResScreenshotUI.cpp
--- a/Engine/Source/Editor/LevelEditor/Private/HighResScreenshotUI.cpp
+++ b/Engine/Source/Editor/LevelEditor/Private/HighResScreenshotUI.cpp
@@ -296,7 +296,14 @@ TWeakPtr<class SWindow> SHighResScreenshotDialog::OpenDialog(const TSharedPtr<FS
 
 		TSharedPtr<SWindow> ParentWindow = FGlobalTabmanager::Get()->GetRootWindow();
 
-		FSlateApplication::Get().AddWindowAsNativeChild(Window, ParentWindow.ToSharedRef());
+		if (ParentWindow.IsValid())
+		{
+			FSlateApplication::Get().AddWindowAsNativeChild(Window, ParentWindow.ToSharedRef());
+		}
+		else
+		{
+			FSlateApplication::Get().AddWindow(Window);
+		}
 
 		CurrentWindow = TWeakPtr<SWindow>(Window);
 		CurrentDialog = TWeakPtr<SHighResScreenshotDialog>(Dialog);
