Update Mac menu bar when user adds or removes Window menu items by changing editor's experimental settings

diff --git a/Engine/Source/Editor/LevelEditor/Private/SLevelEditor.cpp b/Engine/Source/Editor/LevelEditor/Private/SLevelEditor.cpp
--- a/Engine/Source/Editor/LevelEditor/Private/SLevelEditor.cpp
+++ b/Engine/Source/Editor/LevelEditor/Private/SLevelEditor.cpp
@@ -1152,6 +1152,9 @@ TSharedRef<SWidget> SLevelEditor::RestoreContentArea( const TSharedRef<SDockTab>
 
 void SLevelEditor::HandleExperimentalSettingChanged(FName PropertyName)
 {
+	FLevelEditorModule& LevelEditorModule = FModuleManager::GetModuleChecked<FLevelEditorModule>("LevelEditor");
+	TSharedPtr<FTabManager> LevelEditorTabManager = LevelEditorModule.GetLevelEditorTabManager();
+	LevelEditorTabManager->UpdateMainMenu(true);
 }
 
 FName SLevelEditor::GetEditorModeTabId( FEditorModeID ModeID )
