Apply feature level change to all running PIE sessions.

diff --git a/Engine/Source/Editor/LevelEditor/Private/LevelEditorActions.cpp b/Engine/Source/Editor/LevelEditor/Private/LevelEditorActions.cpp
--- a/Engine/Source/Editor/LevelEditor/Private/LevelEditorActions.cpp
+++ b/Engine/Source/Editor/LevelEditor/Private/LevelEditorActions.cpp
@@ -457,6 +457,18 @@ bool FLevelEditorActionCallbacks::IsMaterialQualityLevelChecked( EMaterialQualit
 void FLevelEditorActionCallbacks::SetFeatureLevelPreview(ERHIFeatureLevel::Type InPreviewFeatureLevel)
 {
 	GetWorld()->ChangeFeatureLevel(InPreviewFeatureLevel);
+
+	// Update any currently running PIE sessions.
+	for (TObjectIterator<UWorld> It; It; ++It)
+	{
+		UWorld* ItWorld = *It;
+		if (ItWorld->WorldType == EWorldType::PIE)
+		{
+			ItWorld->ChangeFeatureLevel(InPreviewFeatureLevel);
+		}
+	}
+
+	GUnrealEd->RedrawAllViewports();
 }
 
 bool FLevelEditorActionCallbacks::IsFeatureLevelPreviewChecked(ERHIFeatureLevel::Type InPreviewFeatureLevel)
