Fixed issue where saved ShowFlags were being obliterated immediately after load.

diff --git a/Engine/Source/Editor/LevelEditor/Private/SLevelViewport.cpp b/Engine/Source/Editor/LevelEditor/Private/SLevelViewport.cpp
--- a/Engine/Source/Editor/LevelEditor/Private/SLevelViewport.cpp
+++ b/Engine/Source/Editor/LevelEditor/Private/SLevelViewport.cpp
@@ -127,6 +127,7 @@ void SLevelViewport::Construct(const FArguments& InArgs)
 
 	ParentLayout = InArgs._ParentLayout;
 	ParentLevelEditor = InArgs._ParentLevelEditor;
+	ConfigKey = InArgs._ConfigKey;
 
 	// Store border brushes for differentiating between active and inactive viewports
 	ActiveBorder = FEditorStyle::GetBrush( "LevelViewport.ActiveViewportBorder" );
@@ -284,7 +285,6 @@ void SLevelViewport::ConstructLevelEditorViewportClient( const FArguments& InArg
 	FEngineShowFlags GameShowFlags(ESFIM_Game);
 		
 	// Use config key if it exists to set up the level viewport client
-	FString ConfigKey = InArgs._ConfigKey;
 	if(!ConfigKey.IsEmpty())
 	{
 		const FLevelEditorViewportInstanceSettings* const ViewportInstanceSettingsPtr = GetDefault<ULevelEditorViewportSettings>()->GetViewportInstanceSettings(ConfigKey);
@@ -1269,7 +1269,7 @@ void SLevelViewport::BindShowCommands( FUICommandList& CommandList )
 {
 	CommandList.MapAction( 
 		FLevelViewportCommands::Get().UseDefaultShowFlags,
-		FExecuteAction::CreateSP( this, &SLevelViewport::OnUseDefaultShowFlags ) );
+		FExecuteAction::CreateSP( this, &SLevelViewport::OnUseDefaultShowFlags, false ) );
 
 	const TArray<FShowFlagData>& ShowFlagData = GetShowFlagMenuItems();
 
@@ -1835,7 +1835,7 @@ bool SLevelViewport::IsShowFlagEnabled( uint32 EngineShowFlagIndex ) const
 	return LevelViewportClient->EngineShowFlags.GetSingleFlag(EngineShowFlagIndex);
 }
 
-void SLevelViewport::OnUseDefaultShowFlags()
+void SLevelViewport::OnUseDefaultShowFlags(bool bUseSavedDefaults)
 {
 	// cache off the current viewmode as it gets trashed when applying FEngineShowFlags()
 	const EViewModeIndex CachedViewMode = LevelViewportClient->GetViewMode();
@@ -1843,11 +1843,32 @@ void SLevelViewport::OnUseDefaultShowFlags()
 	// Setting show flags to the defaults should not stomp on the current viewmode settings.
 	LevelViewportClient->SetGameView(false);
 
+	// Get default save flags
+	FEngineShowFlags EditorShowFlags(ESFIM_Editor);
+	FEngineShowFlags GameShowFlags(ESFIM_Game);
+
+	if (bUseSavedDefaults && !ConfigKey.IsEmpty())
+	{
+		// Get saved defaults if specified
+		const FLevelEditorViewportInstanceSettings* const ViewportInstanceSettingsPtr = GetDefault<ULevelEditorViewportSettings>()->GetViewportInstanceSettings(ConfigKey);
+		FLevelEditorViewportInstanceSettings ViewportInstanceSettings = (ViewportInstanceSettingsPtr) ? *ViewportInstanceSettingsPtr : LoadLegacyConfigFromIni(ConfigKey, ViewportInstanceSettings);
+
+		if (!ViewportInstanceSettings.EditorShowFlagsString.IsEmpty())
+		{
+			EditorShowFlags.SetFromString(*ViewportInstanceSettings.EditorShowFlagsString);
+		}
+
+		if (!ViewportInstanceSettings.GameShowFlagsString.IsEmpty())
+		{
+			GameShowFlags.SetFromString(*ViewportInstanceSettings.GameShowFlagsString);
+		}
+	}
+
 	// this trashes the current viewmode!
-	LevelViewportClient->EngineShowFlags = FEngineShowFlags(ESFIM_Editor);
+	LevelViewportClient->EngineShowFlags = EditorShowFlags;
 	// Restore the state of SelectionOutline based on user settings
 	LevelViewportClient->EngineShowFlags.SelectionOutline = GetDefault<ULevelEditorViewportSettings>()->bUseSelectionOutline;
-	LevelViewportClient->LastEngineShowFlags = FEngineShowFlags(ESFIM_Game);
+	LevelViewportClient->LastEngineShowFlags = GameShowFlags;
 
 	// re-apply the cached viewmode, as it was trashed with FEngineShowFlags()
 	ApplyViewMode(CachedViewMode, LevelViewportClient->IsPerspective(), LevelViewportClient->EngineShowFlags);
@@ -1857,7 +1878,6 @@ void SLevelViewport::OnUseDefaultShowFlags()
 }
 
 
-
 void SLevelViewport::SetKeyboardFocusToThisViewport()
 {
 	if( ensure( ViewportWidget.IsValid() ) )
@@ -3391,7 +3411,8 @@ void SLevelViewport::ResetNewLevelViewFlags()
 	const bool bIsPerspective = LevelViewportClient->IsPerspective();
 	LevelViewportClient->SetViewMode(bIsPerspective ? FEditorViewportClient::DefaultPerspectiveViewMode : FEditorViewportClient::DefaultOrthoViewMode);
 
-	OnUseDefaultShowFlags();
+	const bool bUseSavedDefaults = true;
+	OnUseDefaultShowFlags(bUseSavedDefaults);
 }
 
 void SLevelViewport::EndPlayInEditorSession()
diff --git a/Engine/Source/Editor/LevelEditor/Public/SLevelViewport.h b/Engine/Source/Editor/LevelEditor/Public/SLevelViewport.h
--- a/Engine/Source/Editor/LevelEditor/Public/SLevelViewport.h
+++ b/Engine/Source/Editor/LevelEditor/Public/SLevelViewport.h
@@ -426,9 +426,9 @@ private:
 	virtual void ToggleStatCommand(FString CommandName) override;
 
 	/**
-	 * Called when show flags for this viewport should be reset to default
+	 * Called when show flags for this viewport should be reset to default, or the saved settings
 	 */
-	void OnUseDefaultShowFlags();
+	void OnUseDefaultShowFlags(bool bUseSavedDefaults = false);
 
 	/**
 	 * Changes the buffer visualization mode for this viewport
@@ -746,6 +746,9 @@ private:
 	/** The current device profile string */
 	FString DeviceProfile;
 
+	/** The current viewport config key */
+	FString ConfigKey;
+
 	/**
 	 * Contains information about an actor being previewed within this viewport
 	 */
