Fixed uninitialized variable oversight from CL 2336924 which Visual Studio didn't catch and which was causing build warnings on other platforms.

diff --git a/Engine/Source/Editor/LevelEditor/Private/SLevelViewport.cpp b/Engine/Source/Editor/LevelEditor/Private/SLevelViewport.cpp
--- a/Engine/Source/Editor/LevelEditor/Private/SLevelViewport.cpp
+++ b/Engine/Source/Editor/LevelEditor/Private/SLevelViewport.cpp
@@ -1849,9 +1849,12 @@ void SLevelViewport::OnUseDefaultShowFlags(bool bUseSavedDefaults)
 
 	if (bUseSavedDefaults && !ConfigKey.IsEmpty())
 	{
+		FLevelEditorViewportInstanceSettings ViewportInstanceSettings;
+		ViewportInstanceSettings.ViewportType = LevelViewportClient->ViewportType;
+
 		// Get saved defaults if specified
 		const FLevelEditorViewportInstanceSettings* const ViewportInstanceSettingsPtr = GetDefault<ULevelEditorViewportSettings>()->GetViewportInstanceSettings(ConfigKey);
-		FLevelEditorViewportInstanceSettings ViewportInstanceSettings = (ViewportInstanceSettingsPtr) ? *ViewportInstanceSettingsPtr : LoadLegacyConfigFromIni(ConfigKey, ViewportInstanceSettings);
+		ViewportInstanceSettings = ViewportInstanceSettingsPtr ? *ViewportInstanceSettingsPtr : LoadLegacyConfigFromIni(ConfigKey, ViewportInstanceSettings);
 
 		if (!ViewportInstanceSettings.EditorShowFlagsString.IsEmpty())
 		{
