Stats - Added misc comment

diff --git a/Engine/Source/Editor/LevelEditor/Private/LevelViewportActions.cpp b/Engine/Source/Editor/LevelEditor/Private/LevelViewportActions.cpp
--- a/Engine/Source/Editor/LevelEditor/Private/LevelViewportActions.cpp
+++ b/Engine/Source/Editor/LevelEditor/Private/LevelViewportActions.cpp
@@ -264,6 +264,7 @@ void FLevelViewportCommands::RegisterCommands()
 
 void FLevelViewportCommands::HandleNewStatGroup(const TArray<FStatNameAndInfo>& NameAndInfos)
 {
+	// @TODO yrx 2014-10-17 FStatNameAndInfo should be private and visible only to stats2 system
 	for (int32 InfoIdx = 0; InfoIdx < NameAndInfos.Num(); InfoIdx++)
 	{
 		const FStatNameAndInfo& NameAndInfo = NameAndInfos[InfoIdx];
