refactor

diff --git a/Engine/Source/Editor/LevelEditor/Private/SLevelViewport.cpp b/Engine/Source/Editor/LevelEditor/Private/SLevelViewport.cpp
--- a/Engine/Source/Editor/LevelEditor/Private/SLevelViewport.cpp
+++ b/Engine/Source/Editor/LevelEditor/Private/SLevelViewport.cpp
@@ -811,7 +811,7 @@ void SLevelViewport::Tick( const FGeometry& AllottedGeometry, const double InCur
 	// viewport with Slate widgets that are part of the game, don't throttle.
 	if ( bPIEContainsFocus != bContainsFocus )
 	{
-		IConsoleVariable* AllowThrottling = IConsoleManager::Get().FindConsoleVariable(TEXT("Slate.bAllowThrottling"));
+		static IConsoleVariable* AllowThrottling = IConsoleManager::Get().FindConsoleVariable(TEXT("Slate.bAllowThrottling"));
 		check(AllowThrottling);
 
 		if ( bContainsFocus )
