Fixed not being able to exit immersive mode

diff --git a/Engine/Source/Editor/LevelEditor/Private/LevelViewportLayout.cpp b/Engine/Source/Editor/LevelEditor/Private/LevelViewportLayout.cpp
--- a/Engine/Source/Editor/LevelEditor/Private/LevelViewportLayout.cpp
+++ b/Engine/Source/Editor/LevelEditor/Private/LevelViewportLayout.cpp
@@ -99,7 +99,6 @@ void SViewportsOverlay::Construct( const FArguments& InArgs )
 void SViewportsOverlay::Tick( const FGeometry& AllottedGeometry, const double InCurrentTime, const float InDeltaTime )
 {
 	CachedSize = AllottedGeometry.Size;
-	LevelViewportTab->Tick();
 }
 
 SOverlay::FOverlaySlot& SViewportsOverlay::AddSlot()
@@ -271,7 +270,7 @@ void FLevelViewportLayout::RequestMaximizeViewport( TSharedRef<class SLevelViewp
 
 		// We flush commands here because there could be a pending slow viewport draw already enqueued in the render thread
 		// We take the hitch here so that our transition to/from maximize animation is responsive next tick
-		FSlateApplication::Get().GetRenderer()->FlushCommands();
+		FlushRenderingCommands();
 
 		DeferredMaximizeCommands.Add( FMaximizeViewportCommand(ViewportToMaximize, bWantMaximize, bWantImmersive) );
 	}
@@ -569,7 +568,7 @@ bool FLevelViewportLayout::IsViewportImmersive( const SLevelViewport& InViewport
 EVisibility FLevelViewportLayout::OnGetNonMaximizedVisibility() const
 {
 	// The non-maximized viewports are not visible if there is a maximized viewport on top of them
-	return (!bIsQueryingLayoutMetrics && MaximizedViewport.IsValid() && !bIsTransitioning) ? EVisibility::Collapsed : EVisibility::Visible;
+	return (!bIsQueryingLayoutMetrics && MaximizedViewport.IsValid() && !bIsTransitioning && DeferredMaximizeCommands.Num() == 0) ? EVisibility::Collapsed : EVisibility::Visible;
 }
 
 
@@ -669,7 +668,7 @@ void FLevelViewportLayout::FinishMaximizeTransition()
 }
 
 
-void FLevelViewportLayout::Tick()
+void FLevelViewportLayout::Tick( float DeltaTime )
 {
 	// If we have an animation that has finished playing, then complete the transition
 	if( bIsTransitioning && !MaximizeAnimation.IsPlaying() )
@@ -696,3 +695,8 @@ void FLevelViewportLayout::Tick()
 		DeferredMaximizeCommands.Empty();
 	}
 }
+
+bool FLevelViewportLayout::IsTickable() const
+{
+	return DeferredMaximizeCommands.Num() > 0 || (bIsTransitioning && !MaximizeAnimation.IsPlaying());
+}
diff --git a/Engine/Source/Editor/LevelEditor/Private/LevelViewportLayout.h b/Engine/Source/Editor/LevelEditor/Private/LevelViewportLayout.h
--- a/Engine/Source/Editor/LevelEditor/Private/LevelViewportLayout.h
+++ b/Engine/Source/Editor/LevelEditor/Private/LevelViewportLayout.h
@@ -13,7 +13,7 @@
  * Base class for level viewport layout configurations
  * Handles maximizing and restoring well as visibility of specific viewports.
  */
-class FLevelViewportLayout : public TSharedFromThis<FLevelViewportLayout>
+class FLevelViewportLayout : public TSharedFromThis<FLevelViewportLayout>, FTickableEditorObject
 {
 public:
 	/**
@@ -95,8 +95,12 @@ public:
 
 	virtual const FName& GetLayoutTypeName() const = 0;
 
-	/** Ticks this level viewport layout object */
-	virtual void Tick();
+
+	/** FTickableEditorObject interface */
+	virtual void Tick( float DeltaTime ) override;
+	virtual bool IsTickable() const override;
+	virtual TStatId GetStatId() const override { return TStatId(); }
+
 
 	/** Tells this layout whether it was the intial layout or replaced an existing one when the user switched layouts */
 	void SetIsReplacement(bool bInIsReplacement) { bIsReplacement = bInIsReplacement; }
diff --git a/Engine/Source/Editor/LevelEditor/Private/LevelViewportTabContent.cpp b/Engine/Source/Editor/LevelEditor/Private/LevelViewportTabContent.cpp
--- a/Engine/Source/Editor/LevelEditor/Private/LevelViewportTabContent.cpp
+++ b/Engine/Source/Editor/LevelEditor/Private/LevelViewportTabContent.cpp
@@ -74,14 +74,6 @@ const TArray< TSharedPtr< SLevelViewport > >* FLevelViewportTabContent::GetViewp
 	return NULL;
 }
 
-void FLevelViewportTabContent::Tick()
-{
-	if (ActiveLevelViewportLayout.IsValid())
-	{
-		ActiveLevelViewportLayout->Tick();
-	}
-}
-
 void FLevelViewportTabContent::SetViewportConfiguration(const FName& ConfigurationName)
 {
 	bool bSwitchingLayouts = ActiveLevelViewportLayout.IsValid();
diff --git a/Engine/Source/Editor/LevelEditor/Private/LevelViewportTabContent.h b/Engine/Source/Editor/LevelEditor/Private/LevelViewportTabContent.h
--- a/Engine/Source/Editor/LevelEditor/Private/LevelViewportTabContent.h
+++ b/Engine/Source/Editor/LevelEditor/Private/LevelViewportTabContent.h
@@ -25,9 +25,6 @@ public:
 	/** Returns an array of viewports inside this tab */
 	const TArray< TSharedPtr< class SLevelViewport > >* GetViewports() const;
 
-	/** Ticks the contained layout object */
-	void Tick();
-
 	/**
 	 * Sets the current layout by changing the contained layout object
 	 * 
