Adding a viewport controls help popup using the web browser to display html.

Also fixed non-unity build errors

diff --git a/Engine/Source/Editor/LevelEditor/LevelEditor.Build.cs b/Engine/Source/Editor/LevelEditor/LevelEditor.Build.cs
--- a/Engine/Source/Editor/LevelEditor/LevelEditor.Build.cs
+++ b/Engine/Source/Editor/LevelEditor/LevelEditor.Build.cs
@@ -13,7 +13,8 @@ public class LevelEditor : ModuleRules
 				"MainFrame",
                 "PlacementMode",
 				"ReferenceViewer",
-                "IntroTutorials"
+                "IntroTutorials",
+                "AppFramework"
 			}
 		);
 
@@ -52,6 +53,7 @@ public class LevelEditor : ModuleRules
 				"TargetPlatform",
 				"EngineSettings",
 				"PropertyEditor",
+                "AppFramework",
 			}
 		);
 
diff --git a/Engine/Source/Editor/LevelEditor/Private/SLevelViewport.cpp b/Engine/Source/Editor/LevelEditor/Private/SLevelViewport.cpp
--- a/Engine/Source/Editor/LevelEditor/Private/SLevelViewport.cpp
+++ b/Engine/Source/Editor/LevelEditor/Private/SLevelViewport.cpp
@@ -35,6 +35,7 @@
 #include "SDPIScaler.h"
 #include "SNotificationList.h"
 #include "NotificationManager.h"
+#include "SLevelViewportControlsPopup.h"
 
 static const FName LevelEditorName("LevelEditor");
 
@@ -194,6 +195,14 @@ void SLevelViewport::ConstructViewportOverlayContent()
 		SAssignNew( ActorPreviewHorizontalBox, SHorizontalBox )
 	];
 
+	ViewportOverlay->AddSlot(SlotIndex)
+	.VAlign(VAlign_Bottom)
+	.HAlign(HAlign_Left)
+	.Padding(5.0f)
+	[
+		SNew(SLevelViewportControlsPopup)
+	];
+
 	ViewportOverlay->AddSlot( SlotIndex )
 	.VAlign( VAlign_Bottom )
 	.HAlign( HAlign_Right )
diff --git a/Engine/Source/Editor/LevelEditor/Private/SLevelViewportControlsPopup.cpp b/Engine/Source/Editor/LevelEditor/Private/SLevelViewportControlsPopup.cpp
new file mode 100644
--- /dev/null
+++ b/Engine/Source/Editor/LevelEditor/Private/SLevelViewportControlsPopup.cpp
@@ -0,0 +1,94 @@
+// Copyright 1998-2014 Epic Games, Inc. All Rights Reserved.
+
+
+#include "LevelEditor.h"
+#include "SLevelViewportControlsPopup.h"
+#include "SWebBrowser.h"
+
+void SLevelViewportControlsPopup::Construct(const FArguments& InArgs)
+{
+	FVector2D PopupSize;
+	FString PopupPath;
+
+#if PLATFORM_WINDOWS
+	PopupSize.Set(252, 558);
+	PopupPath = FPaths::ConvertRelativePathToFull(FPaths::Combine(*FPaths::EngineDir(), TEXT("Documentation/Source/Shared/Editor/Overlays/ViewportControlsWindows.html")));
+#endif
+
+	if (!PopupPath.IsEmpty())
+	{
+		// The size of the web browser is a hack to ensure that the browser shrinks,
+		// otherwise it seems to show a scroll bar if the size never changes.
+		ControlsWidget =
+			SNew(SBox)
+			.WidthOverride(PopupSize.X)
+			.HeightOverride(PopupSize.Y)
+			[
+				SNew(SWebBrowser)
+				.InitialURL(PopupPath)
+				.ShowControls(false)
+				.SupportsTransparency(true)
+				.ViewportSize(PopupSize+1)
+			];
+
+		TAttribute<FText> ToolTipText = NSLOCTEXT("LevelViewportControlsPopup", "ViewportControlsToolTip", "Click to show Viewport Controls");
+		Default = FEditorStyle::GetBrush("HelpIcon");
+		Hovered = FEditorStyle::GetBrush("HelpIcon.Hovered");
+		Pressed = FEditorStyle::GetBrush("HelpIcon.Pressed");
+
+		ChildSlot
+		[
+			SAssignNew(MenuAnchor, SMenuAnchor)
+			.Placement(MenuPlacement_AboveAnchor)
+			[
+				SAssignNew(Button, SButton)
+				.ContentPadding(5)
+				.ButtonStyle(FEditorStyle::Get(), "HelpButton")
+				.OnClicked(this, &SLevelViewportControlsPopup::OnClicked)
+				.HAlign(HAlign_Center)
+				.VAlign(VAlign_Center)
+				.ToolTipText(ToolTipText)
+				[
+					SAssignNew(ButtonImage, SImage)
+					.Image(this, &SLevelViewportControlsPopup::GetButtonImage)
+				]
+			]
+			.OnGetMenuContent(this, &SLevelViewportControlsPopup::OnGetMenuContent)
+		];
+	}
+}
+
+const FSlateBrush* SLevelViewportControlsPopup::GetButtonImage() const
+{
+	if (Button->IsPressed())
+	{
+		return Pressed;
+	}
+
+	if (ButtonImage->IsHovered())
+	{
+		return Hovered;
+	}
+
+	return Default;
+}
+
+FReply SLevelViewportControlsPopup::OnClicked() const
+{
+	// If the menu button is clicked toggle the state of the menu anchor which will open or close the menu
+	if (MenuAnchor->ShouldOpenDueToClick())
+	{
+		MenuAnchor->SetIsOpen(true);
+	}
+	else
+	{
+		MenuAnchor->SetIsOpen(false);
+	}
+
+	return FReply::Handled();
+}
+
+TSharedRef<SWidget> SLevelViewportControlsPopup::OnGetMenuContent()
+{
+	return ControlsWidget.ToSharedRef();
+}
diff --git a/Engine/Source/Editor/LevelEditor/Private/SLevelViewportControlsPopup.h b/Engine/Source/Editor/LevelEditor/Private/SLevelViewportControlsPopup.h
new file mode 100644
--- /dev/null
+++ b/Engine/Source/Editor/LevelEditor/Private/SLevelViewportControlsPopup.h
@@ -0,0 +1,51 @@
+// Copyright 1998-2014 Epic Games, Inc. All Rights Reserved.
+
+
+#pragma once
+
+/**
+* Widget that opens a controls help popup when clicked
+*/
+class SLevelViewportControlsPopup : public SCompoundWidget
+{
+public:
+	SLATE_BEGIN_ARGS(SLevelViewportControlsPopup)
+	{}
+	SLATE_END_ARGS()
+
+	/**
+	* Constructs the menu
+	*/
+	void Construct(const FArguments& Declaration);
+private:
+	/**
+	 * Gets the current image for the button depending on state
+	 */
+	const FSlateBrush* GetButtonImage() const;
+
+	/**
+	 * Responds to mouse clicks
+	 */
+	FReply OnClicked() const;
+
+	/**
+	 * Gets content for the popup
+	 */
+	TSharedRef<SWidget> OnGetMenuContent();
+
+private:
+	/** Our menu's anchor */
+	TSharedPtr<SMenuAnchor> MenuAnchor;
+	/** The Button to open the menu*/
+	TSharedPtr< SButton > Button;
+	/** The Image to display on the button */
+	TSharedPtr< SImage > ButtonImage;
+	/** Widget to show the controls in the popup */
+	TSharedPtr< SWidget > ControlsWidget;
+	/** Default Button Image */
+	const FSlateBrush* Default;
+	/** Button Image while hovered */
+	const FSlateBrush* Hovered;
+	/** Button Image when pressed */
+	const FSlateBrush* Pressed;
+};
diff --git a/Engine/Source/Runtime/AppFramework/Private/Widgets/SWebBrowser.cpp b/Engine/Source/Runtime/AppFramework/Private/Widgets/SWebBrowser.cpp
--- a/Engine/Source/Runtime/AppFramework/Private/Widgets/SWebBrowser.cpp
+++ b/Engine/Source/Runtime/AppFramework/Private/Widgets/SWebBrowser.cpp
@@ -2,7 +2,7 @@
 
 #include "AppFrameworkPrivatePCH.h"
 #include "SWebBrowser.h"
-
+#include "SThrobber.h"
 #if WITH_EDITOR || IS_PROGRAM
 #include "WebBrowserModule.h"
 #include "IWebBrowserSingleton.h"
diff --git a/Engine/Source/Runtime/GameplayAbilities/Public/GameplayCueNotify_HitImpact.h b/Engine/Source/Runtime/GameplayAbilities/Public/GameplayCueNotify_HitImpact.h
--- a/Engine/Source/Runtime/GameplayAbilities/Public/GameplayCueNotify_HitImpact.h
+++ b/Engine/Source/Runtime/GameplayAbilities/Public/GameplayCueNotify_HitImpact.h
@@ -4,6 +4,7 @@
 
 #include "GameplayTags.h"
 #include "GameplayEffect.h"
+#include "GameplayCueNotify_Static.h"
 #include "GameplayCueNotify_HitImpact.generated.h"
 
 /**
