Feature level preview renamed to Render Level Preview (in UX.)
Render Level Preview removed from experimental menu.

diff --git a/Engine/Source/Editor/LevelEditor/Private/LevelEditorActions.cpp b/Engine/Source/Editor/LevelEditor/Private/LevelEditorActions.cpp
--- a/Engine/Source/Editor/LevelEditor/Private/LevelEditorActions.cpp
+++ b/Engine/Source/Editor/LevelEditor/Private/LevelEditorActions.cpp
@@ -2716,6 +2716,22 @@ void FLevelEditorCommands::RegisterCommands()
 	UI_COMMAND( MaterialQualityLevel_Low, "Low", "Sets material quality in the scene to low.", EUserInterfaceActionType::RadioButton, FInputGesture() );
 	UI_COMMAND( MaterialQualityLevel_High, "High", "Sets material quality in the scene to high.", EUserInterfaceActionType::RadioButton, FInputGesture() );
 
+	static const FText FeatureLevelLabels[ERHIFeatureLevel::Num] = 
+	{
+		NSLOCTEXT("LevelEditorCommands", "FeatureLevelPreviewType_ES2", "Mobile / HTML5"),
+		NSLOCTEXT("LevelEditorCommands", "FeatureLevelPreviewType_ES3", "High-End Mobile"),
+		NSLOCTEXT("LevelEditorCommands", "FeatureLevelPreviewType_SM4", "Shader Model 4"),
+		NSLOCTEXT("LevelEditorCommands", "FeatureLevelPreviewType_SM5", "Shader Model 5"),
+	};
+
+	static const FText FeatureLevelToolTips[ERHIFeatureLevel::Num] = 
+	{
+		NSLOCTEXT("LevelEditorCommands", "FeatureLevelPreviewTooltip_ES2", "OpenGLES 2"),
+		NSLOCTEXT("LevelEditorCommands", "FeatureLevelPreviewTooltip_ES3", "OpenGLES 3"),
+		NSLOCTEXT("LevelEditorCommands", "FeatureLevelPreviewTooltip_SM4", "DirectX 10, OpenGL 3.3+"),
+		NSLOCTEXT("LevelEditorCommands", "FeatureLevelPreviewTooltip_SM5", "DirectX 11, OpenGL 4.3+, PS4, XB1"),
+	};
+
 	for (int32 i = 0; i < ERHIFeatureLevel::Num; ++i)
 	{
 		FName Name;
@@ -2725,8 +2741,8 @@ void FLevelEditorCommands::RegisterCommands()
 			FUICommandInfoDecl(
 			this->AsShared(),
 			Name,
-			FText::FromName(Name),
-			FText::Format(NSLOCTEXT("LevelEditorCommands", "FeatureLevelPreviewTooltip", "Preview scenes using feature level {0}"), FText::FromName(Name)))
+			FeatureLevelLabels[i],
+			FeatureLevelToolTips[i])
 			.UserInterfaceType(EUserInterfaceActionType::RadioButton)
 			.DefaultGesture(FInputGesture());
 	}
diff --git a/Engine/Source/Editor/LevelEditor/Private/LevelEditorToolBar.cpp b/Engine/Source/Editor/LevelEditor/Private/LevelEditorToolBar.cpp
--- a/Engine/Source/Editor/LevelEditor/Private/LevelEditorToolBar.cpp
+++ b/Engine/Source/Editor/LevelEditor/Private/LevelEditorToolBar.cpp
@@ -1512,7 +1512,7 @@ static void MakeMaterialQualityLevelMenu( FMenuBuilder& MenuBuilder )
 
 static void MakeShaderModelPreviewMenu(FMenuBuilder& MenuBuilder)
 {
-	MenuBuilder.BeginSection("LevelEditorShaderModelPreview", NSLOCTEXT("LevelToolBarViewMenu", "FeatureLevelPreviewHeading", "Feature Level Preview"));
+	MenuBuilder.BeginSection("LevelEditorShaderModelPreview", NSLOCTEXT("LevelToolBarViewMenu", "FeatureLevelPreviewHeading", "Preview Rendering Level"));
 	{
 		for (int32 i = GMaxRHIFeatureLevel; i >= 0; --i)
 		{
@@ -1613,13 +1613,10 @@ TSharedRef< SWidget > FLevelEditorToolBar::GenerateQuickSettingsMenu( TSharedRef
 			LOCTEXT( "MaterialQualityLevelSubMenu_ToolTip", "Sets the value of the CVar \"r.MaterialQualityLevel\" (low=0, high=1). This affects materials via the QualitySwitch material expression." ),
 			FNewMenuDelegate::CreateStatic( &MakeMaterialQualityLevelMenu ) );
 
-		if (GetDefault<UEditorExperimentalSettings>()->bFeatureLevelPreview)
-		{
-			MenuBuilder.AddSubMenu(
-				LOCTEXT("FeatureLevelPreviewSubMenu", "Feature Level Preview"),
-				LOCTEXT("FeatureLevelPreviewSubMenu_ToolTip", "Sets the feature level preview mode"),
-				FNewMenuDelegate::CreateStatic(&MakeShaderModelPreviewMenu));
-		}
+		MenuBuilder.AddSubMenu(
+			LOCTEXT("FeatureLevelPreviewSubMenu", "Preview Rendering Level"),
+			LOCTEXT("FeatureLevelPreviewSubMenu_ToolTip", "Sets the rendering level used by the main editor"),
+			FNewMenuDelegate::CreateStatic(&MakeShaderModelPreviewMenu));
 	}
 	MenuBuilder.EndSection();
 
diff --git a/Engine/Source/Editor/UnrealEd/Classes/Settings/EditorExperimentalSettings.h b/Engine/Source/Editor/UnrealEd/Classes/Settings/EditorExperimentalSettings.h
--- a/Engine/Source/Editor/UnrealEd/Classes/Settings/EditorExperimentalSettings.h
+++ b/Engine/Source/Editor/UnrealEd/Classes/Settings/EditorExperimentalSettings.h
@@ -65,10 +65,6 @@ public:
 	UPROPERTY(EditAnywhere,config,Category=UserInterface,meta=(DisplayName="Allow ChunkID Assignments"))
 	bool bContextMenuChunkAssignments;
 
-	/** Enables the dynamic feature level switching functionality */
-	UPROPERTY(EditAnywhere, config, Category = Rendering, meta = (DisplayName = "Feature Level Preview"))
-	bool bFeatureLevelPreview;
-
 	/** Disable cook in the editor */
 	UPROPERTY(EditAnywhere, config, Category = Cooking, meta = (DisplayName = "Disable Cook In The Editor feature (cooks from launch on will be run in a separate process if disabled)"))
 	bool bDisableCookInEditor;
