Made sure Viewport Controls popup supports transparency

Changed the SMenuAnchor to open the menu in the same window, as there was no way to specify whether the new window supports transparency and it doesn't need to extend out of the current window.

diff --git a/Engine/Source/Editor/LevelEditor/Private/SLevelViewportControlsPopup.cpp b/Engine/Source/Editor/LevelEditor/Private/SLevelViewportControlsPopup.cpp
--- a/Engine/Source/Editor/LevelEditor/Private/SLevelViewportControlsPopup.cpp
+++ b/Engine/Source/Editor/LevelEditor/Private/SLevelViewportControlsPopup.cpp
@@ -39,7 +39,9 @@ void SLevelViewportControlsPopup::Construct(const FArguments& InArgs)
 		ChildSlot
 		[
 			SAssignNew(MenuAnchor, SMenuAnchor)
+			.Method(SMenuAnchor::UseCurrentWindow)
 			.Placement(MenuPlacement_AboveAnchor)
+			.OnGetMenuContent(this, &SLevelViewportControlsPopup::OnGetMenuContent)
 			[
 				SAssignNew(Button, SButton)
 				.ContentPadding(5)
@@ -53,7 +55,6 @@ void SLevelViewportControlsPopup::Construct(const FArguments& InArgs)
 					.Image(this, &SLevelViewportControlsPopup::GetButtonImage)
 				]
 			]
-			.OnGetMenuContent(this, &SLevelViewportControlsPopup::OnGetMenuContent)
 		];
 	}
 }
