Paper2D: Move SmartSnapping module to load post-engine-init, to prevent it from loading the level editor module at pre-init
- Added check() to LevelEditor to make sure that GUnrealEd is set prior to use

diff --git a/Engine/Plugins/2D/Paper2D/Paper2D.uplugin b/Engine/Plugins/2D/Paper2D/Paper2D.uplugin
--- a/Engine/Plugins/2D/Paper2D/Paper2D.uplugin
+++ b/Engine/Plugins/2D/Paper2D/Paper2D.uplugin
@@ -32,7 +32,8 @@
 		},
 		{
 			"Name" : "SmartSnapping",
-			"Type" : "Editor"
+			"Type" : "Editor",
+			"LoadingPhase" : "PostEngineInit"
 		}
 	],
 
diff --git a/Engine/Source/Editor/LevelEditor/Private/LevelViewportActions.cpp b/Engine/Source/Editor/LevelEditor/Private/LevelViewportActions.cpp
--- a/Engine/Source/Editor/LevelEditor/Private/LevelViewportActions.cpp
+++ b/Engine/Source/Editor/LevelEditor/Private/LevelViewportActions.cpp
@@ -23,7 +23,8 @@ FLevelViewportCommands::~FLevelViewportCommands()
 PRAGMA_DISABLE_OPTIMIZATION
 void FLevelViewportCommands::RegisterCommands()
 {
-	
+	check(GUnrealEd);
+
 	UI_COMMAND( ToggleMaximize, "Maximize Viewport", "Toggles the Maximize state of the current viewport", EUserInterfaceActionType::ToggleButton, FInputGesture() );
 	UI_COMMAND( ToggleGameView, "Game View", "Toggles game view.  Game view shows the scene as it appears in game", EUserInterfaceActionType::ToggleButton, FInputGesture( EKeys::G ) );
 	UI_COMMAND( ToggleImmersive, "Immersive Mode", "Switches this viewport between immersive mode and regular mode", EUserInterfaceActionType::ToggleButton, FInputGesture( EKeys::F11 ) );
