Updates  to use Epic launcher terminology.

OPP-1300 As a user I want the installer and launcher to be named the epic launcher

diff --git a/Engine/Source/Editor/LevelEditor/Private/LevelEditorActions.cpp b/Engine/Source/Editor/LevelEditor/Private/LevelEditorActions.cpp
--- a/Engine/Source/Editor/LevelEditor/Private/LevelEditorActions.cpp
+++ b/Engine/Source/Editor/LevelEditor/Private/LevelEditorActions.cpp
@@ -1691,7 +1691,7 @@ void FLevelEditorActionCallbacks::OpenMarketplace()
 		{
 			EventAttributes.Add(FAnalyticsEventAttribute(TEXT("OpenSucceeded"), TEXT("FALSE")));
 
-			if (EAppReturnType::Yes == FMessageDialog::Open(EAppMsgType::YesNo, LOCTEXT("InstallMarketplacePrompt", "The Marketplace requires the Unreal Engine Launcher, which does not seem to be installed on your computer. Would you like to install it now?")))
+			if (EAppReturnType::Yes == FMessageDialog::Open(EAppMsgType::YesNo, LOCTEXT("InstallMarketplacePrompt", "The Marketplace requires the Epic Launcher, which does not seem to be installed on your computer. Would you like to install it now?")))
 			{
 				if (!DesktopPlatform->OpenLauncher(true, TEXT("-OpenMarket")))
 				{
diff --git a/Engine/Source/Editor/UnrealEd/Private/Editor.cpp b/Engine/Source/Editor/UnrealEd/Private/Editor.cpp
--- a/Engine/Source/Editor/UnrealEd/Private/Editor.cpp
+++ b/Engine/Source/Editor/UnrealEd/Private/Editor.cpp
@@ -385,7 +385,8 @@ void UEditorEngine::InitEditor(IEngineLoop* InEngineLoop)
 		!GIsDemoMode &&
 		!IsRunningCommandlet() &&
 		!FPlatformProcess::IsApplicationRunning(TEXT("UnrealEngineLauncher")) &&
-		!FPlatformProcess::IsApplicationRunning(TEXT("Unreal Engine Launcher")) ) )
+		!FPlatformProcess::IsApplicationRunning(TEXT("Unreal Engine Launcher")) &&
+		!FPlatformProcess::IsApplicationRunning(TEXT("Epic Launcher")) ) )
 	{
 		IDesktopPlatform* DesktopPlatform = FDesktopPlatformModule::Get();
 		if ( DesktopPlatform != NULL )
