Made sure controls popup toggles correctly

Changed help button to fire on mouse down, so that the popup isn't triggered a frame after it's dismissed when you lift the mouse up.

diff --git a/Engine/Source/Editor/LevelEditor/Private/SLevelViewportControlsPopup.cpp b/Engine/Source/Editor/LevelEditor/Private/SLevelViewportControlsPopup.cpp
--- a/Engine/Source/Editor/LevelEditor/Private/SLevelViewportControlsPopup.cpp
+++ b/Engine/Source/Editor/LevelEditor/Private/SLevelViewportControlsPopup.cpp
@@ -47,6 +47,7 @@ void SLevelViewportControlsPopup::Construct(const FArguments& InArgs)
 				.ContentPadding(5)
 				.ButtonStyle(FEditorStyle::Get(), "HelpButton")
 				.OnClicked(this, &SLevelViewportControlsPopup::OnClicked)
+				.ClickMethod(EButtonClickMethod::MouseDown)
 				.HAlign(HAlign_Center)
 				.VAlign(VAlign_Center)
 				.ToolTipText(ToolTipText)
