Editor: Remove check() on level editor module load ordering, and converted calls into GUnrealEd into static methods

diff --git a/Engine/Source/Editor/LevelEditor/Private/LevelViewportActions.cpp b/Engine/Source/Editor/LevelEditor/Private/LevelViewportActions.cpp
--- a/Engine/Source/Editor/LevelEditor/Private/LevelViewportActions.cpp
+++ b/Engine/Source/Editor/LevelEditor/Private/LevelViewportActions.cpp
@@ -23,8 +23,6 @@ FLevelViewportCommands::~FLevelViewportCommands()
 PRAGMA_DISABLE_OPTIMIZATION
 void FLevelViewportCommands::RegisterCommands()
 {
-	check(GUnrealEd);
-
 	UI_COMMAND( ToggleMaximize, "Maximize Viewport", "Toggles the Maximize state of the current viewport", EUserInterfaceActionType::ToggleButton, FInputGesture() );
 	UI_COMMAND( ToggleGameView, "Game View", "Toggles game view.  Game view shows the scene as it appears in game", EUserInterfaceActionType::ToggleButton, FInputGesture( EKeys::G ) );
 	UI_COMMAND( ToggleImmersive, "Immersive Mode", "Switches this viewport between immersive mode and regular mode", EUserInterfaceActionType::ToggleButton, FInputGesture( EKeys::F11 ) );
@@ -138,7 +136,7 @@ void FLevelViewportCommands::RegisterCommands()
 		UI_COMMAND( HideAllVolumes, "Hide All Volumes", "Hides all volumes", EUserInterfaceActionType::Button, FInputGesture() );
 
 		TArray< UClass* > VolumeClasses;
-		GUnrealEd->GetSortedVolumeClasses( &VolumeClasses );
+		UUnrealEdEngine::GetSortedVolumeClasses(&VolumeClasses);
 
 		for( int32 VolumeClassIndex = 0; VolumeClassIndex < VolumeClasses.Num(); ++VolumeClassIndex )
 		{
@@ -175,7 +173,7 @@ void FLevelViewportCommands::RegisterCommands()
 		// get all the known layers
 		// Get a fresh list as GUnrealEd->SortedSpriteInfo may not yet be built.
 		TArray<FSpriteCategoryInfo> SortedSpriteInfo;
-		GUnrealEd->MakeSortedSpriteInfo(SortedSpriteInfo);
+		UUnrealEdEngine::MakeSortedSpriteInfo(SortedSpriteInfo);
 
 		FString SpritePrefix = TEXT("ShowSprite_");
 		for( int32 InfoIndex = 0; InfoIndex < SortedSpriteInfo.Num(); ++InfoIndex )
diff --git a/Engine/Source/Editor/LevelEditor/Private/SLevelViewport.cpp b/Engine/Source/Editor/LevelEditor/Private/SLevelViewport.cpp
--- a/Engine/Source/Editor/LevelEditor/Private/SLevelViewport.cpp
+++ b/Engine/Source/Editor/LevelEditor/Private/SLevelViewport.cpp
@@ -1306,7 +1306,7 @@ void SLevelViewport::BindShowCommands( FUICommandList& CommandList )
 
 		// Get all known volume classes
 		TArray< UClass* > VolumeClasses;
-		GUnrealEd->GetSortedVolumeClasses( &VolumeClasses );
+		UUnrealEdEngine::GetSortedVolumeClasses(&VolumeClasses);
 
 		for( int32 VolumeClassIndex = 0; VolumeClassIndex < VolumeClasses.Num(); ++VolumeClassIndex )
 		{
@@ -1716,7 +1716,7 @@ void SLevelViewport::OnToggleAllVolumeActors( bool bVisible )
 void SLevelViewport::ToggleShowVolumeClass( int32 VolumeID )
 {
 	TArray< UClass* > VolumeClasses;
-	GUnrealEd->GetSortedVolumeClasses( &VolumeClasses );
+	UUnrealEdEngine::GetSortedVolumeClasses(&VolumeClasses);
 
 	// Get the corresponding volume class for the clicked menu item.
 	UClass *SelectedVolumeClass = VolumeClasses[ VolumeID ];
diff --git a/Engine/Source/Editor/UnrealEd/Classes/Editor/UnrealEdEngine.h b/Engine/Source/Editor/UnrealEd/Classes/Editor/UnrealEdEngine.h
--- a/Engine/Source/Editor/UnrealEd/Classes/Editor/UnrealEdEngine.h
+++ b/Engine/Source/Editor/UnrealEd/Classes/Editor/UnrealEdEngine.h
@@ -209,7 +209,7 @@ class UNREALED_API UUnrealEdEngine : public UEditorEngine, public FNotifyHook
 	// End UEngine Interface.
 
 	/** Builds a list of sprite categories for use in menus */
-	void MakeSortedSpriteInfo(TArray<struct FSpriteCategoryInfo>& OutSortedSpriteInfo) const;
+	static void MakeSortedSpriteInfo(TArray<struct FSpriteCategoryInfo>& OutSortedSpriteInfo);
 
 	/** called when a package has has its dirty state updated */
 	void OnPackageDirtyStateUpdated( UPackage* Pkg);
@@ -693,7 +693,7 @@ class UNREALED_API UUnrealEdEngine : public UEditorEngine, public FNotifyHook
 	*
 	* @param	VolumeClasses		Array to populate with AVolume classes.
 	*/
-	void GetSortedVolumeClasses( TArray< UClass* >* VolumeClasses );
+	static void GetSortedVolumeClasses( TArray< UClass* >* VolumeClasses );
 
 	/**
 	 * Checks the destination level visibility and warns the user if he is trying to paste to a hidden level, offering the option to cancel the operation or unhide the level that is hidden
diff --git a/Engine/Source/Editor/UnrealEd/Private/LevelEditorViewport.cpp b/Engine/Source/Editor/UnrealEd/Private/LevelEditorViewport.cpp
--- a/Engine/Source/Editor/UnrealEd/Private/LevelEditorViewport.cpp
+++ b/Engine/Source/Editor/UnrealEd/Private/LevelEditorViewport.cpp
@@ -1422,7 +1422,7 @@ FLevelEditorViewportClient::FLevelEditorViewportClient()
 
 	// Get the number of volume classes so we can initialize our bit array
 	TArray<UClass*> VolumeClasses;
-	GUnrealEd->GetSortedVolumeClasses( &VolumeClasses ); 
+	UUnrealEdEngine::GetSortedVolumeClasses(&VolumeClasses);
 	VolumeActorVisibility.Init(true, VolumeClasses.Num() );
 
 	// Initialize all sprite categories to visible
@@ -2741,7 +2741,7 @@ static uint32 GetVolumeActorVisibilityId( const AActor& InActor )
 	{
 		// Build a mapping of volume classes to ID's.  Do this only once
 		TArray< UClass *> VolumeClasses;
-		GUnrealEd->GetSortedVolumeClasses( &VolumeClasses );
+		UUnrealEdEngine::GetSortedVolumeClasses(&VolumeClasses);
 		for( int32 VolumeIdx = 0; VolumeIdx < VolumeClasses.Num(); ++VolumeIdx )
 		{
 			// An actors flag is just the index of the actor in the stored volume array shifted left to represent a unique bit.
diff --git a/Engine/Source/Editor/UnrealEd/Private/UnrealEdEngine.cpp b/Engine/Source/Editor/UnrealEd/Private/UnrealEdEngine.cpp
--- a/Engine/Source/Editor/UnrealEd/Private/UnrealEdEngine.cpp
+++ b/Engine/Source/Editor/UnrealEd/Private/UnrealEdEngine.cpp
@@ -86,7 +86,7 @@ void UUnrealEdEngine::Init(IEngineLoop* InEngineLoop)
 
 	// Populate the data structures related to the sprite category visibility feature for use elsewhere in the editor later
 	TArray<FSpriteCategoryInfo> SortedSpriteInfo;
-	MakeSortedSpriteInfo(SortedSpriteInfo);
+	UUnrealEdEngine::MakeSortedSpriteInfo(SortedSpriteInfo);
 
 	// Iterate over the sorted list, constructing a mapping of unlocalized categories to the index the localized category
 	// resides in. This is an optimization to prevent having to localize values repeatedly.
@@ -157,7 +157,7 @@ void UUnrealEdEngine::CancelCookByTheBookInEditor()
 	CookServer->QueueCancelCookByTheBook();
 }
 
-void UUnrealEdEngine::MakeSortedSpriteInfo(TArray<FSpriteCategoryInfo>& OutSortedSpriteInfo) const
+void UUnrealEdEngine::MakeSortedSpriteInfo(TArray<FSpriteCategoryInfo>& OutSortedSpriteInfo)
 {
 	struct Local
 	{
