Fix crash on project open.

diff --git a/Engine/Source/Editor/LevelEditor/Private/SLevelViewport.cpp b/Engine/Source/Editor/LevelEditor/Private/SLevelViewport.cpp
--- a/Engine/Source/Editor/LevelEditor/Private/SLevelViewport.cpp
+++ b/Engine/Source/Editor/LevelEditor/Private/SLevelViewport.cpp
@@ -824,6 +824,8 @@ void SLevelViewport::Tick( const FGeometry& AllottedGeometry, const double InCur
 	// viewport with Slate widgets that are part of the game, don't throttle.
 	if ( bPIEContainsFocus != bContainsFocus )
 	{
+		// We can arrive at this point before creating throttling manager (which registers the cvar), so create it explicitly.
+		static const FSlateThrottleManager & ThrottleManager = FSlateThrottleManager::Get();
 		static IConsoleVariable* AllowThrottling = IConsoleManager::Get().FindConsoleVariable(TEXT("Slate.bAllowThrottling"));
 		check(AllowThrottling);
 
