Epic Friday :
- record animation from phat
- keep simulated position for skeletalmesh component

diff --git a/Engine/Source/Editor/LevelEditor/LevelEditor.Build.cs b/Engine/Source/Editor/LevelEditor/LevelEditor.Build.cs
--- a/Engine/Source/Editor/LevelEditor/LevelEditor.Build.cs
+++ b/Engine/Source/Editor/LevelEditor/LevelEditor.Build.cs
@@ -54,6 +54,7 @@ public class LevelEditor : ModuleRules
 				"EngineSettings",
 				"PropertyEditor",
 				"WebBrowser",
+                "Persona"
 			}
 		);
 
diff --git a/Engine/Source/Editor/LevelEditor/Private/LevelEditorActions.cpp b/Engine/Source/Editor/LevelEditor/Private/LevelEditorActions.cpp
--- a/Engine/Source/Editor/LevelEditor/Private/LevelEditorActions.cpp
+++ b/Engine/Source/Editor/LevelEditor/Private/LevelEditorActions.cpp
@@ -58,6 +58,8 @@
 #include "Components/LightComponent.h"
 #include "Engine/StaticMesh.h"
 #include "Engine/Light.h"
+#include "Animation/SkeletalMeshActor.h"
+#include "Editor/Persona/Public/AnimationRecorder.h"
 
 DEFINE_LOG_CATEGORY_STATIC(LevelEditorActions, Log, All);
 
@@ -1544,6 +1546,75 @@ bool FLevelEditorActionCallbacks::ScaleGridSnap_IsChecked()
 	return GetDefault<ULevelEditorViewportSettings>()->SnapScaleEnabled;
 }
 
+bool FLevelEditorActionCallbacks::SaveAnimationFromSkeletalMeshComponent(AActor * EditorActor, AActor * SimActor, TArray<class USkeletalMeshComponent*> & OutEditorComponents)
+{
+	// currently blueprint actors don't work because their property can't get copied over. 
+	if (Cast<UBlueprintGeneratedClass>(EditorActor->GetClass()) != nullptr)
+	{
+		return false;
+	}
+
+	// find all skel components
+	TArray<class USkeletalMeshComponent *> SimSkelComponents;
+	SimActor->GetComponents<USkeletalMeshComponent>(SimSkelComponents);
+
+	if(SimSkelComponents.Num() > 0)
+	{
+		// see if simulating, 
+		bool bSimulating = false;
+		for (auto & Comp : SimSkelComponents)
+		{
+			bSimulating |= (Comp->SkeletalMesh && Comp->SkeletalMesh->Skeleton && Comp->IsSimulatingPhysics());
+		}
+
+		// if any of them are legitimately simulating
+		if (bSimulating)
+		{
+			// ask users if you'd like to make an animation
+			FFormatNamedArguments Args;
+			Args.Add(TEXT("ActorName"), FText::FromString(GetNameSafe(EditorActor)));
+			FText AskQuestion = FText::Format(LOCTEXT("KeepSimulationChanges_AskSaveAnimation", "Would you like to save animations from simulation for {ActorName} actor"), Args);
+			if(EAppReturnType::Yes == FMessageDialog::Open(EAppMsgType::YesNo, AskQuestion))
+			{
+				for (auto & Comp : SimSkelComponents)
+				{
+					if (Comp->SkeletalMesh && Comp->SkeletalMesh->Skeleton && Comp->IsSimulatingPhysics())
+					{
+						// now record to animation
+						FAnimationRecorder Recorder;
+						if(Recorder.TriggerRecordAnimation(Comp))
+						{
+							class UAnimSequence * Sequence = Recorder.GetAnimationObject();
+							if(Sequence)
+							{
+								Recorder.StopRecord(false);
+								Comp->SetAnimationMode(EAnimationMode::AnimationSingleNode);
+								Comp->AnimationData.AnimToPlay = Sequence;
+								Comp->SetAnimation(Sequence);
+								Comp->SetSimulatePhysics(false);
+
+								// add the matching component to EditorCompoennts
+								class USkeletalMeshComponent * MatchingComponent = Cast<USkeletalMeshComponent>(EditorUtilities::FindMatchingComponentInstance(Comp, EditorActor));
+								if (MatchingComponent)
+								{
+									OutEditorComponents.Add(MatchingComponent);
+								}
+								else
+								{
+									UE_LOG( LevelEditorActions, Warning, TEXT("Matching component could not be found %s(%s)"), *GetNameSafe(Comp), *GetNameSafe(EditorActor) );
+								}
+							}
+						}
+					}
+				}
+
+				return true;
+			}
+		}
+	}
+
+	return false;
+}
 
 void FLevelEditorActionCallbacks::OnKeepSimulationChanges()
 {
@@ -1560,6 +1631,8 @@ void FLevelEditorActionCallbacks::OnKeepSimulationChanges()
 		{
 			const FScopedTransaction Transaction( NSLOCTEXT( "LevelEditorCommands", "KeepSimulationChanges", "Keep Simulation Changes" ) );
 
+			TArray<class USkeletalMeshComponent*> ComponentsToReinitialize;
+
 			for( auto ActorIt( GEditor->GetSelectedActorIterator() ); ActorIt; ++ActorIt )
 			{
 				auto* SimWorldActor = CastChecked<AActor>( *ActorIt );
@@ -1568,6 +1641,8 @@ void FLevelEditorActionCallbacks::OnKeepSimulationChanges()
 				AActor* EditorWorldActor = EditorUtilities::GetEditorWorldCounterpartActor( SimWorldActor );
 				if( EditorWorldActor != NULL )
 				{
+					SaveAnimationFromSkeletalMeshComponent(EditorWorldActor, SimWorldActor, ComponentsToReinitialize);
+
 					// We only want to copy CPF_Edit properties back, or properties that are set through editor manipulation
 					// NOTE: This needs to match what we're doing in the BuildSelectedActorInfo() function
 					const auto CopyOptions = ( EditorUtilities::ECopyOptions::Type )(
@@ -1587,6 +1662,15 @@ void FLevelEditorActionCallbacks::OnKeepSimulationChanges()
 						}
 					}
 				}
+
+				// need to reinitialize animation
+				for (auto MeshComp : ComponentsToReinitialize)
+				{
+					if(MeshComp->SkeletalMesh)
+					{
+						MeshComp->InitAnim(true);
+					}
+				}
 			}
 		}
 
diff --git a/Engine/Source/Editor/LevelEditor/Public/LevelEditorActions.h b/Engine/Source/Editor/LevelEditor/Public/LevelEditorActions.h
--- a/Engine/Source/Editor/LevelEditor/Public/LevelEditorActions.h
+++ b/Engine/Source/Editor/LevelEditor/Public/LevelEditorActions.h
@@ -1163,5 +1163,13 @@ private:
 	 * @param InDestination		The destination actor we want to move this actor to, NULL assumes we just want to go towards the floor
 	 */
 	static void SnapActorTo_Clicked( const bool InAlign, const bool InUseLineTrace, const bool InUseBounds, const bool InUsePivot, const AActor* InDestination = NULL );
+
+	/** 
+	 * Create and apply animation to the SkeletalMeshComponent if Simulating
+	 * 
+	 * @param EditorActor	Editor Counterpart Actor
+	 * @param SimActor		Simulating Actor in PIE or SIE
+	 */
+	static bool SaveAnimationFromSkeletalMeshComponent(AActor * EditorActor, AActor * SimActor, TArray<class USkeletalMeshComponent*> & OutEditorComponents);
 };
 
diff --git a/Engine/Source/Editor/Persona/Private/Shared/AnimationRecorder.cpp b/Engine/Source/Editor/Persona/Private/Shared/AnimationRecorder.cpp
--- a/Engine/Source/Editor/Persona/Private/Shared/AnimationRecorder.cpp
+++ b/Engine/Source/Editor/Persona/Private/Shared/AnimationRecorder.cpp
@@ -17,6 +17,7 @@ FAnimationRecorder::FAnimationRecorder()
 	: IntervalTime(1.0f/DEFAULT_SAMPLERATE)
 	, MaxFrame(DEFAULT_SAMPLERATE * 60 * 10) // for now, set the max limit to 10 mins. 
 	, AnimationObject(NULL)
+	, bRecordLocalToWorld(false)
 {
 
 }
@@ -121,7 +122,7 @@ void FAnimationRecorder::StartRecord(USkeletalMeshComponent * Component, UAnimSe
 	Record(Component, PreviousSpacesBases, 0);
 }
 
-void FAnimationRecorder::StopRecord(bool bShowMessage)
+UAnimSequence * FAnimationRecorder::StopRecord(bool bShowMessage)
 {
 	if (AnimationObject)
 	{
@@ -133,6 +134,8 @@ void FAnimationRecorder::StopRecord(bool bShowMessage)
 		AnimationObject->PostProcessSequence();
 		AnimationObject->MarkPackageDirty();
 
+		UAnimSequence * ReturnObject = AnimationObject;
+
 		// notify to user
 		if (bShowMessage)
 		{
@@ -154,12 +157,18 @@ void FAnimationRecorder::StopRecord(bool bShowMessage)
 				Notification->SetCompletionState( SNotificationItem::CS_Success );
 			}
 
+			FAssetRegistryModule::AssetCreated(AnimationObject);
+			
 			FMessageDialog::Open(EAppMsgType::Ok, NotificationText);
 		}
 
 		AnimationObject = NULL;
 		PreviousSpacesBases.Empty();
+
+		return ReturnObject;
 	}
+
+	return NULL;
 }
 
 // return false if nothing to update
@@ -253,6 +262,11 @@ void FAnimationRecorder::Record( USkeletalMeshComponent * Component, TArray<FTra
 				{
 					LocalTransform.SetToRelativeTransform(SpacesBases[ParentIndex]);
 				}
+				// if record local to world, we'd like to consider component to world to be in root
+				else if (bRecordLocalToWorld)
+				{
+					LocalTransform *= Component->ComponentToWorld;
+				}
 
 				FRawAnimSequenceTrack& RawTrack = AnimationObject->RawAnimationData[TrackIndex];
 				RawTrack.PosKeys.Add(LocalTransform.GetTranslation());
diff --git a/Engine/Source/Editor/Persona/Public/AnimationRecorder.h b/Engine/Source/Editor/Persona/Public/AnimationRecorder.h
--- a/Engine/Source/Editor/Persona/Public/AnimationRecorder.h
+++ b/Engine/Source/Editor/Persona/Public/AnimationRecorder.h
@@ -23,12 +23,15 @@ public:
 	// also provides interface
 	bool TriggerRecordAnimation(USkeletalMeshComponent * Component);
 	void StartRecord(USkeletalMeshComponent * Component, UAnimSequence * InAnimationObject);
-	void StopRecord(bool bShowMessage);
+	UAnimSequence * StopRecord(bool bShowMessage);
 	void UpdateRecord(USkeletalMeshComponent * Component, float DeltaTime);
-	const UAnimSequence* GetAnimationObject() const { return AnimationObject; }
+	UAnimSequence* GetAnimationObject() const { return AnimationObject; }
 	bool InRecording() const { return AnimationObject != NULL; }
 	float GetTimeRecorded() const { return TimePassed; }
 
+	// If true, it will record root to include LocalToWorld
+	bool bRecordLocalToWorld;
+
 private:
 	void Record( USkeletalMeshComponent * Component, TArray<FTransform> SpacesBases, int32 FrameToAdd );
 };
diff --git a/Engine/Source/Editor/PhAT/PhAT.Build.cs b/Engine/Source/Editor/PhAT/PhAT.Build.cs
--- a/Engine/Source/Editor/PhAT/PhAT.Build.cs
+++ b/Engine/Source/Editor/PhAT/PhAT.Build.cs
@@ -21,6 +21,7 @@ public class PhAT : ModuleRules
 				"LevelEditor",
 				"UnrealEd",
                 "Kismet",
+                "Persona"
 			}
 		);
 
diff --git a/Engine/Source/Editor/PhAT/Private/PhAT.cpp b/Engine/Source/Editor/PhAT/Private/PhAT.cpp
--- a/Engine/Source/Editor/PhAT/Private/PhAT.cpp
+++ b/Engine/Source/Editor/PhAT/Private/PhAT.cpp
@@ -851,6 +851,17 @@ void FPhAT::ExtendToolbar()
 				ToolbarBuilder.AddWidget(PhATAnimation);
 			}
 			ToolbarBuilder.EndSection();
+
+			ToolbarBuilder.BeginSection("PhATRecord");
+			{
+				ToolbarBuilder.AddToolBarButton(Commands.RecordAnimation,
+											NAME_None,
+											 TAttribute<FText>(Phat, &FPhAT::GetRecordStatusLabel),
+											 TAttribute<FText>(Phat, &FPhAT::GetRecordStatusTooltip),
+											 TAttribute<FSlateIcon>(Phat, &FPhAT::GetRecordStatusImage),
+											 NAME_None);
+			}
+			ToolbarBuilder.EndSection();
 		}
 	};
 
@@ -1362,6 +1373,15 @@ void FPhAT::BindCommands()
 		FExecuteAction::CreateSP(this, &FPhAT::Mirror),
 		FCanExecuteAction::CreateSP(this, &FPhAT::IsNotSimulation)
 		);
+
+	// record animation
+	ToolkitCommands->MapAction(
+		Commands.RecordAnimation,
+		FExecuteAction::CreateSP(this, &FPhAT::RecordAnimation),
+		FCanExecuteAction::CreateSP(this, &FPhAT::IsRecordAvailable),
+		FIsActionChecked(),
+		FIsActionButtonVisible()
+		);
 }
 
 void FPhAT::Mirror()
@@ -3193,4 +3213,89 @@ void FPhAT::OnAddPhatRecord(const FString& Action, bool bRecordSimulate, bool bR
 		FEngineAnalytics::GetProvider().RecordEvent(EventString, Attribs);
 	}
 }
+
+// only during simulation
+bool FPhAT::IsRecordAvailable() const
+{
+	// make sure mesh exists
+	return (SharedData->EditorSkelComp && SharedData->EditorSkelComp->SkeletalMesh/* && SharedData->bRunningSimulation */);
+}
+
+FSlateIcon FPhAT::GetRecordStatusImage() const
+{
+	if(SharedData->Recorder.InRecording())
+	{
+		return FSlateIcon(FEditorStyle::GetStyleSetName(), "Persona.StopRecordAnimation");
+	}
+
+	return FSlateIcon(FEditorStyle::GetStyleSetName(), "Persona.StartRecordAnimation");
+}
+
+FText FPhAT::GetRecordMenuLabel() const
+{
+	if(SharedData->Recorder.InRecording())
+	{
+		return LOCTEXT("Persona_StopRecordAnimationMenuLabel", "Stop Record Animation");
+	}
+
+	return LOCTEXT("Persona_StartRecordAnimationLabel", "Start Record Animation");
+}
+
+FText FPhAT::GetRecordStatusLabel() const
+{
+	if(SharedData->Recorder.InRecording())
+	{
+		return LOCTEXT("Persona_StopRecordAnimationLabel", "Stop");
+	}
+
+	return LOCTEXT("Persona_StartRecordAnimationLabel", "Record");
+}
+
+FText FPhAT::GetRecordStatusTooltip() const
+{
+	if(SharedData->Recorder.InRecording())
+	{
+		return LOCTEXT("Persona_StopRecordAnimation", "Stop Record Animation");
+	}
+
+	return LOCTEXT("Persona_StartRecordAnimation", "Start Record Animation");
+}
+
+void FPhAT::RecordAnimation()
+{
+	if(!SharedData->EditorSkelComp || !SharedData->EditorSkelComp->SkeletalMesh)
+	{
+		// error
+		return;
+	}
+
+	if(SharedData->Recorder.InRecording())
+	{
+		UAnimSequence * AnimSeq = SharedData->Recorder.StopRecord(true);
+
+		// open the asset?
+		if (AnimSeq)
+		{
+			TArray<UObject*> ObjectsToSync;
+			ObjectsToSync.Add(AnimSeq);
+
+			FContentBrowserModule& ContentBrowserModule = FModuleManager::Get().LoadModuleChecked<FContentBrowserModule>("ContentBrowser");
+			ContentBrowserModule.Get().SyncBrowserToAssets(ObjectsToSync, true);
+		}
+	}
+	else
+	{
+		SharedData->Recorder.TriggerRecordAnimation(SharedData->EditorSkelComp);
+	}
+}
+
+void FPhAT::Tick(float DeltaTime)
+{
+}
+
+TStatId FPhAT::GetStatId() const
+{
+	RETURN_QUICK_DECLARE_CYCLE_STAT(FPhAT, STATGROUP_Tickables);
+}
+
 #undef LOCTEXT_NAMESPACE
diff --git a/Engine/Source/Editor/PhAT/Private/PhAT.h b/Engine/Source/Editor/PhAT/Private/PhAT.h
--- a/Engine/Source/Editor/PhAT/Private/PhAT.h
+++ b/Engine/Source/Editor/PhAT/Private/PhAT.h
@@ -30,7 +30,7 @@ enum EPhATSimulationMode
    FPhAT
 -----------------------------------------------------------------------------*/
 
-class FPhAT : public IPhAT, public FGCObject, public FEditorUndoClient
+class FPhAT : public IPhAT, public FGCObject, public FEditorUndoClient, public FTickableEditorObject
 {
 public:
 	virtual void RegisterTabSpawners(const TSharedRef<class FTabManager>& TabManager) override;
@@ -91,6 +91,12 @@ public:
 	/** Returns whether a PIE session is running. */
 	static bool IsPIERunning();
 
+	// FTickableEditorObject interface
+	virtual void Tick(float DeltaTime) override;
+	virtual bool IsTickable() const override { return true; }
+	virtual TStatId GetStatId() const override;
+	// End of FTickableEditorObject interface
+
 private:
 
 	enum EPhatHierarchyFilterMode
@@ -297,4 +303,12 @@ private:
 
 	/** Records PhAT related data - simulating or mode change */
 	void OnAddPhatRecord(const FString& Action, bool bRecordSimulate, bool bRecordMode);
+
+private:
+	void RecordAnimation();
+	bool IsRecordAvailable() const;
+	FSlateIcon GetRecordStatusImage() const;
+	FText GetRecordStatusTooltip() const;
+	FText GetRecordStatusLabel() const;
+	FText GetRecordMenuLabel() const;
 };
diff --git a/Engine/Source/Editor/PhAT/Private/PhATActions.cpp b/Engine/Source/Editor/PhAT/Private/PhATActions.cpp
--- a/Engine/Source/Editor/PhAT/Private/PhATActions.cpp
+++ b/Engine/Source/Editor/PhAT/Private/PhATActions.cpp
@@ -56,6 +56,7 @@ void FPhATCommands::RegisterCommands()
 	UI_COMMAND(ConvertToSkeletal, "To Skeletal", "Convert Selected Constraint To Skeletal", EUserInterfaceActionType::Button, FInputGesture());
 	UI_COMMAND(DeleteConstraint, "Delete", "Delete Selected Constraint", EUserInterfaceActionType::Button, FInputGesture());
 	UI_COMMAND(PlayAnimation, "Play", "Play Animation", EUserInterfaceActionType::ToggleButton, FInputGesture());
+	UI_COMMAND(RecordAnimation, "Record", "Record Animation", EUserInterfaceActionType::ToggleButton, FInputGesture());
 	UI_COMMAND(ShowSkeleton, "Skeleton", "Show Skeleton", EUserInterfaceActionType::ToggleButton, FInputGesture());
 	UI_COMMAND(MakeBodyKinematic, "Kinematic", "Make Body Kinematic", EUserInterfaceActionType::RadioButton, FInputGesture());
 	UI_COMMAND(MakeBodySimulated, "Simulated", "Make Body Simulated", EUserInterfaceActionType::RadioButton, FInputGesture());
diff --git a/Engine/Source/Editor/PhAT/Private/PhATActions.h b/Engine/Source/Editor/PhAT/Private/PhATActions.h
--- a/Engine/Source/Editor/PhAT/Private/PhATActions.h
+++ b/Engine/Source/Editor/PhAT/Private/PhATActions.h
@@ -68,6 +68,7 @@ public:
 	TSharedPtr<FUICommandInfo> ConvertToSkeletal;
 	TSharedPtr<FUICommandInfo> DeleteConstraint;
 	TSharedPtr<FUICommandInfo> PlayAnimation;
+	TSharedPtr<FUICommandInfo> RecordAnimation;
 	TSharedPtr<FUICommandInfo> ShowSkeleton;
 	TSharedPtr<FUICommandInfo> MakeBodyKinematic;
 	TSharedPtr<FUICommandInfo> MakeBodySimulated;
diff --git a/Engine/Source/Editor/PhAT/Private/PhATPreviewViewportClient.cpp b/Engine/Source/Editor/PhAT/Private/PhATPreviewViewportClient.cpp
--- a/Engine/Source/Editor/PhAT/Private/PhATPreviewViewportClient.cpp
+++ b/Engine/Source/Editor/PhAT/Private/PhATPreviewViewportClient.cpp
@@ -71,7 +71,7 @@ FPhATEdPreviewViewportClient::FPhATEdPreviewViewportClient(TWeakPtr<FPhAT> InPhA
 	if (CollBoxExtent.X > CollBoxExtent.Y)
 	{
 		SetViewLocation( FVector(WorldSphere.Center.X, WorldSphere.Center.Y - 1.5*WorldSphere.W, WorldSphere.Center.Z) );
-		SetViewRotation( FRotator(0,90.f,0) );	
+		SetViewRotation( EditorViewportDefs::DefaultPerspectiveViewRotation );	
 	}
 	else
 	{
@@ -164,7 +164,7 @@ void FPhATEdPreviewViewportClient::DrawCanvas( FViewport& InViewport, FSceneView
 	if ((SharedData->bShowHierarchy && SharedData->EditorSimOptions->bShowNamesInHierarchy))
 	{
 		// Iterate over each graphics bone.
-		for (int32 i = 0; i <SharedData->EditorSkelComp->GetNumSpaceBases(); ++i)
+		for(int32 i = 0; i <SharedData->EditorSkelComp->GetNumSpaceBases(); ++i)
 		{
 			FVector BonePos = SharedData->EditorSkelComp->ComponentToWorld.TransformPosition(SharedData->EditorSkelComp->GetSpaceBases()[i].GetLocation());
 
@@ -668,6 +668,12 @@ void FPhATEdPreviewViewportClient::Tick(float DeltaSeconds)
 		// We back up the transforms array now
 		SharedData->EditorSkelComp->AnimationSpaceBases = SharedData->EditorSkelComp->GetSpaceBases();
 		SharedData->EditorSkelComp->SetPhysicsBlendWeight(SharedData->EditorSimOptions->PhysicsBlend);
+
+		if(SharedData->Recorder.InRecording())
+		{
+			// make sure you don't allow switch SharedData->EditorSkelComp
+			SharedData->Recorder.UpdateRecord(SharedData->EditorSkelComp, DeltaSeconds);
+		}
 	}
 }
 
diff --git a/Engine/Source/Editor/PhAT/Private/PhATSharedData.cpp b/Engine/Source/Editor/PhAT/Private/PhATSharedData.cpp
--- a/Engine/Source/Editor/PhAT/Private/PhATSharedData.cpp
+++ b/Engine/Source/Editor/PhAT/Private/PhATSharedData.cpp
@@ -67,6 +67,9 @@ FPhATSharedData::FPhATSharedData()
 	EditorSimOptions->HandleAngularDamping = MouseHandle->AngularDamping;
 	EditorSimOptions->HandleAngularStiffness = MouseHandle->AngularStiffness;
 	EditorSimOptions->InterpolationSpeed = MouseHandle->InterpolationSpeed;
+
+	// when record simulation, we'd like to record local to world change
+	Recorder.bRecordLocalToWorld = true;
 }
 
 FPhATSharedData::~FPhATSharedData()
diff --git a/Engine/Source/Editor/PhAT/Private/PhATSharedData.h b/Engine/Source/Editor/PhAT/Private/PhATSharedData.h
--- a/Engine/Source/Editor/PhAT/Private/PhATSharedData.h
+++ b/Engine/Source/Editor/PhAT/Private/PhATSharedData.h
@@ -4,6 +4,7 @@
 
 #include "PreviewScene.h"
 #include "PhysicsPublic.h"
+#include "AnimationRecorder.h"
 
 class UPhATEdSkeletalMeshComponent;
 class UPhysicsHandleComponent;
@@ -215,6 +216,8 @@ public:
 	UBodySetup * CopiedBodySetup;
 	UPhysicsConstraintTemplate * CopiedConstraintTemplate;
 
+	struct FAnimationRecorder Recorder;
+
 	/** Constraint editing */
 	TArray<FSelection> SelectedConstraints;
 	FSelection * GetSelectedConstraint()
