Fixed compile error

diff --git a/Engine/Source/Editor/LevelEditor/Private/LevelEditorActions.cpp b/Engine/Source/Editor/LevelEditor/Private/LevelEditorActions.cpp
--- a/Engine/Source/Editor/LevelEditor/Private/LevelEditorActions.cpp
+++ b/Engine/Source/Editor/LevelEditor/Private/LevelEditorActions.cpp
@@ -105,6 +105,11 @@ namespace LevelEditorActionsHelpers
 	}
 }
 
+namespace EditorUtilities
+{
+	extern UActorComponent* FindMatchingComponentInstance( UActorComponent* SourceComponent, AActor* TargetActor );
+}
+
 bool FLevelEditorActionCallbacks::DefaultCanExecuteAction()
 {
 	return FSlateApplication::Get().IsNormalExecution();
