Back out revision 313 from //depot/UE4/Engine/Source/Editor/LevelEditor/Private/LevelEditorActions.cpp

diff --git a/Engine/Source/Editor/LevelEditor/Private/LevelEditorActions.cpp b/Engine/Source/Editor/LevelEditor/Private/LevelEditorActions.cpp
--- a/Engine/Source/Editor/LevelEditor/Private/LevelEditorActions.cpp
+++ b/Engine/Source/Editor/LevelEditor/Private/LevelEditorActions.cpp
@@ -105,11 +105,6 @@ namespace LevelEditorActionsHelpers
 	}
 }
 
-namespace EditorUtilities
-{
-	extern UActorComponent* FindMatchingComponentInstance( UActorComponent* SourceComponent, AActor* TargetActor );
-}
-
 bool FLevelEditorActionCallbacks::DefaultCanExecuteAction()
 {
 	return FSlateApplication::Get().IsNormalExecution();
