Viewport controls popup is now only shown in active viewport window.

Also made sure that it doesn't show up when PIE is active.

Popup content loads at the point where it is clicked instead of being pre-loaded.

diff --git a/Engine/Source/Editor/LevelEditor/Private/SLevelViewport.cpp b/Engine/Source/Editor/LevelEditor/Private/SLevelViewport.cpp
--- a/Engine/Source/Editor/LevelEditor/Private/SLevelViewport.cpp
+++ b/Engine/Source/Editor/LevelEditor/Private/SLevelViewport.cpp
@@ -205,6 +205,7 @@ void SLevelViewport::ConstructViewportOverlayContent()
 	.Padding(5.0f)
 	[
 		SNew(SLevelViewportControlsPopup)
+		.Visibility(this, &SLevelViewport::GetViewportControlsVisibility)
 	];
 
 	ViewportOverlay->AddSlot( SlotIndex )
@@ -3186,6 +3187,13 @@ EVisibility SLevelViewport::GetCurrentFeatureLevelPreviewTextVisibility() const
 	}
 }
 
+EVisibility SLevelViewport::GetViewportControlsVisibility() const
+{
+	// Do not show the controls if this viewport has a play in editor session
+	// or is not the current viewport
+	return (&GetLevelViewportClient() == GCurrentLevelEditingViewportClient && !IsPlayInEditorViewportActive()) ? OnGetViewportContentVisibility() : EVisibility::Collapsed;
+}
+
 void SLevelViewport::OnSetViewportConfiguration(FName ConfigurationName)
 {
 	TSharedPtr<FLevelViewportLayout> LayoutPinned = ParentLayout.Pin();
diff --git a/Engine/Source/Editor/LevelEditor/Private/SLevelViewportControlsPopup.cpp b/Engine/Source/Editor/LevelEditor/Private/SLevelViewportControlsPopup.cpp
--- a/Engine/Source/Editor/LevelEditor/Private/SLevelViewportControlsPopup.cpp
+++ b/Engine/Source/Editor/LevelEditor/Private/SLevelViewportControlsPopup.cpp
@@ -7,9 +7,6 @@
 
 void SLevelViewportControlsPopup::Construct(const FArguments& InArgs)
 {
-	FVector2D PopupSize;
-	FString PopupPath;
-
 #if PLATFORM_WINDOWS
 	PopupSize.Set(252, 558);
 	PopupPath = FPaths::ConvertRelativePathToFull(FPaths::Combine(*FPaths::EngineDir(), TEXT("Documentation/Source/Shared/Editor/Overlays/ViewportControlsWindows.html")));
@@ -17,20 +14,6 @@ void SLevelViewportControlsPopup::Construct(const FArguments& InArgs)
 
 	if (!PopupPath.IsEmpty())
 	{
-		// The size of the web browser is a hack to ensure that the browser shrinks,
-		// otherwise it seems to show a scroll bar if the size never changes.
-		ControlsWidget =
-			SNew(SBox)
-			.WidthOverride(PopupSize.X)
-			.HeightOverride(PopupSize.Y)
-			[
-				SNew(SWebBrowser)
-				.InitialURL(PopupPath)
-				.ShowControls(false)
-				.SupportsTransparency(true)
-				.ViewportSize(PopupSize+1)
-			];
-
 		TAttribute<FText> ToolTipText = NSLOCTEXT("LevelViewportControlsPopup", "ViewportControlsToolTip", "Click to show Viewport Controls");
 		Default = FEditorStyle::GetBrush("HelpIcon");
 		Hovered = FEditorStyle::GetBrush("HelpIcon.Hovered");
@@ -92,5 +75,21 @@ FReply SLevelViewportControlsPopup::OnClicked() const
 
 TSharedRef<SWidget> SLevelViewportControlsPopup::OnGetMenuContent()
 {
-	return ControlsWidget.ToSharedRef();
+	if (!PopupPath.IsEmpty())
+	{
+		// The size of the web browser is a hack to ensure that the browser shrinks,
+		// otherwise it seems to show a scroll bar if the size never changes.
+		return SNew(SBox)
+			.WidthOverride(PopupSize.X)
+			.HeightOverride(PopupSize.Y)
+			[
+				SNew(SWebBrowser)
+				.InitialURL(PopupPath)
+				.ShowControls(false)
+				.SupportsTransparency(true)
+				.ViewportSize(PopupSize + 1)
+			];
+	}
+
+	return SNullWidget::NullWidget;
 }
diff --git a/Engine/Source/Editor/LevelEditor/Private/SLevelViewportControlsPopup.h b/Engine/Source/Editor/LevelEditor/Private/SLevelViewportControlsPopup.h
--- a/Engine/Source/Editor/LevelEditor/Private/SLevelViewportControlsPopup.h
+++ b/Engine/Source/Editor/LevelEditor/Private/SLevelViewportControlsPopup.h
@@ -40,12 +40,14 @@ private:
 	TSharedPtr< SButton > Button;
 	/** The Image to display on the button */
 	TSharedPtr< SImage > ButtonImage;
-	/** Widget to show the controls in the popup */
-	TSharedPtr< SWidget > ControlsWidget;
 	/** Default Button Image */
 	const FSlateBrush* Default;
 	/** Button Image while hovered */
 	const FSlateBrush* Hovered;
 	/** Button Image when pressed */
 	const FSlateBrush* Pressed;
+	/** Size of the viewport controls popup */
+	FVector2D PopupSize;
+	/** Path to html file for viewport controls */
+	FString PopupPath;
 };
diff --git a/Engine/Source/Editor/LevelEditor/Public/SLevelViewport.h b/Engine/Source/Editor/LevelEditor/Public/SLevelViewport.h
--- a/Engine/Source/Editor/LevelEditor/Public/SLevelViewport.h
+++ b/Engine/Source/Editor/LevelEditor/Public/SLevelViewport.h
@@ -247,6 +247,9 @@ public:
 	/** @return The visibility of the current feature level preview text display */
 	EVisibility GetCurrentFeatureLevelPreviewTextVisibility() const;
 
+	/** @return The visibility of the viewport controls popup */
+	EVisibility GetViewportControlsVisibility() const;
+
 	/**
 	 * Sets the current layout on the parent tab that this viewport belongs to.
 	 * 
