Removed the Recompile Level Editor feature
- This didn't work reliably, cluttered the main tool bar and confused users

diff --git a/Engine/Source/Editor/LevelEditor/Private/LevelEditor.cpp b/Engine/Source/Editor/LevelEditor/Private/LevelEditor.cpp
--- a/Engine/Source/Editor/LevelEditor/Private/LevelEditor.cpp
+++ b/Engine/Source/Editor/LevelEditor/Private/LevelEditor.cpp
@@ -154,10 +154,6 @@ void FLevelEditorModule::StartupModule()
 
 	NotificationBarExtensibilityManager = MakeShareable(new FExtensibilityManager);
 
-	// Figure out if we recompile the level editor.
-	FString SourcePath = FPaths::Combine(*FPaths::EngineDir(), TEXT("Source/Editor/LevelEditor/Private"));
-	bCanBeRecompiled = IFileManager::Get().DirectoryExists(*SourcePath) && !GEngineVersion.IsPromotedBuild();
-
 	// Note this must come before any tab spawning because that can create the SLevelEditor and attempt to map commands
 	FLevelEditorCommands::Register();
 	FLevelEditorModesCommands::Register();
@@ -553,20 +549,6 @@ void FLevelEditorModule::BindGlobalLevelEditorCommands()
 	ActionList.MapAction( Commands.Build,
 		FExecuteAction::CreateStatic( &FLevelEditorActionCallbacks::Build_Execute ) );
 
-
-	if (CanBeRecompiled())
-	{
-		ActionList.MapAction( Commands.RecompileLevelEditor,
-			FExecuteAction::CreateStatic( &FLevelEditorActionCallbacks::RecompileLevelEditor_Clicked ),
-			FCanExecuteAction::CreateStatic( &FLevelEditorActionCallbacks::Recompile_CanExecute )
-			);
-
-		ActionList.MapAction( Commands.ReloadLevelEditor,
-			FExecuteAction::CreateStatic( &FLevelEditorActionCallbacks::ReloadLevelEditor_Clicked ),
-			FCanExecuteAction::CreateStatic( &FLevelEditorActionCallbacks::Reload_CanExecute )
-			);
-	}
-
 	ActionList.MapAction( Commands.RecompileGameCode,
 		FExecuteAction::CreateStatic( &FLevelEditorActionCallbacks::RecompileGameCode_Clicked ),
 		FCanExecuteAction::CreateStatic( &FLevelEditorActionCallbacks::Recompile_CanExecute )
diff --git a/Engine/Source/Editor/LevelEditor/Private/LevelEditorActions.cpp b/Engine/Source/Editor/LevelEditor/Private/LevelEditorActions.cpp
--- a/Engine/Source/Editor/LevelEditor/Private/LevelEditorActions.cpp
+++ b/Engine/Source/Editor/LevelEditor/Private/LevelEditorActions.cpp
@@ -900,16 +900,6 @@ bool FLevelEditorActionCallbacks::CanShowSourceCodeActions()
 	return HotReloadSupport.IsAnyGameModuleLoaded();
 }
 
-void FLevelEditorActionCallbacks::RecompileLevelEditor_Clicked()
-{
-	GEngine->DeferredCommands.Add( TEXT( "Module Recompile LevelEditor" ) );
-}
-
-void FLevelEditorActionCallbacks::ReloadLevelEditor_Clicked()
-{
-	GEngine->DeferredCommands.Add( TEXT( "Module Reload LevelEditor" ) );
-}
-
 void FLevelEditorActionCallbacks::RecompileGameCode_Clicked()
 {
 	// Don't allow a recompile while already compiling!
@@ -928,13 +918,6 @@ bool FLevelEditorActionCallbacks::Recompile_CanExecute()
 	return !HotReloadSupport.IsCurrentlyCompiling() && !(GEngineVersion.IsPromotedBuild() && FEngineBuildSettings::IsPerforceBuild());
 }
 
-bool FLevelEditorActionCallbacks::Reload_CanExecute()
-{
-	// We're not able to reload if a compile is already in progress!
-	IHotReloadInterface& HotReloadSupport = FModuleManager::LoadModuleChecked<IHotReloadInterface>(HotReloadModule);
-	return !HotReloadSupport.IsCurrentlyCompiling();
-}
-
 void FLevelEditorActionCallbacks::GoToCodeForActor_Clicked()
 {
 	const auto& SelectedActorInfo = AssetSelectionUtils::GetSelectedActorInfo();
@@ -2648,13 +2631,6 @@ void FLevelEditorCommands::RegisterCommands()
 	UI_COMMAND( TextureStats, "Open Texture Stats", "Opens the Texture Stats viewer", EUserInterfaceActionType::Button, FInputGesture() );
 	UI_COMMAND( MapCheck, "Open Map Check", "Checks map for errors", EUserInterfaceActionType::Button, FInputGesture() );
 
-	FLevelEditorModule& LevelEditorModule = FModuleManager::GetModuleChecked<FLevelEditorModule>("LevelEditor");
-	if (LevelEditorModule.CanBeRecompiled())
-	{
-		UI_COMMAND( RecompileLevelEditor, "Recompile Level Editor", "Recompiles and reloads C++ code for the Level Editor on the fly", EUserInterfaceActionType::Button, FInputGesture() );
-		UI_COMMAND( ReloadLevelEditor, "Reload Level Editor", "Reloads C++ code for the Level Editor on the fly", EUserInterfaceActionType::Button, FInputGesture() );
-	}
-
 	UI_COMMAND( RecompileGameCode, "Recompile Game Code", "Recompiles and reloads C++ code for game systems on the fly", EUserInterfaceActionType::Button, FInputGesture( EKeys::P, EModifierKey::Alt | EModifierKey::Control | EModifierKey::Shift ) );
 
 	UI_COMMAND( EditAsset, "Edit Asset", "Edits the asset associated with the selected actor", EUserInterfaceActionType::Button, FInputGesture( EKeys::E, EModifierKey::Control ) );
diff --git a/Engine/Source/Editor/LevelEditor/Private/LevelEditorToolBar.cpp b/Engine/Source/Editor/LevelEditor/Private/LevelEditorToolBar.cpp
--- a/Engine/Source/Editor/LevelEditor/Private/LevelEditorToolBar.cpp
+++ b/Engine/Source/Editor/LevelEditor/Private/LevelEditorToolBar.cpp
@@ -1168,18 +1168,6 @@ TSharedRef< SWidget > FLevelEditorToolBar::MakeLevelEditorToolBar( const TShared
 				FLevelEditorCommands::Get().RecompileGameCode->GetDescription(),
 				FSlateIcon(FEditorStyle::GetStyleSetName(), "LevelEditor.Recompile")
 				);
-
-			ToolbarBuilder.AddComboButton(
-				FUIAction(
-					FExecuteAction(),
-					FCanExecuteAction(),
-					FIsActionChecked(),
-					FIsActionButtonVisible::CreateStatic(FLevelEditorActionCallbacks::CanShowSourceCodeActions)),
-				FOnGetContent::CreateStatic( &FLevelEditorToolBar::GenerateCompileMenuContent, InCommandList ),
-				LOCTEXT( "CompileCombo_Label", "Compile Options" ),
-				LOCTEXT( "CompileMenuCombo_ToolTip", "Recompile options" ),
-				FSlateIcon(FEditorStyle::GetStyleSetName(), "LevelEditor.Recompile"),
-				true);
 		}
 	}
 	ToolbarBuilder.EndSection();
@@ -1693,39 +1681,6 @@ TSharedRef< SWidget > FLevelEditorToolBar::GenerateQuickSettingsMenu( TSharedRef
 	return MenuBuilder.MakeWidget();
 }
 
-TSharedRef< SWidget > FLevelEditorToolBar::GenerateCompileMenuContent( TSharedRef<FUICommandList> InCommandList )
-{
-#define LOCTEXT_NAMESPACE "LevelToolBarViewMenu"
-
-	// Get all menu extenders for this context menu from the level editor module
-	FLevelEditorModule& LevelEditorModule = FModuleManager::GetModuleChecked<FLevelEditorModule>( TEXT("LevelEditor") );
-	TArray<FLevelEditorModule::FLevelEditorMenuExtender> MenuExtenderDelegates = LevelEditorModule.GetAllLevelEditorToolbarCompileMenuExtenders();
-
-	TArray<TSharedPtr<FExtender>> Extenders;
-	for (int32 i = 0; i < MenuExtenderDelegates.Num(); ++i)
-	{
-		if (MenuExtenderDelegates[i].IsBound())
-		{
-			Extenders.Add(MenuExtenderDelegates[i].Execute(InCommandList));
-		}
-	}
-	TSharedPtr<FExtender> MenuExtender = FExtender::Combine(Extenders);
-
-	const bool bShouldCloseWindowAfterMenuSelection = true;
-	FMenuBuilder MenuBuilder( bShouldCloseWindowAfterMenuSelection, InCommandList, MenuExtender );
-
-	if (LevelEditorModule.CanBeRecompiled())
-	{
-		MenuBuilder.AddMenuEntry( FLevelEditorCommands::Get().RecompileLevelEditor );
-		MenuBuilder.AddMenuEntry( FLevelEditorCommands::Get().ReloadLevelEditor );
-	}
-
-	MenuBuilder.AddMenuEntry( FLevelEditorCommands::Get().RecompileGameCode );
-
-#undef LOCTEXT_NAMESPACE
-
-	return MenuBuilder.MakeWidget();
-}
 
 TSharedRef< SWidget > FLevelEditorToolBar::GenerateOpenBlueprintMenuContent( TSharedRef<FUICommandList> InCommandList, TWeakPtr< SLevelEditor > InLevelEditor )
 {
diff --git a/Engine/Source/Editor/LevelEditor/Private/LevelEditorToolBar.h b/Engine/Source/Editor/LevelEditor/Private/LevelEditorToolBar.h
--- a/Engine/Source/Editor/LevelEditor/Private/LevelEditorToolBar.h
+++ b/Engine/Source/Editor/LevelEditor/Private/LevelEditorToolBar.h
@@ -49,13 +49,6 @@ protected:
 	 */
 	static TSharedRef< SWidget > GenerateQuickSettingsMenu( TSharedRef<FUICommandList> InCommandList );
 
-	/**
-	 * Generates menu content for the compile combo button drop down menu
-	 *
-	 * @return	Menu content widget
-	 */
-	static TSharedRef< SWidget > GenerateCompileMenuContent( TSharedRef<FUICommandList> InCommandList );
-
 	/**
 	 * Generates menu content for the compile combo button drop down menu
 	 *
diff --git a/Engine/Source/Editor/LevelEditor/Public/LevelEditor.h b/Engine/Source/Editor/LevelEditor/Public/LevelEditor.h
--- a/Engine/Source/Editor/LevelEditor/Public/LevelEditor.h
+++ b/Engine/Source/Editor/LevelEditor/Public/LevelEditor.h
@@ -223,9 +223,6 @@ public:
 	DECLARE_EVENT( FLevelEditorModule, FTakeHighResScreenShotsEvent );
 	virtual FTakeHighResScreenShotsEvent& OnTakeHighResScreenShots() { return TakeHighResScreenShotsEvent; }
 
-	/** Determines whether the level editor can be recompiled on the fly. */
-	inline bool CanBeRecompiled() const { return bCanBeRecompiled; }
-
 private:
 	/**
 	 * Binds all global level editor commands to delegates
@@ -289,9 +286,6 @@ private:
 
 	/* Holds the Editor's tab manager */
 	TSharedPtr<FTabManager> LevelEditorTabManager;
-
-	/* Whether we have source for the level editor, and it can be recompiled. */
-	bool bCanBeRecompiled;
 };
 
 #endif // __LevelEditor_h__
diff --git a/Engine/Source/Editor/LevelEditor/Public/LevelEditorActions.h b/Engine/Source/Editor/LevelEditor/Public/LevelEditorActions.h
--- a/Engine/Source/Editor/LevelEditor/Public/LevelEditorActions.h
+++ b/Engine/Source/Editor/LevelEditor/Public/LevelEditorActions.h
@@ -748,11 +748,8 @@ public:
 	/**
 	 * Called when the recompile buttons are clicked.
 	 */
-	static void RecompileLevelEditor_Clicked();
-	static void ReloadLevelEditor_Clicked();
 	static void RecompileGameCode_Clicked();
 	static bool Recompile_CanExecute();
-	static bool Reload_CanExecute();
 
 	/**
 	 * Called when the FindInContentBrowser command is executed
