Streamlined the Select submenu in the level viewport context menu.

diff --git a/Engine/Source/Editor/LevelEditor/Private/LevelEditor.cpp b/Engine/Source/Editor/LevelEditor/Private/LevelEditor.cpp
--- a/Engine/Source/Editor/LevelEditor/Private/LevelEditor.cpp
+++ b/Engine/Source/Editor/LevelEditor/Private/LevelEditor.cpp
@@ -1002,16 +1002,6 @@ void FLevelEditorModule::BindGlobalLevelEditorCommands()
 		FExecuteAction::CreateStatic( &FLevelEditorActionCallbacks::ExecuteExecCommand, FString( TEXT("MAP SELECT SUBTRACTS") ) )
 		);
 
-	ActionList.MapAction(
-		Commands.SelectAllSemiSolidBrushes,
-		FExecuteAction::CreateStatic( &FLevelEditorActionCallbacks::ExecuteExecCommand, FString( TEXT("MAP SELECT SEMISOLIDS") ) )
-		);
-
-	ActionList.MapAction(
-		Commands.SelectAllNonSolidBrushes,
-		FExecuteAction::CreateStatic( &FLevelEditorActionCallbacks::ExecuteExecCommand, FString( TEXT("MAP SELECT NONSOLIDS") ) )
-		);
-
 	ActionList.MapAction(
 		Commands.SelectAllSurfaces,
 		FExecuteAction::CreateStatic( &FLevelEditorActionCallbacks::ExecuteExecCommand, FString( TEXT("POLY SELECT ALL") ) )
diff --git a/Engine/Source/Editor/LevelEditor/Private/LevelEditorActions.cpp b/Engine/Source/Editor/LevelEditor/Private/LevelEditorActions.cpp
--- a/Engine/Source/Editor/LevelEditor/Private/LevelEditorActions.cpp
+++ b/Engine/Source/Editor/LevelEditor/Private/LevelEditorActions.cpp
@@ -2734,8 +2734,6 @@ void FLevelEditorCommands::RegisterCommands()
 	UI_COMMAND( SelectStationaryLightsExceedingOverlap, "Select Stationary Lights exceeding overlap", "Selects all stationary lights exceeding the overlap limit", EUserInterfaceActionType::Button, FInputGesture() ); 
 	UI_COMMAND( SelectAllAddditiveBrushes, "Select All Additive Brushes", "Selects all additive brushes", EUserInterfaceActionType::Button, FInputGesture() ); 
 	UI_COMMAND( SelectAllSubtractiveBrushes, "Select All Subtractive Brushes", "Selects all subtractive brushes", EUserInterfaceActionType::Button, FInputGesture() ); 
-	UI_COMMAND( SelectAllSemiSolidBrushes, "Select All Semi-Solid Brushes", "Selects all semi-solid brushes", EUserInterfaceActionType::Button, FInputGesture() ); 
-	UI_COMMAND( SelectAllNonSolidBrushes, "Select All Non-Solid Brushes", "Selects all non-solid brushes", EUserInterfaceActionType::Button, FInputGesture() ); 
 	UI_COMMAND( SelectAllActorsControlledByMatinee, "Select Actors Used by This Matinee", "Selects all actors controlled by this Matinee", EUserInterfaceActionType::Button, FInputGesture() ); 
 
 	UI_COMMAND( SelectAllSurfaces, "Select All Surfaces", "Selects all bsp surfaces", EUserInterfaceActionType::Button, FInputGesture(EModifierKey::Shift, EKeys::S) ); 
diff --git a/Engine/Source/Editor/LevelEditor/Private/LevelEditorContextMenu.cpp b/Engine/Source/Editor/LevelEditor/Private/LevelEditorContextMenu.cpp
--- a/Engine/Source/Editor/LevelEditor/Private/LevelEditorContextMenu.cpp
+++ b/Engine/Source/Editor/LevelEditor/Private/LevelEditorContextMenu.cpp
@@ -487,7 +487,7 @@ void FLevelEditorContextMenuImpl::FillSelectActorMenu( FMenuBuilder& MenuBuilder
 	}
 
 	// Add brush commands when we have a brush or any surfaces selected
-	MenuBuilder.BeginSection("SelectBrush", LOCTEXT("SelectBrushHeading", "Brushes") );
+	MenuBuilder.BeginSection("SelectBSP", LOCTEXT("SelectBSPHeading", "BSP") );
 	{
 		if( SelectionInfo.bHaveBrush || NumSelectedSurfaces > 0 )
 		{
@@ -499,8 +499,7 @@ void FLevelEditorContextMenuImpl::FillSelectActorMenu( FMenuBuilder& MenuBuilder
 
 		MenuBuilder.AddMenuEntry( FLevelEditorCommands::Get().SelectAllAddditiveBrushes );
 		MenuBuilder.AddMenuEntry( FLevelEditorCommands::Get().SelectAllSubtractiveBrushes );
-		MenuBuilder.AddMenuEntry( FLevelEditorCommands::Get().SelectAllSemiSolidBrushes );
-		MenuBuilder.AddMenuEntry( FLevelEditorCommands::Get().SelectAllNonSolidBrushes );
+		MenuBuilder.AddMenuEntry( FLevelEditorCommands::Get().SelectAllSurfaces );
 	}
 	MenuBuilder.EndSection();
 
@@ -560,13 +559,6 @@ void FLevelEditorContextMenuImpl::FillSelectActorMenu( FMenuBuilder& MenuBuilder
 		MenuBuilder.EndSection();
 	}
 
-	// Allow users to select all surfaces in the level in a single click
-	MenuBuilder.BeginSection("SelectSurfaces", LOCTEXT("SelectAllSurfaces", "Surfaces") );
-	{
-		MenuBuilder.AddMenuEntry( FLevelEditorCommands::Get().SelectAllSurfaces );
-	}
-	MenuBuilder.EndSection();
-
 	// build matinee related selection menu
 	FillMatineeSelectActorMenu( MenuBuilder );
 }
