Added a method to add a blank actor from the placement browser
Integrated MarcA's changes to display and edit per actor instance components.

Must set bInWorldBPEditing to true in experimental settings to enable.

diff --git a/Engine/Source/Editor/DetailCustomizations/Private/ActorDetails.cpp b/Engine/Source/Editor/DetailCustomizations/Private/ActorDetails.cpp
--- a/Engine/Source/Editor/DetailCustomizations/Private/ActorDetails.cpp
+++ b/Engine/Source/Editor/DetailCustomizations/Private/ActorDetails.cpp
@@ -22,7 +22,6 @@
 #include "ActorEditorUtils.h"
 #include "BlueprintEditorUtils.h"
 #include "ComponentTransformDetails.h"
-#include "ComponentsTree.h"
 #include "IPropertyUtilities.h"
 #include "IDocumentation.h"
 #include "Runtime/Engine/Classes/Engine/BrushShape.h"
@@ -128,13 +127,6 @@ void FActorDetails::CustomizeDetails( IDetailLayoutBuilder& DetailLayout )
 		}
 
 		OnExtendActorDetails.Broadcast(DetailLayout, FGetSelectedActors::CreateSP(this, &FActorDetails::GetSelectedActors));
-
-		/*if (!HideCategories.Contains(TEXT("Layers")))
-		{
-			AddLayersCategory(DetailLayout);
-		}*/
-
-		//AddComponentsCategory( DetailLayout );
 	}
 
 	TSharedPtr<IPropertyHandle> PrimaryTickProperty = DetailLayout.GetProperty(GET_MEMBER_NAME_CHECKED(AActor, PrimaryActorTick));
@@ -686,23 +678,6 @@ void FActorDetails::AddBlutilityCategory( IDetailLayoutBuilder& DetailBuilder, c
 	}
 }
 
-void FActorDetails::AddComponentsCategory( IDetailLayoutBuilder& DetailBuilder )
-{
-	AActor* SelectedActor = NULL;
-	const TArray< TWeakObjectPtr<AActor> >& Actors = GetSelectedActors();
-	if(Actors.Num() == 1 && Actors[0].IsValid())
-	{
-		SelectedActor = Actors[0].Get();
-	}
-
-	DetailBuilder.EditCategory( "Components", NSLOCTEXT("ActorDetails", "ComponentsSection", "Components"), ECategoryPriority::Uncommon )
-		.InitiallyCollapsed( true )
-		.AddCustomRow( NSLOCTEXT("ActorDetails", "ComponentsSection", "Components") )
-		[
-			SNew( SComponentsTree, SelectedActor )
-		];
-}
-
 const TArray< TWeakObjectPtr<AActor> >& FActorDetails::GetSelectedActors() const
 {
 	return SelectedActors;
diff --git a/Engine/Source/Editor/DetailCustomizations/Private/ActorDetails.h b/Engine/Source/Editor/DetailCustomizations/Private/ActorDetails.h
--- a/Engine/Source/Editor/DetailCustomizations/Private/ActorDetails.h
+++ b/Engine/Source/Editor/DetailCustomizations/Private/ActorDetails.h
@@ -67,8 +67,6 @@ private:
 
 	void AddBlueprintCategory( IDetailLayoutBuilder& DetailBuilder, const TMap<UBlueprint*, UObject*>& UniqueBlueprints );
 
-	void AddComponentsCategory( IDetailLayoutBuilder& DetailBuilder );
-
 	void AddLayersCategory( IDetailLayoutBuilder& DetailBuilder );
 
 	void AddTransformCategory( IDetailLayoutBuilder& DetailBuilder );
diff --git a/Engine/Source/Editor/DetailCustomizations/Private/ComponentsTree.cpp b/Engine/Source/Editor/DetailCustomizations/Private/ComponentsTree.cpp
deleted file mode 100644
diff --git a/Engine/Source/Editor/DetailCustomizations/Private/ComponentsTree.h b/Engine/Source/Editor/DetailCustomizations/Private/ComponentsTree.h
deleted file mode 100644
diff --git a/Engine/Source/Editor/LevelEditor/Private/ComponentsTree.cpp b/Engine/Source/Editor/LevelEditor/Private/ComponentsTree.cpp
new file mode 100644
--- /dev/null
+++ b/Engine/Source/Editor/LevelEditor/Private/ComponentsTree.cpp
@@ -0,0 +1,258 @@
+// Copyright 1998-2015 Epic Games, Inc. All Rights Reserved.
+
+#include "LevelEditor.h"
+#include "ComponentsTree.h"
+#include "ScopedTransaction.h"
+#include "Editor/UnrealEd/Public/SComponentClassCombo.h"
+#include "Editor/UnrealEd/Public/ClassIconFinder.h"
+#include "Editor/PropertyEditor/Public/PropertyEditorModule.h"
+#include "Editor/PropertyEditor/Public/IDetailsView.h"
+
+//////////////////////////////////////////////////////////////////////////
+
+FComponentTreeNode::FComponentTreeNode(UActorComponent* InComponent)
+{
+	Component = InComponent;
+
+	USceneComponent* SceneComp = Cast<USceneComponent>(InComponent);
+	if(SceneComp != NULL)
+	{
+		for(int32 ChildIdx=0; ChildIdx<SceneComp->AttachChildren.Num(); ChildIdx++)
+		{
+			USceneComponent* ChildComp = SceneComp->AttachChildren[ChildIdx];
+			if(ChildComp && !ChildComp->IsPendingKill())
+			{
+				FComponentTreeNodePtrType ChildNode = MakeShareable(new FComponentTreeNode(ChildComp));
+				ChildNodes.Add(ChildNode);
+			}
+		}
+	}
+}
+
+
+//////////////////////////////////////////////////////////////////////////
+
+void SComponentRowWidget::Construct( const FArguments& InArgs, FComponentTreeNodePtrType InNodePtr )
+{
+	NodePtr = InNodePtr;
+
+	FString ComponentName;
+	const FSlateBrush* ComponentIcon = FEditorStyle::GetBrush("SCS.NativeComponent");
+	FSlateColor ComponentColor = FLinearColor::White;
+
+	TSharedPtr<FComponentTreeNode> Node = NodePtr.Pin();
+	if(Node.IsValid() && Node->Component.IsValid())
+	{
+		ComponentName = Node->Component->GetName();
+		ComponentIcon = FClassIconFinder::FindIconForClass( Node->Component->GetClass(), TEXT("SCS.Component") );
+
+		// Color native and BP-made components in different colors
+		if(Node->Component->IsDefaultSubobject())
+		{
+			ComponentColor = FLinearColor(0.08f,0.15f,0.6f);
+		}
+		else if(Node->Component->bCreatedByConstructionScript)
+		{
+			ComponentColor = FLinearColor(0.08f,0.35f,0.6f);
+		}
+	}
+
+	this->ChildSlot
+	[
+		SNew(SHorizontalBox)
+		+SHorizontalBox::Slot()
+		.AutoWidth()
+		.VAlign(VAlign_Center)
+		[
+			SNew(SImage)
+			.Image(ComponentIcon)
+			.ColorAndOpacity(ComponentColor)
+		]
+		+SHorizontalBox::Slot()
+		.AutoWidth()
+		.VAlign(VAlign_Center)
+		.Padding(4,0,4,0)
+		[
+			SNew(STextBlock)
+			.Text(ComponentName)
+		]
+	];
+}
+
+//////////////////////////////////////////////////////////////////////////
+
+// Don't show any 'disable on instance' properties, these are instances of components we are seeing
+static bool IsPropertyVisible( const FPropertyAndParent& PropertyAndParent )
+{
+	if ( PropertyAndParent.Property.HasAllPropertyFlags(CPF_DisableEditOnInstance) )
+	{
+		return false;
+	}
+
+	return true;
+}
+
+void SComponentsTree::Construct(const FArguments& InArgs, TSharedPtr<class IDetailsView> InPropertyView )
+{
+	FPropertyEditorModule& EditModule = FModuleManager::Get().GetModuleChecked<FPropertyEditorModule>("PropertyEditor");
+	FDetailsViewArgs DetailsViewArgs( /*bUpdateFromSelection=*/ false, /*bLockable=*/ false, /*bAllowSearch=*/ true, /*bObjectsUseNameArea=*/ true, /*bHideSelectionTip=*/ true );
+	DetailsViewArgs.bHideActorNameArea = true;
+
+	PropertyView = InPropertyView;
+
+	this->ChildSlot
+	[
+		SNew(SBorder)
+		.Padding( 2 )
+		.BorderImage( FEditorStyle::GetBrush("ToolPanel.GroupBorder") )
+		.Content()
+		[
+			SNew(SVerticalBox)
+			+ SVerticalBox::Slot()
+			.AutoHeight()
+			.Padding(0.f, 2.f, 0.f, 0.f)
+			[
+				SNew(SComponentClassCombo)
+				.OnComponentClassSelected(this, &SComponentsTree::OnSelectedCompClass)
+			]
+			+ SVerticalBox::Slot()
+			.AutoHeight()
+			[
+				SAssignNew(TreeWidget, SComponentTreeType)
+				.TreeItemsSource( &RootNodes )
+				.OnGenerateRow( this, &SComponentsTree::MakeTableRowWidget )
+				.OnGetChildren( this, &SComponentsTree::OnGetChildrenForTree )
+				.OnSelectionChanged( this, &SComponentsTree::OnTreeSelectionChanged )
+				.ItemHeight( 24 )
+			]
+		]
+	];
+
+	UpdateTree();
+}
+
+void SComponentsTree::SetObjects(const TArray<UObject*>& InObjects)
+{
+	if (InObjects.Num() == 1)
+	{
+		Actor = Cast<AActor>(InObjects[0]);
+	}
+	else
+	{
+		Actor = NULL;
+	}
+	Visibility = (Actor.IsValid() ? EVisibility::Visible : EVisibility::Hidden);
+	UpdateTree();
+}
+
+TSharedRef<ITableRow> SComponentsTree::MakeTableRowWidget( FComponentTreeNodePtrType InNodePtr, const TSharedRef<STableViewBase>& OwnerTable )
+{
+	return 
+		SNew(STableRow<FComponentTreeNodePtrType>, OwnerTable)
+		[
+			SNew(SComponentRowWidget, InNodePtr)
+		];
+}
+
+void SComponentsTree::OnGetChildrenForTree( FComponentTreeNodePtrType InNodePtr, TArray<FComponentTreeNodePtrType>& OutChildren )
+{
+	if(InNodePtr.IsValid())
+	{
+		OutChildren = InNodePtr->ChildNodes;
+	}
+}
+
+void SComponentsTree::OnTreeSelectionChanged(FComponentTreeNodePtrType InSelectedNodePtr, ESelectInfo::Type SelectInfo)
+{
+	TArray<UObject*> Objects;
+
+	USelection* SelectedComponents = GEditor->GetSelectedComponents();
+	SelectedComponents->BeginBatchSelectOperation();
+	SelectedComponents->DeselectAll();
+
+	for (FComponentTreeNodePtrType SelectedNode : TreeWidget->GetSelectedItems())
+	{
+		if (SelectedNode.IsValid() && SelectedNode->Component.IsValid())
+		{
+			UActorComponent* Component = SelectedNode->Component.Get();
+			Objects.Add(Component);
+			SelectedComponents->Select(Component);
+		}
+	}
+	
+	SelectedComponents->EndBatchSelectOperation();
+	PropertyView->SetObjects(Objects, false);
+
+	GEditor->RedrawLevelEditingViewports();
+}
+
+void SComponentsTree::OnSelectedCompClass(TSubclassOf<UActorComponent> CompClass)
+{
+	if(Actor.IsValid())
+	{
+		const FScopedTransaction Transaction(NSLOCTEXT("Editor", "UndoAction_AddComponent", "Add Component"));
+
+		Actor->Modify();
+
+		// Create new component
+		UActorComponent* NewComp = ConstructObject<UActorComponent>(CompClass, Actor.Get(), NAME_None, RF_Transactional);
+		check(NewComp);
+
+		// Add to SerializedComponents array so it gets saved
+		Actor->SerializedComponents.Add(NewComp);
+
+		// If a scene component, and actor has a root component, attach it to root
+		USceneComponent* NewSceneComp = Cast<USceneComponent>(NewComp);
+		USceneComponent* RootComp = Actor->GetRootComponent();
+		if(NewSceneComp && RootComp)
+		{
+			NewSceneComp->AttachTo(RootComp);
+		}
+
+		// Register component
+		NewComp->RegisterComponent();
+
+		// Update tree to see result
+		UpdateTree();
+	}
+}
+
+
+void SComponentsTree::UpdateTree()
+{
+	check(TreeWidget.IsValid());
+
+	RootNodes.Empty();
+
+	if(Actor.IsValid())
+	{
+		// Add root component (which will add its children)
+		USceneComponent* RootComp = Actor->GetRootComponent();
+		if(RootComp && !RootComp->IsPendingKill())
+		{
+			FComponentTreeNodePtrType RootNode = MakeShareable(new FComponentTreeNode(RootComp));
+			RootNodes.Add(RootNode);
+		}
+
+		// Add non-scene components in Components array
+		TInlineComponentArray<UActorComponent*> Components;
+		Actor->GetComponents(Components);
+
+		for(int32 CompIdx=0; CompIdx<Components.Num(); CompIdx++)
+		{
+			UActorComponent* Comp = Components[CompIdx];
+			if(!Comp->IsA(USceneComponent::StaticClass()) && !Comp->IsPendingKill())
+			{
+				FComponentTreeNodePtrType Node = MakeShareable(new FComponentTreeNode(Comp));
+				RootNodes.Add(Node);
+			}
+		}
+	}
+
+	// refresh widget
+	for (const FComponentTreeNodePtrType& RootNode : RootNodes)
+	{
+		TreeWidget->SetItemExpansion(RootNode, true);
+	}
+	TreeWidget->RequestTreeRefresh();
+}
diff --git a/Engine/Source/Editor/LevelEditor/Private/ComponentsTree.h b/Engine/Source/Editor/LevelEditor/Private/ComponentsTree.h
new file mode 100644
--- /dev/null
+++ b/Engine/Source/Editor/LevelEditor/Private/ComponentsTree.h
@@ -0,0 +1,69 @@
+// Copyright 1998-2015 Epic Games, Inc. All Rights Reserved.
+
+
+#pragma once
+
+typedef TSharedPtr<class FComponentTreeNode> FComponentTreeNodePtrType;
+typedef STreeView<FComponentTreeNodePtrType> SComponentTreeType;
+
+//////////////////////////////////////////////////////////////////////////
+
+class FComponentTreeNode : public TSharedFromThis<FComponentTreeNode>
+{
+public:
+	FComponentTreeNode(UActorComponent* InComponent);
+
+	TWeakObjectPtr<UActorComponent>	Component;
+
+	TArray<FComponentTreeNodePtrType> ChildNodes;
+};
+
+//////////////////////////////////////////////////////////////////////////
+
+class SComponentRowWidget : public SCompoundWidget
+{
+public:
+	SLATE_BEGIN_ARGS( SComponentRowWidget ){}
+	SLATE_END_ARGS()
+
+	void Construct( const FArguments& InArgs, FComponentTreeNodePtrType InNodePtr );
+
+protected:
+	/** Pointer to node we represent */
+	TWeakPtr<FComponentTreeNode> NodePtr;
+};
+
+//////////////////////////////////////////////////////////////////////////
+
+class SComponentsTree : public SCompoundWidget
+{
+
+public:
+
+	SLATE_BEGIN_ARGS( SComponentsTree ){}
+
+	SLATE_END_ARGS()
+
+	void Construct( const FArguments& InArgs, TSharedPtr<class IDetailsView> InPropertyView );
+
+	void UpdateTree();
+
+	void SetObjects(const TArray<UObject*>& InObjects);
+
+	// Tree delegates
+	TSharedRef<ITableRow> MakeTableRowWidget( FComponentTreeNodePtrType InNodePtr, const TSharedRef<STableViewBase>& OwnerTable );
+	void OnGetChildrenForTree( FComponentTreeNodePtrType InNodePtr, TArray<FComponentTreeNodePtrType>& OutChildren );
+	void OnTreeSelectionChanged(FComponentTreeNodePtrType InSelectedNodePtr, ESelectInfo::Type SelectInfo);
+
+	void OnSelectedCompClass(TSubclassOf<UActorComponent> CompClass);
+
+private:
+
+	TWeakObjectPtr<AActor>				Actor;
+
+	TArray<FComponentTreeNodePtrType>	RootNodes;
+
+	TSharedPtr<SComponentTreeType>		TreeWidget;
+
+	TSharedPtr<class IDetailsView>		PropertyView;
+};
\ No newline at end of file
diff --git a/Engine/Source/Editor/LevelEditor/Private/SLevelEditor.cpp b/Engine/Source/Editor/LevelEditor/Private/SLevelEditor.cpp
--- a/Engine/Source/Editor/LevelEditor/Private/SLevelEditor.cpp
+++ b/Engine/Source/Editor/LevelEditor/Private/SLevelEditor.cpp
@@ -38,7 +38,7 @@
 #include "NewsFeed.h"
 #include "TutorialMetaData.h"
 #include "SDockTab.h"
-
+#include "ComponentsTree.h"
 
 static const FName LevelEditorBuildAndSubmitTab("LevelEditorBuildAndSubmit");
 static const FName LevelEditorStatsViewerTab("LevelEditorStatsViewer");
@@ -143,6 +143,10 @@ void SLevelEditor::Initialize( const TSharedRef<SDockTab>& OwnerTab, const TShar
 	// Bind the level editor tab's label to the currently loaded level name string in the main frame
 	OwnerTab->SetLabel( TAttribute<FText>( this, &SLevelEditor::GetTabTitle) );
 
+	FLevelEditorModule& LevelEditorModule = FModuleManager::GetModuleChecked< FLevelEditorModule >(LevelEditorModuleName);
+
+	LevelEditorModule.OnActorSelectionChanged().AddSP(this, &SLevelEditor::OnActorSelectionChanged);
+
 	TSharedRef<SWidget> Widget2 = RestoreContentArea( OwnerTab, OwnerWindow );
 	TSharedRef<SWidget> Widget1 = FLevelEditorMenu::MakeLevelEditorMenu( LevelEditorCommands, SharedThis(this) );
 
@@ -486,35 +490,105 @@ TSharedRef<FTabManager> SLevelEditor::GetTabManager() const
 	return LevelEditorTabManager.ToSharedRef();
 }
 
-// @todo Slate TEMP to support both detail views
-static TSharedRef<SDockTab> SummonDetailsPanel( FName TabIdentifier )
+class SActorDetails : public SCompoundWidget
 {
-	struct Local
+public:
+	SLATE_BEGIN_ARGS( SActorDetails )
+	{}
+	SLATE_END_ARGS()
+
+	void Construct( const FArguments& InArgs, const FName TabIdentifier )
 	{
-		static bool IsPropertyVisible( const FPropertyAndParent& PropertyAndParent )
+		struct Local
 		{
-			// For details views in the level editor all properties are the instanced versions
-			if (PropertyAndParent.Property.HasAllPropertyFlags(CPF_DisableEditOnInstance))
+			static bool IsPropertyVisible( const FPropertyAndParent& PropertyAndParent )
 			{
-				return false;
+				// For details views in the level editor all properties are the instanced versions
+				if (PropertyAndParent.Property.HasAllPropertyFlags(CPF_DisableEditOnInstance))
+				{
+					return false;
+				}
+
+				return true;
 			}
+		};
+
+		FPropertyEditorModule& PropPlugin = FModuleManager::LoadModuleChecked<FPropertyEditorModule>("PropertyEditor");
+		const FDetailsViewArgs DetailsViewArgs( true, true, true, false, false, GUnrealEd, false, TabIdentifier );
+		DetailsView = PropPlugin.CreateDetailView(DetailsViewArgs);
 
-			return true;
+		DetailsView->SetIsPropertyVisibleDelegate( FIsPropertyVisible::CreateStatic( &Local::IsPropertyVisible ) );
+
+		// Set up a delegate to call to add generic details to the view
+		DetailsView->SetGenericLayoutDetailsDelegate( FOnGetDetailCustomizationInstance::CreateStatic( &FLevelEditorGenericDetails::MakeInstance ) );
+
+		SAssignNew(ComponentsBox, SScrollBox)
+		+SScrollBox::Slot()
+		[
+			SAssignNew(ComponentsTree, SComponentsTree, DetailsView)
+			.Visibility(EVisibility::Collapsed)
+		];
+
+		ChildSlot
+		[
+			SAssignNew(DetailsSplitter, SSplitter)
+			.Orientation(Orient_Vertical)
+			+ SSplitter::Slot()
+			[
+				DetailsView.ToSharedRef()
+			]
+		];
+
+	}
+
+	void SetObjects( const TArray<UObject*>& InObjects )
+	{
+		if(!DetailsView->IsLocked() )
+		{
+			DetailsView->SetObjects( InObjects );
+				
+			if( GetDefault<UEditorExperimentalSettings>()->bInWorldBPEditing )
+			{
+				bool bTreeVisible = ComponentsTree->GetVisibility().IsVisible();
+				ComponentsTree->SetObjects(InObjects);
+				if( bTreeVisible != ComponentsTree->GetVisibility().IsVisible())
+				{
+					if( bTreeVisible )
+					{
+						DetailsSplitter->RemoveAt(0);
+					}
+					else
+					{
+						DetailsSplitter->AddSlot(0)
+						.Value(.2f)
+						[
+							ComponentsBox.ToSharedRef()
+						];
+					}
+				}
+			}
+			else
+			{
+				ComponentsTree->SetVisibility( EVisibility::Collapsed );
+			}
 		}
-	};
+	}
 
-	FPropertyEditorModule& PropPlugin = FModuleManager::LoadModuleChecked<FPropertyEditorModule>("PropertyEditor");
-	const FDetailsViewArgs DetailsViewArgs( true, true, true, false, false, GUnrealEd, false, TabIdentifier );
-	TSharedRef<IDetailsView> DetailsView = PropPlugin.CreateDetailView( DetailsViewArgs );
+private:
+	TSharedPtr<SSplitter> DetailsSplitter;
+	TSharedPtr<IDetailsView> DetailsView;
+	TSharedPtr<SScrollBox> ComponentsBox;
+	TSharedPtr<SComponentsTree> ComponentsTree;
+};
 
-	DetailsView->SetIsPropertyVisibleDelegate( FIsPropertyVisible::CreateStatic( &Local::IsPropertyVisible ) );
 
-	// Set up a delegate to call to add generic details to the view
-	DetailsView->SetGenericLayoutDetailsDelegate( FOnGetDetailCustomizationInstance::CreateStatic( &FLevelEditorGenericDetails::MakeInstance ) );
+TSharedRef<SDockTab> SLevelEditor::SummonDetailsPanel( FName TabIdentifier )
+{
+	TSharedPtr<SActorDetails> ActorDetails;
 
 	FText Label = NSLOCTEXT( "LevelEditor", "DetailsTabTitle", "Details" );
 
-	return SNew( SDockTab )
+	TSharedRef<SDockTab> DocTab = SNew(SDockTab)
 		.Icon( FEditorStyle::GetBrush( "LevelEditor.Tabs.Details" ) )
 		.Label( Label )
 		.ToolTip( IDocumentation::Get()->CreateToolTip( Label, nullptr, "Shared/LevelEditor", "DetailsTab" ) )
@@ -522,11 +596,14 @@ static TSharedRef<SDockTab> SummonDetailsPanel( FName TabIdentifier )
 			SNew( SBox )
 			.AddMetaData<FTutorialMetaData>(FTutorialMetaData(TEXT("ActorDetails"), TEXT("LevelEditorSelectionDetails")))
 			[
-				DetailsView
+				SAssignNew( ActorDetails, SActorDetails, TabIdentifier )
 			]
 		];
-}
 
+	AllActorDetailPanels.Add( ActorDetails );
+
+	return DocTab;
+}
 /** Method to call when a tab needs to be spawned by the FLayoutService */
 TSharedRef<SDockTab> SLevelEditor::SpawnLevelEditorTab( const FSpawnTabArgs& Args, FName TabIdentifier, FString InitializationPayload )
 {
@@ -555,7 +632,7 @@ TSharedRef<SDockTab> SLevelEditor::SpawnLevelEditorTab( const FSpawnTabArgs& Arg
 			[
 				SNew(SHorizontalBox)
 				.AddMetaData<FTagMetaData>(FTagMetaData(TEXT("LevelToolbar")))
-				+ SHorizontalBox::Slot()
+				+SHorizontalBox::Slot()
 				.FillWidth(1)
 				.VAlign(VAlign_Bottom)
 				.HAlign(HAlign_Left)
@@ -1419,3 +1496,19 @@ void SLevelEditor::HandleEditorMapChange( uint32 MapChangeFlags )
 		WorldSettingsView->SetObject(GetWorld()->GetWorldSettings(), true);
 	}
 }
+
+void SLevelEditor::OnActorSelectionChanged( const TArray<UObject*>& NewSelection )
+{
+	for( auto It = AllActorDetailPanels.CreateIterator(); It; ++It )
+	{
+		TSharedPtr<SActorDetails> ActorDetails = It->Pin();
+		if( ActorDetails.IsValid() )
+		{
+			ActorDetails->SetObjects( NewSelection );
+		}
+		else
+		{
+			// remove stray entries here
+		}
+	}
+}
\ No newline at end of file
diff --git a/Engine/Source/Editor/LevelEditor/Private/SLevelEditor.h b/Engine/Source/Editor/LevelEditor/Private/SLevelEditor.h
--- a/Engine/Source/Editor/LevelEditor/Private/SLevelEditor.h
+++ b/Engine/Source/Editor/LevelEditor/Private/SLevelEditor.h
@@ -19,10 +19,6 @@ class SLevelEditor
 {
 
 public:
-	// @todo why are these public?
-	TSharedRef<SDockTab> SpawnLevelEditorTab( const FSpawnTabArgs& Args, FName TabIdentifier, FString InitializationPayload );
-	//TSharedRef<SDockTab> SpawnLevelEditorModeTab( const FSpawnTabArgs& Args, FEdMode* EditorMode );
-
 	SLATE_BEGIN_ARGS( SLevelEditor ){}
 
 	SLATE_END_ARGS()
@@ -128,6 +124,10 @@ public:
 	TSharedPtr<SDockTab> SequencerTab;
 
 private:
+	
+	TSharedRef<SDockTab> SpawnLevelEditorTab(const FSpawnTabArgs& Args, FName TabIdentifier, FString InitializationPayload);
+	//TSharedRef<SDockTab> SpawnLevelEditorModeTab(const FSpawnTabArgs& Args, FEdMode* EditorMode);
+	TSharedRef<SDockTab> SummonDetailsPanel( FName Identifier );
 
 	/**
 	 * Binds UI commands to actions for the level editor                   
@@ -187,9 +187,11 @@ private:
 	/** Reset the transient viewport information */
 	void ResetViewportTabInfo();
 
-	// Handles Editor map changes.
+	/** Handles Editor map changes */
 	void HandleEditorMapChange( uint32 MapChangeFlags );
 
+	/** Called when actors are selected or unselected */
+	void OnActorSelectionChanged( const TArray<UObject*>& NewSelection );
 private:
 
 	// Tracking the active viewports in this level editor.
@@ -223,6 +225,9 @@ private:
 
 	/** Transient editor viewport states - one for each view type. Key is "LayoutId[ELevelViewportType]", eg) "Viewport 1[0]" */
 	TMap<FString, FLevelViewportInfo> TransientEditorViews;
+
+	/** List of all actor details panels to update when selection changes */
+	TArray< TWeakPtr<class SActorDetails> > AllActorDetailPanels;
 };
 
 
diff --git a/Engine/Source/Editor/PlacementMode/Private/SPlacementModeTools.cpp b/Engine/Source/Editor/PlacementMode/Private/SPlacementModeTools.cpp
--- a/Engine/Source/Editor/PlacementMode/Private/SPlacementModeTools.cpp
+++ b/Engine/Source/Editor/PlacementMode/Private/SPlacementModeTools.cpp
@@ -666,7 +666,7 @@ TSharedRef< SWidget > SPlacementModeTools::BuildVisualWidget()
 
 TSharedRef< SWidget > SPlacementModeTools::BuildBasicWidget()
 {
-	return SNew( SVerticalBox )
+	TSharedRef<SVerticalBox> VerticalBox = SNew( SVerticalBox )
 
 	// Basics
 	+ SVerticalBox::Slot()
@@ -674,45 +674,42 @@ TSharedRef< SWidget > SPlacementModeTools::BuildBasicWidget()
 	[
 		BuildDraggableAssetWidget( UActorFactoryCameraActor::StaticClass() )
 	]
-
 	+ SVerticalBox::Slot()
 	.AutoHeight()
 	[
 		BuildDraggableAssetWidget( UActorFactoryPlayerStart::StaticClass() )
 	]
-
 	+ SVerticalBox::Slot()
 	.AutoHeight()
 	[
-		BuildDraggableAssetWidget(UActorFactoryPointLight::StaticClass())
+		BuildDraggableAssetWidget( UActorFactoryPointLight::StaticClass() )
 	]
-
-
-	// Triggers
 	+ SVerticalBox::Slot()
 	.AutoHeight()
 	[
 		BuildDraggableAssetWidget( UActorFactoryTriggerBox::StaticClass() )
 	]
-
 	+ SVerticalBox::Slot()
 	.AutoHeight()
 	[
 		BuildDraggableAssetWidget( UActorFactoryTriggerSphere::StaticClass() )
 	]
-
 	+ SVerticalBox::Slot()
 	.AutoHeight()
 	[
 		BuildDraggableAssetWidget( UActorFactoryTriggerCapsule::StaticClass() )
-	]
+	];
 
+	if( GetDefault<UEditorExperimentalSettings>()->bInWorldBPEditing )
+	{
+		VerticalBox->InsertSlot(0)
+		.AutoHeight()
+		[
+			BuildDraggableAssetWidget(UActorFactoryEmptyActor::StaticClass())
+		];
+	}
 
-	+ SVerticalBox::Slot()
-	.AutoHeight()
-	[
-		BuildDraggableAssetWidget( UActorFactoryTargetPoint::StaticClass() )
-	];
+	return VerticalBox;
 }
 
 void SPlacementModeTools::UpdateRecentlyPlacedAssets( const TArray< FActorPlacementInfo >& RecentlyPlaced )
diff --git a/Engine/Source/Editor/PropertyEditor/Private/PropertyEditorModule.cpp b/Engine/Source/Editor/PropertyEditor/Private/PropertyEditorModule.cpp
--- a/Engine/Source/Editor/PropertyEditor/Private/PropertyEditorModule.cpp
+++ b/Engine/Source/Editor/PropertyEditor/Private/PropertyEditorModule.cpp
@@ -507,22 +507,11 @@ void FPropertyEditorModule::UpdatePropertyViews( const TArray<UObject*>& NewObje
 			{
 				if( !DetailViewPin->IsLocked() )
 				{
-					DetailViewPin->SetObjects( NewObjectList, true );
+					DetailViewPin->SetObjects(NewObjectList, true);
 				}
 				else
 				{
-					if( ValidObjects.Num() == 0 )
-					{
-						// Populate the valid objects list only once
-						UWorld* IteratorWorld = GWorld;
-						for( FActorIterator It(IteratorWorld); It; ++It )
-						{
-							ValidObjects.Add( *It );
-						}
-					}
-
-					// If the property window is locked make sure all the actors still exist
-					DetailViewPin->RemoveInvalidActors( ValidObjects );
+					DetailViewPin->RemoveInvalidObjects();
 				}
 			}
 		}
diff --git a/Engine/Source/Editor/PropertyEditor/Private/SDetailsView.cpp b/Engine/Source/Editor/PropertyEditor/Private/SDetailsView.cpp
--- a/Engine/Source/Editor/PropertyEditor/Private/SDetailsView.cpp
+++ b/Engine/Source/Editor/PropertyEditor/Private/SDetailsView.cpp
@@ -313,6 +313,31 @@ void SDetailsView::SetObject( UObject* InObject, bool bForceRefresh )
 	SetObjects( ObjectWeakPtrs, bForceRefresh );
 }
 
+void SDetailsView::RemoveInvalidObjects()
+{
+	TArray< TWeakObjectPtr< UObject > > ResetArray;
+
+	bool bAllFound = true;
+	for (TPropObjectIterator Itor(RootPropertyNode->ObjectIterator()); Itor; ++Itor)
+	{
+		TWeakObjectPtr<UObject> Object = *Itor;
+
+		if( Object.IsValid() && !Object->IsPendingKill() )
+		{
+			ResetArray.Add(Object);
+		}
+		else
+		{
+			bAllFound = false;
+		}
+	}
+
+	if (!bAllFound)
+	{
+		SetObjectArrayPrivate(ResetArray);
+	}
+}
+
 bool SDetailsView::ShouldSetNewObjects( const TArray< TWeakObjectPtr< UObject > >& InObjects ) const
 {
 	bool bShouldSetObjects = false;
@@ -512,39 +537,6 @@ void SDetailsView::RemoveDeletedObjects( const TArray<UObject*>& DeletedObjects
 	}
 }
 
-/**
- * Removes actors from the property nodes object array which are no longer available
- * 
- * @param ValidActors	The list of actors which are still valid
- */
-void SDetailsView::RemoveInvalidActors( const TSet<AActor*>& ValidActors )
-{
-	TArray< TWeakObjectPtr< UObject > > ResetArray;
-
-	bool bAllFound = true;
-	for ( TPropObjectIterator Itor( RootPropertyNode->ObjectIterator() ); Itor; ++Itor )
-	{
-		AActor* Actor = Cast<AActor>( Itor->Get() );
-
-		bool bFound = ValidActors.Contains( Actor );
-
-		// If the selected actor no longer exists, remove it from the property window.
-		if( bFound )
-		{
-			ResetArray.Add(Actor);
-		}
-		else
-		{
-			bAllFound = false;
-		}
-	}
-
-	if ( !bAllFound ) 
-	{
-		SetObjectArrayPrivate( ResetArray );
-	}
-}
-
 /** Called before during SetObjectArray before we change the objects being observed */
 void SDetailsView::PreSetObject()
 {
diff --git a/Engine/Source/Editor/PropertyEditor/Private/SDetailsView.h b/Engine/Source/Editor/PropertyEditor/Private/SDetailsView.h
--- a/Engine/Source/Editor/PropertyEditor/Private/SDetailsView.h
+++ b/Engine/Source/Editor/PropertyEditor/Private/SDetailsView.h
@@ -30,6 +30,7 @@ public:
 	virtual void SetObjects( const TArray<UObject*>& InObjects, bool bForceRefresh = false ) override;
 	virtual void SetObjects( const TArray< TWeakObjectPtr< UObject > >& InObjects, bool bForceRefresh = false ) override;
 	virtual void SetObject( UObject* InObject, bool bForceRefresh = false ) override;
+	virtual void RemoveInvalidObjects() override;
 
 	/**
 	 * Replaces objects being observed by the view with new objects
@@ -45,13 +46,6 @@ public:
 	 */
 	void RemoveDeletedObjects( const TArray<UObject*>& DeletedObjects );
 
-	/**
-     * Removes actors from the property nodes object array which are no longer available
-	 * 
-	 * @param ValidActors	The list of actors which are still valid
-	 */
-	void RemoveInvalidActors( const TSet<AActor*>& ValidActors );
-
 	/** Sets the callback for when the property view changes */
 	virtual void SetOnObjectArrayChanged( FOnObjectArrayChanged OnObjectArrayChangedDelegate);
 
diff --git a/Engine/Source/Editor/PropertyEditor/Private/SStructureDetailsView.h b/Engine/Source/Editor/PropertyEditor/Private/SStructureDetailsView.h
--- a/Engine/Source/Editor/PropertyEditor/Private/SStructureDetailsView.h
+++ b/Engine/Source/Editor/PropertyEditor/Private/SStructureDetailsView.h
@@ -78,12 +78,13 @@ public:
 		return false;
 	}
 
-	virtual void SetOnObjectArrayChanged(FOnObjectArrayChanged OnObjectArrayChangedDelegate) override{}
+	virtual void SetOnObjectArrayChanged(FOnObjectArrayChanged OnObjectArrayChangedDelegate) override {}
 	virtual void RegisterInstancedCustomPropertyLayout(UClass* Class, FOnGetDetailCustomizationInstance DetailLayoutDelegate) override {}
 	virtual void UnregisterInstancedCustomPropertyLayout(UClass* Class) override {}
-	virtual void SetObjects(const TArray<UObject*>& InObjects, bool bForceRefresh = false) override{}
-	virtual void SetObjects(const TArray< TWeakObjectPtr< UObject > >& InObjects, bool bForceRefresh = false) override{}
+	virtual void SetObjects(const TArray<UObject*>& InObjects, bool bForceRefresh = false) override {}
+	virtual void SetObjects(const TArray< TWeakObjectPtr< UObject > >& InObjects, bool bForceRefresh = false) override {}
 	virtual void SetObject(UObject* InObject, bool bForceRefresh = false) override{}
+	virtual void RemoveInvalidObjects() override {}
 
 	virtual TSharedPtr<class FComplexPropertyNode> GetRootNode() override;
 protected:
diff --git a/Engine/Source/Editor/PropertyEditor/Public/IDetailsView.h b/Engine/Source/Editor/PropertyEditor/Public/IDetailsView.h
--- a/Engine/Source/Editor/PropertyEditor/Public/IDetailsView.h
+++ b/Engine/Source/Editor/PropertyEditor/Public/IDetailsView.h
@@ -122,6 +122,9 @@ public:
 	 */
 	virtual void SetObject( UObject* InObject, bool bForceRefresh = false ) = 0;
 
+	/** Removes all invalid objects being observed by this details panel */
+	virtual void RemoveInvalidObjects() = 0;
+
 	/**
 	 * Returns true if the details view is locked and cant have its observed objects changed 
 	 */
diff --git a/Engine/Source/Editor/UnrealEd/Classes/ActorFactories/ActorFactoryEmptyActor.h b/Engine/Source/Editor/UnrealEd/Classes/ActorFactories/ActorFactoryEmptyActor.h
new file mode 100644
--- /dev/null
+++ b/Engine/Source/Editor/UnrealEd/Classes/ActorFactories/ActorFactoryEmptyActor.h
@@ -0,0 +1,17 @@
+// Copyright 1998-2015 Epic Games, Inc. All Rights Reserved.
+
+
+#pragma once
+#include "ActorFactoryEmptyActor.generated.h"
+
+UCLASS(MinimalAPI,config=Editor)
+class UActorFactoryEmptyActor : public UActorFactory
+{
+	GENERATED_UCLASS_BODY()
+
+public:
+	virtual bool CanCreateActorFrom( const FAssetData& AssetData, FText& OutErrorMsg ) override;
+
+protected:
+	virtual AActor* SpawnActor( UObject* Asset, ULevel* InLevel, const FVector& Location, const FRotator& Rotation, EObjectFlags ObjectFlags, const FName& Name ) override;
+};
diff --git a/Engine/Source/Editor/UnrealEd/Classes/Editor/EditorEngine.h b/Engine/Source/Editor/UnrealEd/Classes/Editor/EditorEngine.h
--- a/Engine/Source/Editor/UnrealEd/Classes/Editor/EditorEngine.h
+++ b/Engine/Source/Editor/UnrealEd/Classes/Editor/EditorEngine.h
@@ -1663,6 +1663,20 @@ public:
 	 */
 	class FSelectionIterator GetSelectedActorIterator() const;
 
+	/**
+	* Returns an FSelectionIterator that iterates over the set of selected components.
+	*/
+	class FSelectionIterator GetSelectedComponentIterator() const;
+
+	/**
+	* Returns the number of currently selected components.
+	*/
+	int32 GetSelectedComponentCount() const;
+
+	/**
+	* @return the set of selected components.
+	*/
+	class USelection* GetSelectedComponents() const;
 
 	/**
 	 * @return the set of selected non-actor objects.
diff --git a/Engine/Source/Editor/UnrealEd/Private/EdMode.cpp b/Engine/Source/Editor/UnrealEd/Private/EdMode.cpp
--- a/Engine/Source/Editor/UnrealEd/Private/EdMode.cpp
+++ b/Engine/Source/Editor/UnrealEd/Private/EdMode.cpp
@@ -432,7 +432,24 @@ FVector FEdMode::GetWidgetLocation() const
 
 bool FEdMode::ShouldDrawWidget() const
 {
-	return (Owner->GetSelectedActors()->GetTop<AActor>() != NULL);
+	bool bDrawWidget = false;
+	if (GEditor->GetSelectedComponentCount() > 0)
+	{
+		for (FSelectionIterator It(GEditor->GetSelectedComponentIterator()); It; ++It)
+		{
+			if (It->IsA<USceneComponent>())
+			{
+				bDrawWidget = true;
+				break;
+			}
+		}
+	}
+	else
+	{
+		bDrawWidget = (GEditor->GetSelectedActors()->GetTop<AActor>() != NULL);
+	}
+
+	return bDrawWidget;
 }
 
 /**
diff --git a/Engine/Source/Editor/UnrealEd/Private/Editor.cpp b/Engine/Source/Editor/UnrealEd/Private/Editor.cpp
--- a/Engine/Source/Editor/UnrealEd/Private/Editor.cpp
+++ b/Engine/Source/Editor/UnrealEd/Private/Editor.cpp
@@ -162,18 +162,56 @@ static inline USelection*& PrivateGetSelectedActors()
 	return SSelectedActors;
 };
 
+static inline USelection*& PrivateGetSelectedComponents()
+{
+	static USelection* SSelectedComponents = NULL;
+	return SSelectedComponents;
+}
+
 static inline USelection*& PrivateGetSelectedObjects()
 {
 	static USelection* SSelectedObjects = NULL;
 	return SSelectedObjects;
 };
 
+static void OnObjectSelected(UObject* Object)
+{
+	// Whenever an actor is unselected we must remove its components from the components selection
+	if (!Object->IsSelected())
+	{
+		TArray<UActorComponent*> ComponentsToDeselect;
+		for (FSelectionIterator It(*PrivateGetSelectedComponents()); It; ++It)
+		{
+			UActorComponent* Component = CastChecked<UActorComponent>(*It);
+			if (Component->GetOwner() == Object)
+			{
+				ComponentsToDeselect.Add(Component);
+			}
+		}
+		if (ComponentsToDeselect.Num() > 0)
+		{
+			PrivateGetSelectedComponents()->Modify();
+			PrivateGetSelectedComponents()->BeginBatchSelectOperation();
+			for (UActorComponent* Component : ComponentsToDeselect)
+			{
+				PrivateGetSelectedComponents()->Deselect(Component);
+			}
+			PrivateGetSelectedComponents()->EndBatchSelectOperation();
+		}
+	}
+}
+
 static void PrivateInitSelectedSets()
 {
 	PrivateGetSelectedActors() = new( GetTransientPackage(), TEXT("SelectedActors"), RF_Transactional ) USelection(FObjectInitializer());
 	PrivateGetSelectedActors()->AddToRoot();
 
-	PrivateGetSelectedObjects() = new( GetTransientPackage(), TEXT("SelectedObjects"), RF_Transactional ) USelection(FObjectInitializer());
+	PrivateGetSelectedActors()->SelectObjectEvent.AddStatic(&OnObjectSelected);
+
+	PrivateGetSelectedComponents() = new(GetTransientPackage(), TEXT("SelectedComponents"), RF_Transactional) USelection(FObjectInitializer());
+	PrivateGetSelectedComponents()->AddToRoot();
+
+	PrivateGetSelectedObjects() = new(GetTransientPackage(), TEXT("SelectedObjects"), RF_Transactional) USelection(FObjectInitializer());
 	PrivateGetSelectedObjects()->AddToRoot();
 }
 
@@ -182,6 +220,8 @@ static void PrivateDestroySelectedSets()
 #if 0
 	PrivateGetSelectedActors()->RemoveFromRoot();
 	PrivateGetSelectedActors() = NULL;
+	PrivateGetSelectedComponents()->RemoveFromRoot();
+	PrivateGetSelectedComponents() = NULL;
 	PrivateGetSelectedObjects()->RemoveFromRoot();
 	PrivateGetSelectedObjects() = NULL;
 #endif
@@ -278,6 +318,27 @@ FSelectionIterator UEditorEngine::GetSelectedActorIterator() const
 	return FSelectionIterator( *GetSelectedActors() );
 };
 
+int32 UEditorEngine::GetSelectedComponentCount() const
+{
+	int32 NumSelectedComponents = 0;
+	for (FSelectionIterator It(GetSelectedComponentIterator()); It; ++It)
+	{
+		++NumSelectedComponents;
+	}
+
+	return NumSelectedComponents;
+}
+
+FSelectionIterator UEditorEngine::GetSelectedComponentIterator() const
+{
+	return FSelectionIterator(*GetSelectedComponents());
+};
+
+USelection* UEditorEngine::GetSelectedComponents() const
+{
+	return PrivateGetSelectedComponents();
+}
+
 USelection* UEditorEngine::GetSelectedObjects() const
 {
 	return PrivateGetSelectedObjects();
diff --git a/Engine/Source/Editor/UnrealEd/Private/EditorHook.cpp b/Engine/Source/Editor/UnrealEd/Private/EditorHook.cpp
--- a/Engine/Source/Editor/UnrealEd/Private/EditorHook.cpp
+++ b/Engine/Source/Editor/UnrealEd/Private/EditorHook.cpp
@@ -46,14 +46,18 @@ void UUnrealEdEngine::UpdateFloatingPropertyWindows()
 		FSelectionIterator It( GetSelectedActorIterator() );
 		if ( It )
 		{
-			AActor* Actor = static_cast<AActor*>( *It );
-			checkSlow( Actor->IsA(AActor::StaticClass()) );
-
-			if ( !Actor->IsPendingKill() )
+			ALandscapeProxy* LP = Cast<ALandscapeProxy>(*It);
+			if (LP != NULL && !LP->IsPendingKill())
 			{
-				// Currently only for Landscape
-				// Landscape uses Component Selection
-				bProcessed = Actor->GetSelectedComponents(SelectedObjects);
+				ULandscapeInfo* Info = LP->GetLandscapeInfo(false);
+				if (Info && Info->bCurrentlyEditing)
+				{
+					for (ULandscapeComponent* LandscapeComponent : Info->GetSelectedComponents())
+					{
+						SelectedObjects.Add(LandscapeComponent);
+					}
+					bProcessed = true;
+				}
 			}
 		}
 	}
@@ -79,12 +83,7 @@ void UUnrealEdEngine::UpdateFloatingPropertyWindows()
 
 void UUnrealEdEngine::UpdateFloatingPropertyWindowsFromActorList( const TArray< UObject *>& ActorList )
 {
-	FPropertyEditorModule& PropertyEditorModule = FModuleManager::LoadModuleChecked<FPropertyEditorModule>( TEXT("PropertyEditor") );
-
-	PropertyEditorModule.UpdatePropertyViews( ActorList );
-
-	// Also update any selection listers in the level editor
-	FLevelEditorModule& LevelEditor = FModuleManager::LoadModuleChecked<FLevelEditorModule>( TEXT("LevelEditor") );
+	FLevelEditorModule& LevelEditor = FModuleManager::LoadModuleChecked<FLevelEditorModule>(TEXT("LevelEditor"));
 
 	LevelEditor.BroadcastActorSelectionChanged( ActorList );
 }
diff --git a/Engine/Source/Editor/UnrealEd/Private/EditorSelectUtils.cpp b/Engine/Source/Editor/UnrealEd/Private/EditorSelectUtils.cpp
--- a/Engine/Source/Editor/UnrealEd/Private/EditorSelectUtils.cpp
+++ b/Engine/Source/Editor/UnrealEd/Private/EditorSelectUtils.cpp
@@ -271,25 +271,55 @@ void UUnrealEdEngine::SetActorSelectionFlags (AActor* InActor)
 void UUnrealEdEngine::UpdatePivotLocationForSelection( bool bOnChange )
 {
 	// Pick a new common pivot, or not.
-	int32 ActorCount=0;
 	AActor* SingleActor=NULL;
+	USceneComponent* SingleComponent = NULL;
 
-	for ( FSelectionIterator It( GetSelectedActorIterator() ) ; It ; ++It )
+	if (GetSelectedComponentCount() > 0)
 	{
-		AActor* Actor = static_cast<AActor*>( *It );
-		checkSlow( Actor->IsA(AActor::StaticClass()) );
+		for (FSelectionIterator It(GetSelectedComponentIterator()); It; ++It)
+		{
+			UActorComponent* Component = CastChecked<UActorComponent>(*It);
+			AActor* ComponentOwner = Component->GetOwner();
+			check(ComponentOwner);
+			check(GetSelectedActors()->IsSelected(ComponentOwner));
+
+			if (ComponentOwner->GetWorld() == GWorld)
+			{
+				SingleActor = ComponentOwner;
+				if (Component->IsA<USceneComponent>())
+				{
+					SingleComponent = CastChecked<USceneComponent>(Component);
+				}
 
-		if ( Actor->GetWorld() == GWorld )
+				const bool IsTemplate = ComponentOwner->IsTemplate();
+				const bool LevelLocked = !FLevelUtils::IsLevelLocked(ComponentOwner->GetLevel());
+				check(IsTemplate || LevelLocked);
+			}
+		}
+	}
+	else
+	{
+		for (FSelectionIterator It(GetSelectedActorIterator()); It; ++It)
 		{
-			SingleActor = Actor;
-			bool IsTemplate = Actor->IsTemplate();
-			bool LevelLocked = !FLevelUtils::IsLevelLocked(Actor->GetLevel());
-			check( IsTemplate || LevelLocked );
-			ActorCount++;
+			AActor* Actor = static_cast<AActor*>(*It);
+			checkSlow(Actor->IsA(AActor::StaticClass()));
+
+			if (Actor->GetWorld() == GWorld)
+			{
+				const bool IsTemplate = Actor->IsTemplate();
+				const bool LevelLocked = !FLevelUtils::IsLevelLocked(Actor->GetLevel());
+				check(IsTemplate || LevelLocked);
+
+				SingleActor = Actor;
+			}
 		}
 	}
 	
-	if( ActorCount > 0 ) 
+	if (SingleComponent != NULL)
+	{
+		SetPivot(SingleComponent->GetComponentLocation(), false, true);
+	}
+	else if( SingleActor != NULL ) 
 	{		
 		// For geometry mode use current pivot location as it's set to selected face, not actor
 		FEditorModeTools& Tools = GLevelEditorModeTools();
@@ -333,8 +363,8 @@ void UUnrealEdEngine::NoteSelectionChange()
 		ActiveModes[ModeIndex]->ActorSelectionChangeNotify();
 	}
 
-	USelection* Section = GetSelectedActors();
-	USelection::SelectionChangedEvent.Broadcast(Section);
+	USelection* Selection = GetSelectedActors();
+	USelection::SelectionChangedEvent.Broadcast(Selection);
 
 	//whenever selection changes, recompute whether the selection contains a locked actor
 	bCheckForLockActors = true;
@@ -518,6 +548,13 @@ void UUnrealEdEngine::SelectActor(AActor* Actor, bool bInSelected, bool bNotify,
 			}
 
 			GetSelectedActors()->Select( Actor, bInSelected );
+			if (!bInSelected)
+			{
+				for (UActorComponent* Component : Actor->GetComponents())
+				{
+					GetSelectedComponents()->Deselect( Component );
+				}
+			}
 			
 			//A fast path to mark selection rather than reconnecting ALL components for ALL actors that have changed state
 			SetActorSelectionFlags (Actor);
diff --git a/Engine/Source/Editor/UnrealEd/Private/EditorServer.cpp b/Engine/Source/Editor/UnrealEd/Private/EditorServer.cpp
--- a/Engine/Source/Editor/UnrealEd/Private/EditorServer.cpp
+++ b/Engine/Source/Editor/UnrealEd/Private/EditorServer.cpp
@@ -531,7 +531,7 @@ void UEditorEngine::LoadAndSelectAssets( TArray<FAssetData>& Assets, UClass* Typ
 bool UEditorEngine::UsePercentageBasedScaling() const
 {
 	// Use percentage based scaling if the user setting is enabled or more than component or actor is selected.  Multiplicative scaling doesn't work when more than one object is selected
-	return GetDefault<ULevelEditorViewportSettings>()->UsePercentageBasedScaling() || GetSelectedActorCount() > 1;
+	return GetDefault<ULevelEditorViewportSettings>()->UsePercentageBasedScaling() || GetSelectedActorCount() > 1 || GetSelectedComponentCount() > 1;
 }
 
 bool UEditorEngine::Exec_Brush( UWorld* InWorld, const TCHAR* Str, FOutputDevice& Ar )
diff --git a/Engine/Source/Editor/UnrealEd/Private/Factories/ActorFactory.cpp b/Engine/Source/Editor/UnrealEd/Private/Factories/ActorFactory.cpp
--- a/Engine/Source/Editor/UnrealEd/Private/Factories/ActorFactory.cpp
+++ b/Engine/Source/Editor/UnrealEd/Private/Factories/ActorFactory.cpp
@@ -1112,6 +1112,54 @@ UActorFactoryCameraActor::UActorFactoryCameraActor(const FObjectInitializer& Obj
 	NewActorClass = ACameraActor::StaticClass();
 }
 
+/*-----------------------------------------------------------------------------
+UActorFactoryEmptyActor
+-----------------------------------------------------------------------------*/
+UActorFactoryEmptyActor::UActorFactoryEmptyActor(const FObjectInitializer& ObjectInitializer)
+	: Super(ObjectInitializer)
+{
+	DisplayName = LOCTEXT("EmptyActorDisplayName", "Empty Actor");
+	NewActorClass = AActor::StaticClass();
+}
+
+bool UActorFactoryEmptyActor::CanCreateActorFrom( const FAssetData& AssetData, FText& OutErrorMsg )
+{
+	return GetDefault<UEditorExperimentalSettings>()->bInWorldBPEditing ? AssetData.ObjectPath == FName(*AActor::StaticClass()->GetPathName()) : false;
+}
+
+AActor* UActorFactoryEmptyActor::SpawnActor( UObject* Asset, ULevel* InLevel, const FVector& Location, const FRotator& Rotation, EObjectFlags ObjectFlags, const FName& Name )
+{
+	AActor* NewActor = nullptr;
+	if(GetDefault<UEditorExperimentalSettings>()->bInWorldBPEditing) 
+	{
+		NewActor = Super::SpawnActor(Asset, InLevel, Location, Rotation, ObjectFlags, Name);
+
+		USceneComponent* RootComponent = ConstructObject<USceneComponent>(USceneComponent::StaticClass(), NewActor, FName("Root"), RF_Transactional);
+		RootComponent->Mobility = EComponentMobility::Static;
+		RootComponent->SetWorldLocationAndRotation(Location, Rotation);
+		NewActor->SetRootComponent(RootComponent);
+
+
+		// Create a new billboard component to serve as a visualization of the actor until there is another primitive component
+		UBillboardComponent* BillboardComponent = ConstructObject<UBillboardComponent>(UBillboardComponent::StaticClass(), NewActor, NAME_None, RF_Transactional);
+
+		BillboardComponent->Sprite = LoadObject<UTexture2D>(nullptr, TEXT("/Engine/EditorResources/S_Actor.S_Actor"));
+		BillboardComponent->RelativeScale3D = FVector(0.5f, 0.5f, 0.5f);
+		BillboardComponent->SpriteInfo.Category = TEXT("EmptyActors");
+		BillboardComponent->SpriteInfo.DisplayName = LOCTEXT("EmptyActorDisplayName", "Empty Actor");
+		BillboardComponent->AttachTo(RootComponent);
+		BillboardComponent->Mobility = EComponentMobility::Static;
+		BillboardComponent->AlwaysLoadOnClient = false;
+		BillboardComponent->AlwaysLoadOnServer = false;
+
+		NewActor->SerializedComponents.Add(RootComponent);
+		NewActor->SerializedComponents.Add(BillboardComponent);
+	}
+
+	return NewActor;
+}
+
+
 /*-----------------------------------------------------------------------------
 UActorFactoryAmbientSound
 -----------------------------------------------------------------------------*/
diff --git a/Engine/Source/Editor/UnrealEd/Private/LevelEditorViewport.cpp b/Engine/Source/Editor/UnrealEd/Private/LevelEditorViewport.cpp
--- a/Engine/Source/Editor/UnrealEd/Private/LevelEditorViewport.cpp
+++ b/Engine/Source/Editor/UnrealEd/Private/LevelEditorViewport.cpp
@@ -1385,6 +1385,10 @@ void FTrackingTransaction::Begin(const FText& Description)
 	{
 		GroupActor->Modify();
 	}
+	for (FSelectionIterator It(GEditor->GetSelectedComponentIterator()); It; ++It)
+	{
+		CastChecked<UActorComponent>(*It)->Modify();
+	}
 }
 
 void FTrackingTransaction::End()
@@ -2910,12 +2914,12 @@ void FLevelEditorViewportClient::ApplyDeltaToActors(const FVector& InDrag,
 			if(bIsSimulateInEditorViewport)
 			{
 				// If the Actor's outer (level) outer (world) is not the PlayWorld then it cannot be moved in this viewport.
-				if( !(GEditor->PlayWorld == Actor->GetOuter()->GetOuter()) )
+				if( !(GEditor->PlayWorld == Actor->GetWorld()) )
 				{
 					continue;
 				}
 			}
-			else if( !(GEditor->EditorWorld == Actor->GetOuter()->GetOuter()) )
+			else if( !(GEditor->EditorWorld == Actor->GetWorld()) )
 			{
 				continue;
 			}
@@ -2923,7 +2927,19 @@ void FLevelEditorViewportClient::ApplyDeltaToActors(const FVector& InDrag,
 
 		if ( !Actor->bLockLocation )
 		{
-			// find topmost selected group
+			if (GEditor->GetSelectedComponentCount() > 0)
+			{
+				for (FSelectionIterator It(GEditor->GetSelectedComponentIterator()); It; ++It)
+				{
+					USceneComponent* SceneComponent = CastChecked<USceneComponent>(*It);
+					if (SceneComponent)
+					{
+						ApplyDeltaToComponent(SceneComponent, InDrag, InRot, ModifiedScale);
+					}
+				}
+			}
+			else
+			{
 			AGroupActor* ParentGroup = AGroupActor::GetRootForActor(Actor, true, true);
 			if(ParentGroup && GEditor->bGroupingActive )
 			{
@@ -2949,6 +2965,7 @@ void FLevelEditorViewportClient::ApplyDeltaToActors(const FVector& InDrag,
 			}
 		}
 	}
+	}
 	AGroupActor::RemoveSubGroupsFromArray(ActorGroups);
 	for(int32 ActorGroupsIndex=0; ActorGroupsIndex<ActorGroups.Num(); ++ActorGroupsIndex)
 	{
@@ -2956,6 +2973,32 @@ void FLevelEditorViewportClient::ApplyDeltaToActors(const FVector& InDrag,
 	}
 }
 
+void FLevelEditorViewportClient::ApplyDeltaToComponent(USceneComponent* InComponent, const FVector& InDeltaDrag, const FRotator& InDeltaRot, const FVector& InDeltaScale)
+{
+	// If we are scaling, we need to change the scaling factor a bit to properly align to grid.
+	FVector ModifiedDeltaScale = InDeltaScale;
+
+	// we don't scale components when we only have a very small scale change
+	if (!InDeltaScale.IsNearlyZero())
+	{
+		if (!GEditor->UsePercentageBasedScaling())
+		{
+			ModifyScale(InComponent, ModifiedDeltaScale);
+		}
+	}
+	else
+	{
+		ModifiedDeltaScale = FVector::ZeroVector;
+	}
+
+	GEditor->ApplyDeltaToComponent(
+		InComponent,
+		true,
+		&InDeltaDrag,
+		&InDeltaRot,
+		&ModifiedDeltaScale,
+		GEditor->GetPivotLocation());
+}
 
 /** Helper function for ModifyScale - Convert the active Dragging Axis to per-axis flags */
 static void CheckActiveAxes( EAxisList::Type DraggingAxis, bool bActiveAxes[3] )
diff --git a/Engine/Source/Editor/UnrealEd/Public/LevelEditorViewport.h b/Engine/Source/Editor/UnrealEd/Public/LevelEditorViewport.h
--- a/Engine/Source/Editor/UnrealEd/Public/LevelEditorViewport.h
+++ b/Engine/Source/Editor/UnrealEd/Public/LevelEditorViewport.h
@@ -234,6 +234,7 @@ public:
 
 	void ApplyDeltaToActors( const FVector& InDrag, const FRotator& InRot, const FVector& InScale );
 	void ApplyDeltaToActor( AActor* InActor, const FVector& InDeltaDrag, const FRotator& InDeltaRot, const FVector& InDeltaScale );
+	void ApplyDeltaToComponent(USceneComponent* InComponent, const FVector& InDeltaDrag, const FRotator& InDeltaRot, const FVector& InDeltaScale);
 
 	virtual void SetIsSimulateInEditorViewport( bool bInIsSimulateInEditorViewport ) override;
 
diff --git a/Engine/Source/Runtime/Engine/Classes/GameFramework/Actor.h b/Engine/Source/Runtime/Engine/Classes/GameFramework/Actor.h
--- a/Engine/Source/Runtime/Engine/Classes/GameFramework/Actor.h
+++ b/Engine/Source/Runtime/Engine/Classes/GameFramework/Actor.h
@@ -54,7 +54,7 @@ DECLARE_CYCLE_STAT_EXTERN(TEXT("GetComponentsTime"),STAT_GetComponentsTime,STATG
  * @see https://docs.unrealengine.com/latest/INT/Programming/UnrealArchitecture/Actors/
  * @see UActorComponent
  */
-UCLASS(abstract, BlueprintType, Blueprintable, config=Engine)
+UCLASS(BlueprintType, Blueprintable, config=Engine)
 class ENGINE_API AActor : public UObject
 {
 	/**
@@ -1366,9 +1366,6 @@ public:
 	// Called before editor paste, true allow import
 	virtual bool ShouldImport(FString* ActorPropString, bool IsMovingLevel) { return true; }
 
-	// For UUnrealEdEngine::UpdatePropertyWindows()
-	virtual bool GetSelectedComponents(TArray<UObject*>& SelectedObjects) { return false; }
-
 	/** Called by InputKey when an unhandled key is pressed with a selected actor */
 	virtual void EditorKeyPressed(FKey Key, EInputEvent Event) {}
 
@@ -2177,6 +2174,13 @@ public:
 		}
 	}
 
+	// Get a direct reference to the Components array rather than a copy
+	// with the null pointers removed
+	const TArray<UActorComponent*>& GetComponents() const
+	{
+		return OwnedComponents;
+	}
+
 	/** Puts a component in to the OwnedComponents array of the Actor.
 	 *  The Component must be owned by the Actor or else it will assert
 	 *  In general this should not need to be called directly by anything other than UActorComponent functions
@@ -2329,4 +2333,3 @@ FORCEINLINE FVector AActor::GetSimpleCollisionCylinderExtent() const
 	virtual bool TeleportTo( const FVector& DestLocation, const FRotator& DestRotation, bool bIsATest, bool bNoCheck ) override { return Super::TeleportTo(DestLocation, DestRotation, bIsATest, bNoCheck); } \
 	virtual FVector GetVelocity() const override { return Super::GetVelocity(); } \
 	float GetDistanceTo(AActor* OtherActor) { return Super::GetDistanceTo(OtherActor); }
-
diff --git a/Engine/Source/Runtime/Landscape/Classes/LandscapeProxy.h b/Engine/Source/Runtime/Landscape/Classes/LandscapeProxy.h
--- a/Engine/Source/Runtime/Landscape/Classes/LandscapeProxy.h
+++ b/Engine/Source/Runtime/Landscape/Classes/LandscapeProxy.h
@@ -425,7 +425,6 @@ public:
 	virtual void PostEditMove(bool bFinished) override;
 	virtual bool ShouldImport(FString* ActorPropString, bool IsMovingLevel) override;
 	virtual bool ShouldExport() override;
-	virtual bool GetSelectedComponents(TArray<UObject*>& SelectedObjects) override;
 	// End AActor Interface
 #endif	//WITH_EDITOR
 
diff --git a/Engine/Source/Runtime/Landscape/Private/LandscapeEdit.cpp b/Engine/Source/Runtime/Landscape/Private/LandscapeEdit.cpp
--- a/Engine/Source/Runtime/Landscape/Private/LandscapeEdit.cpp
+++ b/Engine/Source/Runtime/Landscape/Private/LandscapeEdit.cpp
@@ -4861,21 +4861,6 @@ bool ALandscapeProxy::IsValidLandscapeActor(ALandscape* Landscape)
 	return false;
 }
 
-bool ALandscapeProxy::GetSelectedComponents(TArray<UObject*>& SelectedObjects)
-{
-	ULandscapeInfo* Info = GetLandscapeInfo(false);
-	if (Info && Info->bCurrentlyEditing)
-	{
-		auto SelectedComponents = Info->GetSelectedComponents();
-		for (auto It = SelectedComponents.CreateIterator(); It; ++It)
-		{
-			SelectedObjects.Add(*It);
-		}
-		return true;
-	}
-	return false;
-}
-
 struct FMobileLayerAllocation
 {
 	FWeightmapLayerAllocationInfo Allocation;
