Allow creation and duplication of Anonymous blueprints

diff --git a/Engine/Source/Editor/LevelEditor/Private/LevelEditor.cpp b/Engine/Source/Editor/LevelEditor/Private/LevelEditor.cpp
--- a/Engine/Source/Editor/LevelEditor/Private/LevelEditor.cpp
+++ b/Engine/Source/Editor/LevelEditor/Private/LevelEditor.cpp
@@ -586,6 +586,11 @@ void FLevelEditorModule::BindGlobalLevelEditorCommands()
 		FExecuteAction::CreateStatic( &FLevelEditorActionCallbacks::GoToDocsForActor_Clicked )
 		);
 
+	ActionList.MapAction(
+		Commands.AddScriptBehavior,
+		FExecuteAction::CreateStatic( &FLevelEditorActionCallbacks::AddScriptBehavior_Clicked )
+		);
+
 	ActionList.MapAction( 
 		FGenericCommands::Get().Duplicate, 
 		FExecuteAction::CreateStatic( &FLevelEditorActionCallbacks::ExecuteExecCommand, FString( TEXT("DUPLICATE") ) ),
diff --git a/Engine/Source/Editor/LevelEditor/Private/LevelEditorActions.cpp b/Engine/Source/Editor/LevelEditor/Private/LevelEditorActions.cpp
--- a/Engine/Source/Editor/LevelEditor/Private/LevelEditorActions.cpp
+++ b/Engine/Source/Editor/LevelEditor/Private/LevelEditorActions.cpp
@@ -945,6 +945,36 @@ void FLevelEditorActionCallbacks::GoToDocsForActor_Clicked()
 	}
 }
 
+void FLevelEditorActionCallbacks::AddScriptBehavior_Clicked()
+{
+	AActor* SelectedActor = GEditor->GetSelectedActors()->GetTop<AActor>();
+
+	if (SelectedActor)
+	{
+		const FName Name = *FString::Printf(TEXT("%s_BPClass"), *SelectedActor->GetName());
+
+		// Create and init a new Blueprint
+		UBlueprint* NewBP = FKismetEditorUtilities::CreateBlueprint(SelectedActor->GetClass(), SelectedActor->GetLevel(), Name, BPTYPE_Normal, UBlueprint::StaticClass(), UBlueprintGeneratedClass::StaticClass());
+		if (NewBP)
+		{
+			NewBP->ClearFlags(RF_Standalone);
+			ULevel::LevelDirtiedEvent.Broadcast();
+
+			FKismetEditorUtilities::CompileBlueprint(NewBP);
+
+			// the editor should be opened AFTER being added to the asset 
+			// registry (some systems could queue off of the asset registry 
+			// add event, and already having that blueprint open can be odd)
+			FAssetEditorManager::Get().OpenEditorForAsset(NewBP);
+
+			GEditor->ReplaceSelectedActors(GEditor->FindActorFactoryByClass(UActorFactoryBlueprint::StaticClass()), FAssetData(NewBP));
+
+			// Mark the package dirty...
+			SelectedActor->MarkPackageDirty();
+		}
+	}
+}
+
 void FLevelEditorActionCallbacks::FindInContentBrowser_Clicked()
 {
 	GEditor->SyncToContentBrowser();
@@ -2643,6 +2673,8 @@ void FLevelEditorCommands::RegisterCommands()
 	UI_COMMAND( GoToCodeForActor, "Go to C++ Code for Actor", "Opens a code editing IDE and navigates to the source file associated with the seleced actor", EUserInterfaceActionType::Button, FInputGesture() );
 	UI_COMMAND( GoToDocsForActor, "Go to Documentation for Actor", "Opens documentation for the Actor in the default web browser", EUserInterfaceActionType::Button, FInputGesture() );
 
+	UI_COMMAND( AddScriptBehavior, "Customize Scripting Behavior", "Click to customize scripting behavior of this Actor", EUserInterfaceActionType::Button, FInputGesture() );
+
 	UI_COMMAND( PasteHere, "Paste Here", "Pastes the actor at the click location", EUserInterfaceActionType::Button, FInputGesture() );
 
 	UI_COMMAND( SnapOriginToGrid, "Snap Origin to Grid", "Snaps the actor to the nearest grid location at its origin", EUserInterfaceActionType::Button, FInputGesture( EModifierKey::Control, EKeys::End ) );
diff --git a/Engine/Source/Editor/LevelEditor/Private/LevelEditorContextMenu.cpp b/Engine/Source/Editor/LevelEditor/Private/LevelEditorContextMenu.cpp
--- a/Engine/Source/Editor/LevelEditor/Private/LevelEditorContextMenu.cpp
+++ b/Engine/Source/Editor/LevelEditor/Private/LevelEditorContextMenu.cpp
@@ -279,6 +279,22 @@ void FLevelEditorContextMenu::FillMenu( FMenuBuilder& MenuBuilder, TWeakPtr<SLev
 			}
 			MenuBuilder.EndSection();
 		}
+
+		// If it is in the map then it is already anonymous
+		if (   GetDefault<UEditorExperimentalSettings>()->bInWorldBPEditing
+			&& SelectedActors.Num() == 1 
+			&& !FBlueprintEditorUtils::IsAnonymousBlueprintClass(SelectionInfo.SelectionClass) 
+			&& FKismetEditorUtilities::CanCreateBlueprintOfClass(SelectionInfo.SelectionClass) )
+		{
+			MenuBuilder.BeginSection( "Customize", LOCTEXT("CustomizeHeading", "Customize") );
+			{
+				MenuBuilder.AddMenuEntry( FLevelEditorCommands::Get().AddScriptBehavior,
+					NAME_None, 
+					LOCTEXT("AddScriptBehavior", "Customize Scripting Behavior"),
+					LOCTEXT("AddScriptBehavior_ToolTip", "Click to customize scripting behavior of this Actor"));
+			}
+			MenuBuilder.EndSection();
+		}
 	}
 
 	MenuBuilder.BeginSection("ActorSelectVisibilityLevels");
diff --git a/Engine/Source/Editor/LevelEditor/Public/LevelEditorActions.h b/Engine/Source/Editor/LevelEditor/Public/LevelEditorActions.h
--- a/Engine/Source/Editor/LevelEditor/Public/LevelEditorActions.h
+++ b/Engine/Source/Editor/LevelEditor/Public/LevelEditorActions.h
@@ -122,6 +122,9 @@ public:
 	/** Goes to the documentation for the selected actor's class. */
 	TSharedPtr< FUICommandInfo > GoToDocsForActor;
 
+	/** Customize the script behavior of an instance. */
+	TSharedPtr< FUICommandInfo > AddScriptBehavior;
+
 	/** Paste actor at click location*/
 	TSharedPtr< FUICommandInfo > PasteHere;
 
@@ -792,6 +795,9 @@ public:
 	/** Called when 'Go to Documentation for Actor' is clicked */
 	static void GoToDocsForActor_Clicked();
 
+	/** Called when 'Add Script Behavior' is clicked */
+	static void AddScriptBehavior_Clicked();
+
 	/**
 	 * Called when the LockActorMovement command is executed
 	 */
diff --git a/Engine/Source/Editor/UnrealEd/Private/EditorObject.cpp b/Engine/Source/Editor/UnrealEd/Private/EditorObject.cpp
--- a/Engine/Source/Editor/UnrealEd/Private/EditorObject.cpp
+++ b/Engine/Source/Editor/UnrealEd/Private/EditorObject.cpp
@@ -7,6 +7,7 @@
 #include "UnrealEd.h"
 #include "Factories.h"
 #include "BSPOps.h"
+#include "Kismet2/BlueprintEditorUtils.h"
 
 #include "Foliage/InstancedFoliageActor.h"
 #include "Foliage/FoliageType.h"
@@ -375,21 +376,26 @@ static const TCHAR* ImportProperties(
 				UObject* Archetype = NULL;
 				UObject* ComponentTemplate = NULL;
 
-				// if an archetype was specified in the Begin Object block, use that as the template for the ConstructObject call.
-				FString ArchetypeName;
-				if (FParse::Value(Str, TEXT("Archetype="), ArchetypeName))
+				// Since we are changing the class we can't use the Archetype,
+				// however that is fine since we will have been editing the CDO anyways
+				if (!FBlueprintEditorUtils::IsAnonymousBlueprintClass(SubobjectOuter->GetClass()))
 				{
-					// if given a name, break it up along the ' so separate the class from the name
-					FString ObjectClass;
-					FString ObjectPath;
-					if ( FPackageName::ParseExportTextPath(ArchetypeName, &ObjectClass, &ObjectPath) )
+					// if an archetype was specified in the Begin Object block, use that as the template for the ConstructObject call.
+					FString ArchetypeName;
+					if (FParse::Value(Str, TEXT("Archetype="), ArchetypeName))
 					{
-						// find the class
-						UClass* ArchetypeClass = (UClass*)StaticFindObject(UClass::StaticClass(), ANY_PACKAGE, *ObjectClass);
-						if (ArchetypeClass)
+						// if given a name, break it up along the ' so separate the class from the name
+						FString ObjectClass;
+						FString ObjectPath;
+						if ( FPackageName::ParseExportTextPath(ArchetypeName, &ObjectClass, &ObjectPath) )
 						{
-							// if we had the class, find the archetype
-							Archetype = StaticFindObject(ArchetypeClass, ANY_PACKAGE, *ObjectPath);
+							// find the class
+							UClass* ArchetypeClass = (UClass*)StaticFindObject(UClass::StaticClass(), ANY_PACKAGE, *ObjectClass);
+							if (ArchetypeClass)
+							{
+								// if we had the class, find the archetype
+								Archetype = StaticFindObject(ArchetypeClass, ANY_PACKAGE, *ObjectPath);
+							}
 						}
 					}
 				}
diff --git a/Engine/Source/Editor/UnrealEd/Private/Factories/EditorFactories.cpp b/Engine/Source/Editor/UnrealEd/Private/Factories/EditorFactories.cpp
--- a/Engine/Source/Editor/UnrealEd/Private/Factories/EditorFactories.cpp
+++ b/Engine/Source/Editor/UnrealEd/Private/Factories/EditorFactories.cpp
@@ -536,6 +536,24 @@ UObject* ULevelFactory::FactoryCreateText
 					}
 				}
 
+				// If we're pasting from a class that belongs to a map we need to duplicate the class and use that instead
+				if (FBlueprintEditorUtils::IsAnonymousBlueprintClass(TempClass))
+				{
+					UBlueprint* NewBP = DuplicateObject(CastChecked<UBlueprint>(TempClass->ClassGeneratedBy), World->GetCurrentLevel(), *FString::Printf(TEXT("%s_BPClass"), *ActorUniqueName.ToString()));
+					if (NewBP)
+					{
+						NewBP->ClearFlags(RF_Standalone);
+
+						FKismetEditorUtilities::CompileBlueprint(NewBP, false, true);
+
+						TempClass = NewBP->GeneratedClass;
+
+						// Since we changed the class we can't use an Archetype,
+						// however that is fine since we will have been editing the CDO anyways
+						Archetype = nullptr;
+					}
+				}
+
 				if (TempClass->IsChildOf(AWorldSettings::StaticClass()))
 				{
 					// if we see a WorldSettings, then we are importing an entire level, so if we
diff --git a/Engine/Source/Editor/UnrealEd/Private/Kismet2/BlueprintEditorUtils.cpp b/Engine/Source/Editor/UnrealEd/Private/Kismet2/BlueprintEditorUtils.cpp
--- a/Engine/Source/Editor/UnrealEd/Private/Kismet2/BlueprintEditorUtils.cpp
+++ b/Engine/Source/Editor/UnrealEd/Private/Kismet2/BlueprintEditorUtils.cpp
@@ -2435,6 +2435,11 @@ bool FBlueprintEditorUtils::IsLevelScriptBlueprint(const UBlueprint* Blueprint)
 	return (Blueprint->BlueprintType == BPTYPE_LevelScript);
 }
 
+bool FBlueprintEditorUtils::IsAnonymousBlueprintClass(const UClass* Class)
+{
+	return (Class->GetOutermost()->ContainsMap());
+}
+
 bool FBlueprintEditorUtils::SupportsConstructionScript(const UBlueprint* Blueprint)
 {
 	return(	!FBlueprintEditorUtils::IsInterfaceBlueprint(Blueprint) && 
diff --git a/Engine/Source/Editor/UnrealEd/Public/Kismet2/BlueprintEditorUtils.h b/Engine/Source/Editor/UnrealEd/Public/Kismet2/BlueprintEditorUtils.h
--- a/Engine/Source/Editor/UnrealEd/Public/Kismet2/BlueprintEditorUtils.h
+++ b/Engine/Source/Editor/UnrealEd/Public/Kismet2/BlueprintEditorUtils.h
@@ -367,10 +367,17 @@ public:
 	/**
 	 * Whether or not this blueprint is an interface, used only for defining functions to implement
 	 *
-	 * @return	Whether or not this is an interface blueprint
+	 * @return	Whether or not this is a level script blueprint
 	 */
 	static bool IsLevelScriptBlueprint(const UBlueprint* Blueprint);
 
+	/**
+	 * Whether or not this class represents a class generated by an anonymous actor class stored in a level 
+	 *
+	 * @return	Whether or not this is an anonymous blueprint
+	 */
+	static bool IsAnonymousBlueprintClass(const UClass* Class);
+
 	/** Do we support construction scripts */
 	static bool SupportsConstructionScript(const UBlueprint* Blueprint);
 
