Fix the level name appearing in viewports during PIE sessions

diff --git a/Engine/Source/Editor/LevelEditor/Private/SLevelViewport.cpp b/Engine/Source/Editor/LevelEditor/Private/SLevelViewport.cpp
--- a/Engine/Source/Editor/LevelEditor/Private/SLevelViewport.cpp
+++ b/Engine/Source/Editor/LevelEditor/Private/SLevelViewport.cpp
@@ -3152,35 +3152,39 @@ FText SLevelViewport::GetCurrentLevelText( bool bDrawOnlyLabel ) const
 	FText LabelName;
 	FText CurrentLevelName;
 
-	if( (&GetLevelViewportClient() == GCurrentLevelEditingViewportClient) && GetWorld() && GetWorld()->GetCurrentLevel() != NULL )
+	
+	if( ActiveViewport.IsValid() && (&GetLevelViewportClient() == GCurrentLevelEditingViewportClient) && GetWorld() && GetWorld()->GetCurrentLevel() != nullptr )
 	{
-		if( bDrawOnlyLabel )
+		if( ActiveViewport->GetPlayInEditorIsSimulate() || !ActiveViewport->GetClient()->GetWorld()->IsGameWorld() )
 		{
-			LabelName = LOCTEXT("CurrentLevelLabel", "Level:");
-		}
-		else
-		{
-			// Get the level name (without the number at the end)
-			FText ActualLevelName = FText::FromString( FPackageName::GetShortFName( GetWorld()->GetCurrentLevel()->GetOutermost()->GetFName() ).GetPlainNameString() );
-
-			if( GetWorld()->GetCurrentLevel() == GetWorld()->PersistentLevel )
+			if(bDrawOnlyLabel)
 			{
-				FFormatNamedArguments Args;
-				Args.Add( TEXT("ActualLevelName"), ActualLevelName );
-				CurrentLevelName = FText::Format( LOCTEXT("LevelName", "{0} (Persistent)"), ActualLevelName );
+				LabelName = LOCTEXT("CurrentLevelLabel", "Level:");
 			}
 			else
 			{
-				CurrentLevelName = ActualLevelName;
+				// Get the level name (without the number at the end)
+				FText ActualLevelName = FText::FromString(FPackageName::GetShortFName(GetWorld()->GetCurrentLevel()->GetOutermost()->GetFName()).GetPlainNameString());
+
+				if(GetWorld()->GetCurrentLevel() == GetWorld()->PersistentLevel)
+				{
+					FFormatNamedArguments Args;
+					Args.Add(TEXT("ActualLevelName"), ActualLevelName);
+					CurrentLevelName = FText::Format(LOCTEXT("LevelName", "{0} (Persistent)"), ActualLevelName);
+				}
+				else
+				{
+					CurrentLevelName = ActualLevelName;
+				}
+			}
+
+			if(bDrawOnlyLabel)
+			{
+				return LabelName;
 			}
 		}
 	}
 
-	if( bDrawOnlyLabel )
-	{
-		return LabelName;
-	}
-
 	return CurrentLevelName;
 }
 
