First-pass on in-world component selection
- Must enable experimental In-World Component Editing
- To select a component in the world, click on the component when the owning actor is already selected

diff --git a/Engine/Source/Editor/LevelEditor/Private/ComponentsTree.cpp b/Engine/Source/Editor/LevelEditor/Private/ComponentsTree.cpp
--- a/Engine/Source/Editor/LevelEditor/Private/ComponentsTree.cpp
+++ b/Engine/Source/Editor/LevelEditor/Private/ComponentsTree.cpp
@@ -98,6 +98,9 @@ void SComponentsTree::Construct(const FArguments& InArgs, TSharedPtr<class IDeta
 	FDetailsViewArgs DetailsViewArgs( /*bUpdateFromSelection=*/ false, /*bLockable=*/ false, /*bAllowSearch=*/ true, /*bObjectsUseNameArea=*/ true, /*bHideSelectionTip=*/ true );
 	DetailsViewArgs.bHideActorNameArea = true;
 
+	bSelectionGuard = false;
+	USelection::SelectionChangedEvent.AddRaw(this, &SComponentsTree::OnEditorSelectionChanged);
+
 	PropertyView = InPropertyView;
 
 	this->ChildSlot
@@ -131,10 +134,49 @@ void SComponentsTree::Construct(const FArguments& InArgs, TSharedPtr<class IDeta
 	UpdateTree();
 }
 
+SComponentsTree::~SComponentsTree()
+{
+	USelection::SelectionChangedEvent.RemoveAll(this);
+}
+
+void SComponentsTree::OnEditorSelectionChanged(UObject* Object)
+{
+	if (!bSelectionGuard)
+	{
+		TArray<UObject*> Objects;
+		auto Selection = Cast<USelection>(Object);
+		if (Selection == GEditor->GetSelectedComponents() && GEditor->GetSelectedComponentCount() > 0)
+		{
+			// Enable the selection guard to prevent OnTreeSelectionChanged() from altering the editor's component selection
+			TGuardValue<bool> SelectionGuard(bSelectionGuard, true);
+
+			// Update the tree selection to match the level editor component selection
+			TreeWidget->ClearSelection();
+			for (FSelectionIterator It(GEditor->GetSelectedComponentIterator()); It; ++It)
+			{
+				UActorComponent* Component = CastChecked<UActorComponent>(*It);
+				if (Component)
+				{
+					TreeWidget->SetItemSelection(GetNodeFromActorComponent(Component), true);
+					Objects.Add(Component);
+				}
+			}
+
+			PropertyView->SetObjects(Objects, true);
+		}
+	}
+}
+
 void SComponentsTree::SetObjects(const TArray<UObject*>& InObjects)
 {
 	if (InObjects.Num() == 1)
 	{
+		auto NewActor = Cast<AActor>(InObjects[0]);
+		if (NewActor == Actor)
+		{
+			return;
+		}
+
 		Actor = Cast<AActor>(InObjects[0]);
 	}
 	else
@@ -164,26 +206,33 @@ void SComponentsTree::OnGetChildrenForTree( FComponentTreeNodePtrType InNodePtr,
 
 void SComponentsTree::OnTreeSelectionChanged(FComponentTreeNodePtrType InSelectedNodePtr, ESelectInfo::Type SelectInfo)
 {
-	TArray<UObject*> Objects;
+	if (!bSelectionGuard)
+	{
+		// Enable the selection guard to prevent OnEditorSelectionChanged() from altering the contents of TreeWidget
+		TGuardValue<bool> SelectionGuard(bSelectionGuard, true);
 
-	USelection* SelectedComponents = GEditor->GetSelectedComponents();
-	SelectedComponents->BeginBatchSelectOperation();
-	SelectedComponents->DeselectAll();
+		// Update the editor's component selection
+		TArray<UObject*> Objects;
+		USelection* SelectedComponents = GEditor->GetSelectedComponents();
+		SelectedComponents->BeginBatchSelectOperation();
+		SelectedComponents->DeselectAll();
 
-	for (FComponentTreeNodePtrType SelectedNode : TreeWidget->GetSelectedItems())
-	{
-		if (SelectedNode.IsValid() && SelectedNode->Component.IsValid())
+		for (FComponentTreeNodePtrType SelectedNode : TreeWidget->GetSelectedItems())
 		{
-			UActorComponent* Component = SelectedNode->Component.Get();
-			Objects.Add(Component);
-			SelectedComponents->Select(Component);
+			if (SelectedNode.IsValid() && SelectedNode->Component.IsValid())
+			{
+				UActorComponent* Component = SelectedNode->Component.Get();
+				Objects.Add(Component);
+				SelectedComponents->Select(Component);
+			}
 		}
-	}
-	
-	SelectedComponents->EndBatchSelectOperation();
-	PropertyView->SetObjects(Objects, false);
 
-	GEditor->RedrawLevelEditingViewports();
+		PropertyView->SetObjects(Objects, false);
+
+		SelectedComponents->EndBatchSelectOperation();
+		GUnrealEd->UpdatePivotLocationForSelection(true);
+		GEditor->RedrawLevelEditingViewports();
+	}
 }
 
 void SComponentsTree::OnSelectedCompClass(TSubclassOf<UActorComponent> CompClass)
@@ -256,3 +305,47 @@ void SComponentsTree::UpdateTree()
 	}
 	TreeWidget->RequestTreeRefresh();
 }
+
+FComponentTreeNodePtrType SComponentsTree::GetNodeFromActorComponent(UActorComponent* ActorComponent, bool bIncludeAttachedComponents)
+{
+	FComponentTreeNodePtrType ReturnNode = nullptr;
+
+	if (ActorComponent)
+	{
+		for (auto RootNode : RootNodes)
+		{
+			ReturnNode = FindTreeNodeRecursive(ActorComponent, RootNode);
+
+			if (ReturnNode.IsValid())
+			{
+				break;
+			}
+		}
+	}
+	
+	return ReturnNode;
+}
+
+FComponentTreeNodePtrType SComponentsTree::FindTreeNodeRecursive(const UActorComponent* InComponent, const FComponentTreeNodePtrType& Node) const
+{
+	FComponentTreeNodePtrType LocatedNode = nullptr;
+
+	if (Node->Component.IsValid() && Node->Component.Get() == InComponent)
+	{
+		LocatedNode = Node;
+	}
+	else
+	{
+		// Recursively check NodePtr and all of its children until a matching node is found
+		for (auto ChildNode : Node->ChildNodes)
+		{
+			LocatedNode = FindTreeNodeRecursive(InComponent, ChildNode);
+			if (LocatedNode.IsValid())
+			{
+				break;
+			}
+		}
+	}
+
+	return LocatedNode;
+}
diff --git a/Engine/Source/Editor/LevelEditor/Private/ComponentsTree.h b/Engine/Source/Editor/LevelEditor/Private/ComponentsTree.h
--- a/Engine/Source/Editor/LevelEditor/Private/ComponentsTree.h
+++ b/Engine/Source/Editor/LevelEditor/Private/ComponentsTree.h
@@ -46,6 +46,8 @@ public:
 
 	void Construct( const FArguments& InArgs, TSharedPtr<class IDetailsView> InPropertyView );
 
+	~SComponentsTree();
+
 	void UpdateTree();
 
 	void SetObjects(const TArray<UObject*>& InObjects);
@@ -58,6 +60,9 @@ public:
 	void OnSelectedCompClass(TSubclassOf<UActorComponent> CompClass);
 
 private:
+	void OnEditorSelectionChanged(UObject* Object);
+	FComponentTreeNodePtrType GetNodeFromActorComponent(UActorComponent* ActorComponent, bool bIncludeAttachedComponents = true);
+	FComponentTreeNodePtrType FindTreeNodeRecursive(const UActorComponent* InComponent, const FComponentTreeNodePtrType& NodePtr) const;
 
 	TWeakObjectPtr<AActor>				Actor;
 
@@ -66,4 +71,6 @@ private:
 	TSharedPtr<SComponentTreeType>		TreeWidget;
 
 	TSharedPtr<class IDetailsView>		PropertyView;
+
+	bool bSelectionGuard;
 };
\ No newline at end of file
diff --git a/Engine/Source/Editor/MeshPaint/Public/MeshPaintModule.h b/Engine/Source/Editor/MeshPaint/Public/MeshPaintModule.h
--- a/Engine/Source/Editor/MeshPaint/Public/MeshPaintModule.h
+++ b/Engine/Source/Editor/MeshPaint/Public/MeshPaintModule.h
@@ -2,6 +2,8 @@
 
 #pragma once
 
+#include "ModuleInterface.h"
+
 /**
  * Interface for a class to provide mesh painting support for a subclass of UMeshComponent
  */
diff --git a/Engine/Source/Editor/UnrealEd/Classes/Editor/EditorEngine.h b/Engine/Source/Editor/UnrealEd/Classes/Editor/EditorEngine.h
--- a/Engine/Source/Editor/UnrealEd/Classes/Editor/EditorEngine.h
+++ b/Engine/Source/Editor/UnrealEd/Classes/Editor/EditorEngine.h
@@ -937,7 +937,7 @@ public:
 	 */
 	virtual void TakeHighResScreenShots(){}
 
-	virtual void NoteSelectionChange() { check(0); }
+	virtual void NoteSelectionChange(bool bComponentSelectionChanged = false) { check(0); }
 
 	/**
 	 * Adds an actor to the world at the specified location.
@@ -1253,6 +1253,7 @@ public:
 	virtual void SelectActor(AActor* Actor, bool bInSelected, bool bNotify, bool bSelectEvenIfHidden=false) {}
 	virtual bool CanSelectActor(AActor* Actor, bool bInSelected, bool bSelectEvenIfHidden=false, bool bWarnIfLevelLocked=false) const { return true; }
 	virtual void SelectGroup(class AGroupActor* InGroupActor, bool bForceSelection=false, bool bInSelected=true, bool bNotify=true) {}
+	virtual void SelectComponent(class UActorComponent* Component, bool bInSelected, bool bNotify, bool bSelectEvenIfHidden = false) {}
 
 	/**
 	 * Replaces the components in ActorsToReplace with an primitive component in Replacement
diff --git a/Engine/Source/Editor/UnrealEd/Classes/Editor/UnrealEdEngine.h b/Engine/Source/Editor/UnrealEd/Classes/Editor/UnrealEdEngine.h
--- a/Engine/Source/Editor/UnrealEd/Classes/Editor/UnrealEdEngine.h
+++ b/Engine/Source/Editor/UnrealEd/Classes/Editor/UnrealEdEngine.h
@@ -178,9 +178,10 @@ class UNREALED_API UUnrealEdEngine : public UEditorEngine, public FNotifyHook
 	virtual void SelectActor(AActor* Actor, bool InSelected, bool bNotify, bool bSelectEvenIfHidden=false) override;
 	virtual bool CanSelectActor(AActor* Actor, bool InSelected, bool bSelectEvenIfHidden=false, bool bWarnIfLevelLocked=false) const override;
 	virtual void SelectGroup(AGroupActor* InGroupActor, bool bForceSelection=false, bool bInSelected=true, bool bNotify=true) override;
+	virtual void SelectComponent(class UActorComponent* Component, bool bInSelected, bool bNotify, bool bSelectEvenIfHidden = false) override;
 	virtual void SelectBSPSurf(UModel* InModel, int32 iSurf, bool bSelected, bool bNoteSelectionChange) override;
 	virtual void SelectNone(bool bNoteSelectionChange, bool bDeselectBSPSurfs, bool WarnAboutManyActors=true) override;
-	virtual void NoteSelectionChange() override;
+	virtual void NoteSelectionChange(bool bComponentSelectionChanged = false) override;
 	virtual void NoteActorMovement() override;
 	virtual void FinishAllSnaps() override;
 	virtual void Cleanse( bool ClearSelection, bool Redraw, const FText& Reason ) override;
diff --git a/Engine/Source/Editor/UnrealEd/Private/EditorSelectUtils.cpp b/Engine/Source/Editor/UnrealEd/Private/EditorSelectUtils.cpp
--- a/Engine/Source/Editor/UnrealEd/Private/EditorSelectUtils.cpp
+++ b/Engine/Source/Editor/UnrealEd/Private/EditorSelectUtils.cpp
@@ -271,8 +271,8 @@ void UUnrealEdEngine::SetActorSelectionFlags (AActor* InActor)
 void UUnrealEdEngine::UpdatePivotLocationForSelection( bool bOnChange )
 {
 	// Pick a new common pivot, or not.
-	AActor* SingleActor=NULL;
-	USceneComponent* SingleComponent = NULL;
+	AActor* SingleActor = nullptr;
+	USceneComponent* SingleComponent = nullptr;
 
 	if (GetSelectedComponentCount() > 0)
 	{
@@ -348,7 +348,7 @@ void UUnrealEdEngine::UpdatePivotLocationForSelection( bool bOnChange )
 
 
 
-void UUnrealEdEngine::NoteSelectionChange()
+void UUnrealEdEngine::NoteSelectionChange(bool bComponentSelectionChanged)
 {
 	// The selection changed, so make sure the pivot (widget) is located in the right place
 	UpdatePivotLocationForSelection( true );
@@ -363,16 +363,20 @@ void UUnrealEdEngine::NoteSelectionChange()
 		ActiveModes[ModeIndex]->ActorSelectionChangeNotify();
 	}
 
-	USelection* Selection = GetSelectedActors();
+	USelection* Selection = bComponentSelectionChanged ? GetSelectedComponents() : GetSelectedActors();
 	USelection::SelectionChangedEvent.Broadcast(Selection);
 
-	//whenever selection changes, recompute whether the selection contains a locked actor
-	bCheckForLockActors = true;
+	
+	if (!bComponentSelectionChanged)
+	{
+		//whenever selection changes, recompute whether the selection contains a locked actor
+		bCheckForLockActors = true;
 
-	//whenever selection changes, recompute whether the selection contains a world info actor
-	bCheckForWorldSettingsActors = true;
+		//whenever selection changes, recompute whether the selection contains a world info actor
+		bCheckForWorldSettingsActors = true;
 
-	UpdateFloatingPropertyWindows();
+		UpdateFloatingPropertyWindows();
+	}
 
 	RedrawLevelEditingViewports();
 }
@@ -514,25 +518,25 @@ void UUnrealEdEngine::SelectActor(AActor* Actor, bool bInSelected, bool bNotify,
 			}
 		}
 
-			if(GEditor->bGroupingActive)
-			{
+		if (GEditor->bGroupingActive)
+		{
 			// if this actor is a group, do a group select/deselect
-				AGroupActor* SelectedGroupActor = Cast<AGroupActor>(Actor);
-				if( SelectedGroupActor )
-				{
+			AGroupActor* SelectedGroupActor = Cast<AGroupActor>(Actor);
+			if (SelectedGroupActor)
+			{
 				SelectGroup(SelectedGroupActor, true, bInSelected, bNotify);
-				}
-				else
-				{
+			}
+			else
+			{
 				// Select/Deselect this actor's entire group, starting from the top locked group.
 				// If none is found, just use the actor.
-					AGroupActor* ActorLockedRootGroup = AGroupActor::GetRootForActor(Actor, true);
-					if( ActorLockedRootGroup )
-					{
+				AGroupActor* ActorLockedRootGroup = AGroupActor::GetRootForActor(Actor, true);
+				if (ActorLockedRootGroup)
+				{
 					SelectGroup(ActorLockedRootGroup, false, bInSelected, bNotify);
-					}
 				}
 			}
+		}
 
 		// Don't do any work if the actor's selection state is already the selected state.
 		const bool bActorSelected = Actor->IsSelected();
@@ -581,6 +585,52 @@ void UUnrealEdEngine::SelectActor(AActor* Actor, bool bInSelected, bool bNotify,
 	}
 }
 
+void UUnrealEdEngine::SelectComponent(UActorComponent* Component, bool bInSelected, bool bNotify, bool bSelectEvenIfHidden)
+{
+	// Don't do any work if the component is already selected
+	const bool bComponentSelected = Component->IsSelected();
+	if (( bComponentSelected && !bInSelected ) || ( !bComponentSelected && bInSelected ))
+	{
+		if (bInSelected)
+		{
+			UE_LOG(LogEditorSelectUtils, Verbose, TEXT("Selected Component: %s"), *Component->GetClass()->GetName());
+		}
+		else
+		{
+			UE_LOG(LogEditorSelectUtils, Verbose, TEXT("Deselected Component: %s"), *Component->GetClass()->GetName());
+		}
+
+		GetSelectedComponents()->Select(Component, bInSelected);
+
+		// Update the selection visualization
+		AActor* ComponentOwner = Component->GetOwner();
+		if (ComponentOwner != nullptr)
+		{
+			//@todo dhertzka - need selection override delegates for the primitive components
+			TInlineComponentArray<UPrimitiveComponent*> PrimitiveComponents;
+			ComponentOwner->GetComponents(PrimitiveComponents);
+
+			for (int32 Idx = 0; Idx < PrimitiveComponents.Num(); ++Idx)
+			{
+				PrimitiveComponents[Idx]->PushSelectionToProxy();
+			}
+		}
+
+		if (bNotify)
+		{
+			NoteSelectionChange(true);
+		}
+	}
+	else
+	{
+		if (bNotify)
+		{
+			//reset the property windows.  In case something has changed since previous selection
+			UpdateFloatingPropertyWindows();
+		}
+	}
+}
+
 void UUnrealEdEngine::SelectBSPSurf(UModel* InModel, int32 iSurf, bool bSelected, bool bNoteSelectionChange)
 {
 	if( GEdSelectionLock )
diff --git a/Engine/Source/Editor/UnrealEd/Private/LevelEditorViewport.cpp b/Engine/Source/Editor/UnrealEd/Private/LevelEditorViewport.cpp
--- a/Engine/Source/Editor/UnrealEd/Private/LevelEditorViewport.cpp
+++ b/Engine/Source/Editor/UnrealEd/Private/LevelEditorViewport.cpp
@@ -1761,7 +1761,19 @@ void FLevelEditorViewportClient::ProcessClick(FSceneView& View, HHitProxy* HitPr
 		}
 		else if (HitProxy->IsA(HActor::StaticGetType()))
 		{
-			ClickHandlers::ClickActor(this,((HActor*)HitProxy)->Actor,Click,true);
+			auto ActorHitProxy = (HActor*)HitProxy;
+
+			auto SelectedActors = GEditor->GetSelectedActors();
+			const bool bSelectComponents = SelectedActors->IsSelected(ActorHitProxy->Actor) && GEditor->GetSelectedActorCount() == 1;
+			if (bSelectComponents && GetDefault<UEditorExperimentalSettings>()->bInWorldBPEditing)
+			{
+				// The hit actor is already selected and is the only selected actor, so select the hit component
+				ClickHandlers::ClickComponent(this, ActorHitProxy, Click);
+			}
+			else
+			{
+				ClickHandlers::ClickActor(this, ActorHitProxy->Actor, Click, true);
+			}
 		}
 		else if (HitProxy->IsA(HInstancedStaticMeshInstance::StaticGetType()))
 		{
diff --git a/Engine/Source/Editor/UnrealEd/Private/LevelViewportClickHandlers.cpp b/Engine/Source/Editor/UnrealEd/Private/LevelViewportClickHandlers.cpp
--- a/Engine/Source/Editor/UnrealEd/Private/LevelViewportClickHandlers.cpp
+++ b/Engine/Source/Editor/UnrealEd/Private/LevelViewportClickHandlers.cpp
@@ -264,6 +264,44 @@ namespace ClickHandlers
 		return false;
 	}
 
+	bool ClickComponent(FLevelEditorViewportClient* ViewportClient, HActor* ActorHitProxy, const FViewportClick& Click)
+	{
+		bool bComponentClicked = false;
+		const FScopedTransaction Transaction(NSLOCTEXT("UnrealEd", "ClickingOnComponents", "Clicking on Components"));
+
+		UActorComponent* SelectedComponent = const_cast<UPrimitiveComponent*>(ActorHitProxy->PrimComponent); 
+		GEditor->GetSelectedComponents()->Modify();
+
+		if (Click.IsControlDown())
+		{
+			const bool bSelect = !SelectedComponent->IsSelected();
+			if (bSelect)
+			{
+				UE_LOG(LogEditorViewport, Log, TEXT("Clicking on Component (CTRL LMB): %s (%s)"), *SelectedComponent->GetClass()->GetName(), *SelectedComponent->GetName());
+			}
+			GEditor->SelectComponent(SelectedComponent, bSelect, true, true);
+			bComponentClicked = true;
+		}
+		else if (Click.IsShiftDown())
+		{
+			if (!SelectedComponent->IsSelected())
+			{
+				const bool bSelect = true;
+				GEditor->SelectComponent(SelectedComponent, bSelect, true, true);
+			}
+			bComponentClicked = true;
+		}
+		else
+		{
+			GEditor->GetSelectedComponents()->DeselectAll();
+			UE_LOG(LogEditorViewport, Log, TEXT("Clicking on Component (LMB): %s (%s)"), *SelectedComponent->GetClass()->GetName(), *SelectedComponent->GetName());
+			GEditor->SelectComponent(SelectedComponent, true, true, true);
+			bComponentClicked = true;
+		}
+
+		return bComponentClicked;
+	}
+
 	void ClickBrushVertex(FLevelEditorViewportClient* ViewportClient,ABrush* InBrush,FVector* InVertex,const FViewportClick& Click)
 	{
 		// Pivot snapping
diff --git a/Engine/Source/Editor/UnrealEd/Private/LevelViewportClickHandlers.h b/Engine/Source/Editor/UnrealEd/Private/LevelViewportClickHandlers.h
--- a/Engine/Source/Editor/UnrealEd/Private/LevelViewportClickHandlers.h
+++ b/Engine/Source/Editor/UnrealEd/Private/LevelViewportClickHandlers.h
@@ -36,6 +36,8 @@ namespace ClickHandlers
 {
 	bool ClickActor(FLevelEditorViewportClient* ViewportClient,AActor* Actor,const FViewportClick& Click,bool bAllowSelectionChange);
 
+	bool ClickComponent(FLevelEditorViewportClient* ViewportClient, HActor* ActorHitProxy, const FViewportClick& Click);
+
 	void ClickBrushVertex(FLevelEditorViewportClient* ViewportClient,ABrush* InBrush,FVector* InVertex,const FViewportClick& Click);
 
 	void ClickStaticMeshVertex(FLevelEditorViewportClient* ViewportClient,AActor* InActor,FVector& InVertex,const FViewportClick& Click);
diff --git a/Engine/Source/Editor/UnrealEd/Public/LevelEditorViewport.h b/Engine/Source/Editor/UnrealEd/Public/LevelEditorViewport.h
--- a/Engine/Source/Editor/UnrealEd/Public/LevelEditorViewport.h
+++ b/Engine/Source/Editor/UnrealEd/Public/LevelEditorViewport.h
@@ -153,6 +153,8 @@ public:
 	virtual FLinearColor GetBackgroundColor() const override;
 	virtual int32 GetCameraSpeedSetting() const override;
 	virtual void SetCameraSpeedSetting(int32 SpeedSetting) override;
+	virtual void ProcessClick(FSceneView& View, HHitProxy* HitProxy, FKey Key, EInputEvent Event, uint32 HitX, uint32 HitY) override;
+	virtual UWorld* GetWorld() const override;
 
 	virtual bool OverrideHighResScreenshotCaptureRegion(FIntRect& OutCaptureRegion) override;
 
@@ -294,10 +296,6 @@ public:
 	 */
 	void SetAllSpriteCategoryVisibility( bool bVisible );
 
-	/** FEditorViewportClient Interface*/
-	virtual void ProcessClick(FSceneView& View, HHitProxy* HitProxy, FKey Key, EInputEvent Event, uint32 HitX, uint32 HitY);
-	virtual UWorld* GetWorld() const override;
-
 	void SetReferenceToWorldContext(FWorldContext& WorldContext);
 
 	void RemoveReferenceToWorldContext(FWorldContext& WorldContext);
