Back out 2401916 for SWebBrowser

diff --git a/Engine/Source/Editor/LevelEditor/Private/SLevelViewportControlsPopup.cpp b/Engine/Source/Editor/LevelEditor/Private/SLevelViewportControlsPopup.cpp
--- a/Engine/Source/Editor/LevelEditor/Private/SLevelViewportControlsPopup.cpp
+++ b/Engine/Source/Editor/LevelEditor/Private/SLevelViewportControlsPopup.cpp
@@ -86,7 +86,7 @@ TSharedRef<SWidget> SLevelViewportControlsPopup::OnGetMenuContent()
 				SNew(SWebBrowser)
 				.InitialURL(PopupPath)
 				.ShowControls(false)
-				.SupportsTransparency(EWindowTransparency::PerWindow)
+				.SupportsTransparency(true)
 				.ViewportSize(PopupSize + 1)
 			];
 	}
