Copy instance properties in to the CDO when creating an anonymous blueprint

diff --git a/Engine/Source/Editor/LevelEditor/Private/LevelEditorActions.cpp b/Engine/Source/Editor/LevelEditor/Private/LevelEditorActions.cpp
--- a/Engine/Source/Editor/LevelEditor/Private/LevelEditorActions.cpp
+++ b/Engine/Source/Editor/LevelEditor/Private/LevelEditorActions.cpp
@@ -963,6 +963,8 @@ void FLevelEditorActionCallbacks::AddScriptBehavior_Clicked()
 			NewBP->ClearFlags(RF_Standalone);
 			ULevel::LevelDirtiedEvent.Broadcast();
 
+			UEngine::CopyPropertiesForUnrelatedObjects(SelectedActor, NewBP->GeneratedClass->GetDefaultObject<AActor>());
+
 			FKismetEditorUtilities::CompileBlueprint(NewBP);
 
 			// the editor should be opened AFTER being added to the asset 
