In-world component selection improvements:
- Right-clicking a selected actor will open the context menu, not select the component
- Clicking the selected actor will not select a component if the click also focused the level viewport
- Right-clicking a selected component opens a 1-option context menu to select the owning actor

diff --git a/Engine/Source/Editor/LevelEditor/Private/ComponentsTree.cpp b/Engine/Source/Editor/LevelEditor/Private/ComponentsTree.cpp
--- a/Engine/Source/Editor/LevelEditor/Private/ComponentsTree.cpp
+++ b/Engine/Source/Editor/LevelEditor/Private/ComponentsTree.cpp
@@ -145,7 +145,7 @@ void SComponentsTree::OnEditorSelectionChanged(UObject* Object)
 	{
 		TArray<UObject*> Objects;
 		auto Selection = Cast<USelection>(Object);
-		if (Selection == GEditor->GetSelectedComponents() && GEditor->GetSelectedComponentCount() > 0)
+		if ( Selection == GEditor->GetSelectedComponents() || Selection == GEditor->GetSelectedActors() )
 		{
 			// Enable the selection guard to prevent OnTreeSelectionChanged() from altering the editor's component selection
 			TGuardValue<bool> SelectionGuard(bSelectionGuard, true);
@@ -162,7 +162,7 @@ void SComponentsTree::OnEditorSelectionChanged(UObject* Object)
 				}
 			}
 
-			PropertyView->SetObjects(Objects, true);
+			PropertyView->SetObjects(Objects);
 		}
 	}
 }
@@ -181,7 +181,7 @@ void SComponentsTree::SetObjects(const TArray<UObject*>& InObjects)
 	}
 	else
 	{
-		Actor = NULL;
+		Actor = nullptr;
 	}
 	Visibility = (Actor.IsValid() ? EVisibility::Visible : EVisibility::Hidden);
 	UpdateTree();
@@ -206,7 +206,7 @@ void SComponentsTree::OnGetChildrenForTree( FComponentTreeNodePtrType InNodePtr,
 
 void SComponentsTree::OnTreeSelectionChanged(FComponentTreeNodePtrType InSelectedNodePtr, ESelectInfo::Type SelectInfo)
 {
-	if (!bSelectionGuard)
+	if (!bSelectionGuard && Actor != nullptr)
 	{
 		// Enable the selection guard to prevent OnEditorSelectionChanged() from altering the contents of TreeWidget
 		TGuardValue<bool> SelectionGuard(bSelectionGuard, true);
diff --git a/Engine/Source/Editor/LevelEditor/Private/LevelEditor.cpp b/Engine/Source/Editor/LevelEditor/Private/LevelEditor.cpp
--- a/Engine/Source/Editor/LevelEditor/Private/LevelEditor.cpp
+++ b/Engine/Source/Editor/LevelEditor/Private/LevelEditor.cpp
@@ -947,6 +947,12 @@ void FLevelEditorModule::BindGlobalLevelEditorCommands()
 		FExecuteAction::CreateStatic( &FLevelEditorActionCallbacks::OnSelectAllActorsOfClass, (bool)true )
 		);
 
+	ActionList.MapAction(
+		Commands.SelectComponentOwnerActor,
+		FExecuteAction::CreateStatic( &FLevelEditorActionCallbacks::OnSelectComponentOwnerActor ),
+		FCanExecuteAction::CreateStatic( &FLevelEditorActionCallbacks::CanSelectComponentOwnerActor )
+		);
+
 	ActionList.MapAction(
 		Commands.SelectRelevantLights,
 		FExecuteAction::CreateStatic( &FLevelEditorActionCallbacks::ExecuteExecCommand, FString( TEXT("ACTOR SELECT RELEVANTLIGHTS") ) )
diff --git a/Engine/Source/Editor/LevelEditor/Private/LevelEditorActions.cpp b/Engine/Source/Editor/LevelEditor/Private/LevelEditorActions.cpp
--- a/Engine/Source/Editor/LevelEditor/Private/LevelEditorActions.cpp
+++ b/Engine/Source/Editor/LevelEditor/Private/LevelEditorActions.cpp
@@ -1438,6 +1438,20 @@ void FLevelEditorActionCallbacks::OnSelectAllActorsOfClass( bool bArchetype )
 	GEditor->SelectAllActorsWithClass( bArchetype );
 }
 
+void FLevelEditorActionCallbacks::OnSelectComponentOwnerActor()
+{
+	auto ComponentOwner = Cast<AActor>(*GEditor->GetSelectedActorIterator());
+	check(ComponentOwner);
+
+	GEditor->SelectNone(false, true, false);
+	GEditor->SelectActor(ComponentOwner, true, true, true);
+}
+
+bool FLevelEditorActionCallbacks::CanSelectComponentOwnerActor()
+{
+	return GEditor->GetSelectedComponentCount() > 0;
+}
+
 void FLevelEditorActionCallbacks::OnSelectAllActorsControlledByMatinee()
 {
 	GEditor->SelectAllActorsControlledByMatinee();
@@ -2763,6 +2777,7 @@ void FLevelEditorCommands::RegisterCommands()
 
 	UI_COMMAND( SelectAllActorsOfSameClass, "Select All Actors of Same Class", "Selects all the actors that have the same class", EUserInterfaceActionType::Button, FInputGesture(EModifierKey::Shift|EModifierKey::Control, EKeys::A) );
 	UI_COMMAND( SelectAllActorsOfSameClassWithArchetype, "Select All Actors with Same Archetype", "Selects all the actors of the same class that have the same archetype", EUserInterfaceActionType::Button, FInputGesture() );
+	UI_COMMAND( SelectComponentOwnerActor, "Select Component Owner", "Select the actor that owns the currently selected component(s)", EUserInterfaceActionType::Button, FInputGesture() );
 	UI_COMMAND( SelectRelevantLights, "Select Relevant Lights", "Select all lights relevant to the current selection", EUserInterfaceActionType::Button, FInputGesture() ); 
 	UI_COMMAND( SelectStaticMeshesOfSameClass, "Select All Using Selected Static Meshes (Selected Actor Types)", "Selects all actors with the same static mesh and actor class as the selection", EUserInterfaceActionType::Button, FInputGesture() ); 
 	UI_COMMAND( SelectStaticMeshesAllClasses, "Select All Using Selected Static Meshes (All Actor Types)", "Selects all actors with the same static mesh as the selection", EUserInterfaceActionType::Button, FInputGesture( EModifierKey::Shift, EKeys::E ) ); 
diff --git a/Engine/Source/Editor/LevelEditor/Private/LevelEditorContextMenu.cpp b/Engine/Source/Editor/LevelEditor/Private/LevelEditorContextMenu.cpp
--- a/Engine/Source/Editor/LevelEditor/Private/LevelEditorContextMenu.cpp
+++ b/Engine/Source/Editor/LevelEditor/Private/LevelEditorContextMenu.cpp
@@ -150,258 +150,280 @@ TSharedPtr< SWidget > FLevelEditorContextMenu::BuildMenuWidget( TWeakPtr< SLevel
 
 void FLevelEditorContextMenu::FillMenu( FMenuBuilder& MenuBuilder, TWeakPtr<SLevelEditor> LevelEditor, LevelEditorMenuContext ContextType, TSharedPtr<FExtender> Extender )
 {
-	// Generate information about our selection
-	TArray<AActor*> SelectedActors;
-	GEditor->GetSelectedActors()->GetSelectedObjects<AActor>( SelectedActors );
-
-	FSelectedActorInfo& SelectionInfo = FLevelEditorContextMenuImpl::SelectionInfo;
-	SelectionInfo = AssetSelectionUtils::BuildSelectedActorInfo( SelectedActors );
-
-	// Get all menu extenders for this context menu from the level editor module
-	FLevelEditorModule& LevelEditorModule = FModuleManager::GetModuleChecked<FLevelEditorModule>( TEXT("LevelEditor") );
-	TArray<FLevelEditorModule::FLevelViewportMenuExtender_SelectedActors> MenuExtenderDelegates = LevelEditorModule.GetAllLevelViewportContextMenuExtenders();
-
-	TArray<TSharedPtr<FExtender>> Extenders;
-	if(Extender.IsValid())
-	{
-		Extenders.Add(Extender);
-	}
-
 	auto LevelEditorActions = LevelEditor.Pin()->GetLevelEditorActions().ToSharedRef();
-	for (int32 i = 0; i < MenuExtenderDelegates.Num(); ++i)
+	MenuBuilder.PushCommandList(LevelEditorActions);
+
+	if (GEditor->GetSelectedComponentCount() > 0)
 	{
-		if (MenuExtenderDelegates[i].IsBound())
+		// Show component-specific context menu
+		MenuBuilder.BeginSection("Component", LOCTEXT("ComponentHeading", "Component"));
 		{
-			Extenders.Add(MenuExtenderDelegates[i].Execute(LevelEditorActions, SelectedActors));
+			auto OwnerActor = Cast<AActor>(*(GEditor->GetSelectedActorIterator()));
+			check(OwnerActor);
+
+			MenuBuilder.AddMenuEntry(
+				FLevelEditorCommands::Get().SelectComponentOwnerActor,
+				NAME_None,
+				FText::Format(LOCTEXT("SelectComponentOwner", "Select Owner [{0}]"), FText::FromString(OwnerActor->GetHumanReadableName())),
+				TAttribute<FText>(),
+				FSlateIcon(FEditorStyle::GetStyleSetName(), FClassIconFinder::FindIconNameForClass(OwnerActor->GetClass()))
+				);
 		}
+		MenuBuilder.EndSection();
 	}
-
-	MenuBuilder.PushCommandList(LevelEditorActions);
-	MenuBuilder.PushExtender(FExtender::Combine(Extenders).ToSharedRef());
-
-	TArray<TWeakObjectPtr<UObject>> LabelObjects;
-	for ( FSelectionIterator SelItor(*GEditor->GetSelectedActors()) ; SelItor ; ++SelItor )
+	else
 	{
-		LabelObjects.Add(*SelItor);
-	}
+		// Generate information about our selection
+		TArray<AActor*> SelectedActors;
+		GEditor->GetSelectedActors()->GetSelectedObjects<AActor>(SelectedActors);
 
-	// Check if current selection has any assets that can be browsed to
-	TArray< UObject* > ReferencedAssets;
-	GEditor->GetReferencedAssetsForEditorSelection( ReferencedAssets );
+		FSelectedActorInfo& SelectionInfo = FLevelEditorContextMenuImpl::SelectionInfo;
+		SelectionInfo = AssetSelectionUtils::BuildSelectedActorInfo(SelectedActors);
 
-	const bool bCanSyncToContentBrowser = GEditor->CanSyncToContentBrowser();
+		// Get all menu extenders for this context menu from the level editor module
+		FLevelEditorModule& LevelEditorModule = FModuleManager::GetModuleChecked<FLevelEditorModule>(TEXT("LevelEditor"));
+		TArray<FLevelEditorModule::FLevelViewportMenuExtender_SelectedActors> MenuExtenderDelegates = LevelEditorModule.GetAllLevelViewportContextMenuExtenders();
 
-	if( bCanSyncToContentBrowser || ReferencedAssets.Num() > 0 )		
-	{
-		MenuBuilder.BeginSection("ActorAsset", LOCTEXT("AssetHeading", "Asset") );
+		TArray<TSharedPtr<FExtender>> Extenders;
+		if (Extender.IsValid())
 		{
-			if( bCanSyncToContentBrowser )
-			{
-				MenuBuilder.AddMenuEntry( FGlobalEditorCommonCommands::Get().FindInContentBrowser );
-			}
+			Extenders.Add(Extender);
+		}
 
-			if( ReferencedAssets.Num() == 1 )
+		for (int32 i = 0; i < MenuExtenderDelegates.Num(); ++i)
+		{
+			if (MenuExtenderDelegates[i].IsBound())
 			{
-				auto Asset = ReferencedAssets[0];
-
-				MenuBuilder.AddMenuEntry( 
-					FLevelEditorCommands::Get().EditAsset,
-					NAME_None,
-					FText::Format( LOCTEXT("EditAssociatedAsset", "Edit {0}"), FText::FromString( Asset->GetName() ) ),
-					TAttribute<FText>(),
-					FSlateIcon( FEditorStyle::GetStyleSetName(), FClassIconFinder::FindIconNameForClass( Asset->GetClass() ) )
-					);
+				Extenders.Add(MenuExtenderDelegates[i].Execute(LevelEditorActions, SelectedActors));
 			}
-			else if ( ReferencedAssets.Num() > 1 )
+		}
+		MenuBuilder.PushExtender(FExtender::Combine(Extenders).ToSharedRef());
+
+		TArray<TWeakObjectPtr<UObject>> LabelObjects;
+		for (FSelectionIterator SelItor(*GEditor->GetSelectedActors()); SelItor; ++SelItor)
+		{
+			LabelObjects.Add(*SelItor);
+		}
+
+		// Check if current selection has any assets that can be browsed to
+		TArray< UObject* > ReferencedAssets;
+		GEditor->GetReferencedAssetsForEditorSelection(ReferencedAssets);
+
+		const bool bCanSyncToContentBrowser = GEditor->CanSyncToContentBrowser();
+
+		if (bCanSyncToContentBrowser || ReferencedAssets.Num() > 0)
+		{
+			MenuBuilder.BeginSection("ActorAsset", LOCTEXT("AssetHeading", "Asset"));
 			{
-				MenuBuilder.AddMenuEntry( 
-					FLevelEditorCommands::Get().EditAssetNoConfirmMultiple,
-					NAME_None,
-					LOCTEXT("EditAssociatedAssetsMultiple", "Edit Multiple Assets"),
-					TAttribute<FText>(),
-					FSlateIcon( FEditorStyle::GetStyleSetName(), "ClassIcon.Default" )
-					);
+				if (bCanSyncToContentBrowser)
+				{
+					MenuBuilder.AddMenuEntry(FGlobalEditorCommonCommands::Get().FindInContentBrowser);
+				}
 
-			}
+				if (ReferencedAssets.Num() == 1)
+				{
+					auto Asset = ReferencedAssets[0];
+
+					MenuBuilder.AddMenuEntry(
+						FLevelEditorCommands::Get().EditAsset,
+						NAME_None,
+						FText::Format(LOCTEXT("EditAssociatedAsset", "Edit {0}"), FText::FromString(Asset->GetName())),
+						TAttribute<FText>(),
+						FSlateIcon(FEditorStyle::GetStyleSetName(), FClassIconFinder::FindIconNameForClass(Asset->GetClass()))
+						);
+				}
+				else if (ReferencedAssets.Num() > 1)
+				{
+					MenuBuilder.AddMenuEntry(
+						FLevelEditorCommands::Get().EditAssetNoConfirmMultiple,
+						NAME_None,
+						LOCTEXT("EditAssociatedAssetsMultiple", "Edit Multiple Assets"),
+						TAttribute<FText>(),
+						FSlateIcon(FEditorStyle::GetStyleSetName(), "ClassIcon.Default")
+						);
+
+				}
 
-			MenuBuilder.AddMenuEntry( FGlobalEditorCommonCommands::Get().ViewReferences );
+				MenuBuilder.AddMenuEntry(FGlobalEditorCommonCommands::Get().ViewReferences);
 
+			}
+			MenuBuilder.EndSection();
 		}
-		MenuBuilder.EndSection();
-	}
 
 
-	MenuBuilder.BeginSection( "ActorControl", LOCTEXT("ActorHeading", "Actor") );
-	{
-		MenuBuilder.AddMenuEntry( FEditorViewportCommands::Get().FocusViewportToSelection );
+		MenuBuilder.BeginSection("ActorControl", LOCTEXT("ActorHeading", "Actor"));
+		{
+			MenuBuilder.AddMenuEntry(FEditorViewportCommands::Get().FocusViewportToSelection);
 
-		
-		const FVector* ClickLocation = &GEditor->ClickLocation;
 
-		FUIAction GoHereAction;
-		GoHereAction.ExecuteAction = FExecuteAction::CreateStatic( &FLevelEditorActionCallbacks::GoHere_Clicked, ClickLocation );
+			const FVector* ClickLocation = &GEditor->ClickLocation;
 
-		MenuBuilder.AddMenuEntry( FLevelEditorCommands::Get().GoHere );
-		MenuBuilder.AddMenuEntry( FLevelEditorCommands::Get().SnapCameraToActor );
-	}
-	MenuBuilder.EndSection();
+			FUIAction GoHereAction;
+			GoHereAction.ExecuteAction = FExecuteAction::CreateStatic(&FLevelEditorActionCallbacks::GoHere_Clicked, ClickLocation);
 
-	// Go to C++ Code
-	if( SelectionInfo.SelectionClass != NULL )
-	{
-		if ( FSourceCodeNavigation::IsCompilerAvailable())
+			MenuBuilder.AddMenuEntry(FLevelEditorCommands::Get().GoHere);
+			MenuBuilder.AddMenuEntry(FLevelEditorCommands::Get().SnapCameraToActor);
+		}
+		MenuBuilder.EndSection();
+
+		// Go to C++ Code
+		if (SelectionInfo.SelectionClass != NULL)
 		{
-			FString ClassHeaderPath;
-			if( FSourceCodeNavigation::FindClassHeaderPath( SelectionInfo.SelectionClass, ClassHeaderPath ) && IFileManager::Get().FileSize( *ClassHeaderPath ) != INDEX_NONE )
+			if (FSourceCodeNavigation::IsCompilerAvailable())
 			{
-				const FString CodeFileName = FPaths::GetCleanFilename( *ClassHeaderPath );
+				FString ClassHeaderPath;
+				if (FSourceCodeNavigation::FindClassHeaderPath(SelectionInfo.SelectionClass, ClassHeaderPath) && IFileManager::Get().FileSize(*ClassHeaderPath) != INDEX_NONE)
+				{
+					const FString CodeFileName = FPaths::GetCleanFilename(*ClassHeaderPath);
 
-				MenuBuilder.BeginSection( "ActorCode", LOCTEXT("ActorCodeHeading", "C++") );
+					MenuBuilder.BeginSection("ActorCode", LOCTEXT("ActorCodeHeading", "C++"));
+					{
+						MenuBuilder.AddMenuEntry(FLevelEditorCommands::Get().GoToCodeForActor,
+							NAME_None,
+							FText::Format(LOCTEXT("GoToCodeForActor", "Open {0}"), FText::FromString(CodeFileName)),
+							FText::Format(LOCTEXT("GoToCodeForActor_ToolTip", "Opens the header file for this actor ({0}) in a code editing program"), FText::FromString(CodeFileName)));
+					}
+					MenuBuilder.EndSection();
+				}
+			}
+
+			const FString DocumentationLink = FEditorClassUtils::GetDocumentationLink(SelectionInfo.SelectionClass);
+			if (!DocumentationLink.IsEmpty())
+			{
+				MenuBuilder.BeginSection("ActorDocumentation", LOCTEXT("ActorDocsHeading", "Documentation"));
 				{
-					MenuBuilder.AddMenuEntry( FLevelEditorCommands::Get().GoToCodeForActor,
-						NAME_None, 
-						FText::Format( LOCTEXT("GoToCodeForActor", "Open {0}"), FText::FromString( CodeFileName ) ),
-						FText::Format( LOCTEXT("GoToCodeForActor_ToolTip", "Opens the header file for this actor ({0}) in a code editing program"), FText::FromString( CodeFileName ) ) );
+					MenuBuilder.AddMenuEntry(FLevelEditorCommands::Get().GoToDocsForActor,
+						NAME_None,
+						LOCTEXT("GoToDocsForActor", "View Documentation"),
+						LOCTEXT("GoToDocsForActor_ToolTip", "Click to open documentation for this actor"),
+						FSlateIcon(FEditorStyle::GetStyleSetName(), "HelpIcon.Hovered"));
 				}
 				MenuBuilder.EndSection();
 			}
-		}
 
-		const FString DocumentationLink = FEditorClassUtils::GetDocumentationLink(SelectionInfo.SelectionClass);
-		if (!DocumentationLink.IsEmpty())
-		{
-			MenuBuilder.BeginSection( "ActorDocumentation", LOCTEXT("ActorDocsHeading", "Documentation") );
+			// If it is in the map then it is already anonymous
+			if (GetDefault<UEditorExperimentalSettings>()->bInWorldBPEditing
+				&& SelectedActors.Num() == 1
+				&& !FBlueprintEditorUtils::IsAnonymousBlueprintClass(SelectionInfo.SelectionClass)
+				&& FKismetEditorUtilities::CanCreateBlueprintOfClass(SelectionInfo.SelectionClass))
 			{
-				MenuBuilder.AddMenuEntry( FLevelEditorCommands::Get().GoToDocsForActor,
-					NAME_None, 
-					LOCTEXT("GoToDocsForActor", "View Documentation"),
-					LOCTEXT("GoToDocsForActor_ToolTip", "Click to open documentation for this actor"),
-					FSlateIcon(FEditorStyle::GetStyleSetName(), "HelpIcon.Hovered" ));
+				MenuBuilder.BeginSection("Customize", LOCTEXT("CustomizeHeading", "Customize"));
+				{
+					MenuBuilder.AddMenuEntry(FLevelEditorCommands::Get().AddScriptBehavior,
+						NAME_None,
+						LOCTEXT("AddScriptBehavior", "Customize Scripting Behavior"),
+						LOCTEXT("AddScriptBehavior_ToolTip", "Click to customize scripting behavior of this Actor"));
+				}
+				MenuBuilder.EndSection();
 			}
-			MenuBuilder.EndSection();
 		}
 
-		// If it is in the map then it is already anonymous
-		if (   GetDefault<UEditorExperimentalSettings>()->bInWorldBPEditing
-			&& SelectedActors.Num() == 1 
-			&& !FBlueprintEditorUtils::IsAnonymousBlueprintClass(SelectionInfo.SelectionClass) 
-			&& FKismetEditorUtilities::CanCreateBlueprintOfClass(SelectionInfo.SelectionClass) )
+		MenuBuilder.BeginSection("ActorSelectVisibilityLevels");
 		{
-			MenuBuilder.BeginSection( "Customize", LOCTEXT("CustomizeHeading", "Customize") );
-			{
-				MenuBuilder.AddMenuEntry( FLevelEditorCommands::Get().AddScriptBehavior,
-					NAME_None, 
-					LOCTEXT("AddScriptBehavior", "Customize Scripting Behavior"),
-					LOCTEXT("AddScriptBehavior_ToolTip", "Click to customize scripting behavior of this Actor"));
-			}
-			MenuBuilder.EndSection();
+			// Add a sub-menu for "Select"
+			MenuBuilder.AddSubMenu(
+				LOCTEXT("SelectSubMenu", "Select"),
+				LOCTEXT("SelectSubMenu_ToolTip", "Opens the actor selection menu"),
+				FNewMenuDelegate::CreateStatic(&FLevelEditorContextMenuImpl::FillSelectActorMenu));
+
+			MenuBuilder.AddSubMenu(
+				LOCTEXT("EditSubMenu", "Edit"),
+				FText::GetEmpty(),
+				FNewMenuDelegate::CreateStatic(&FLevelEditorContextMenuImpl::FillEditMenu, ContextType));
+
+			MenuBuilder.AddSubMenu(
+				LOCTEXT("VisibilitySubMenu", "Visibility"),
+				LOCTEXT("VisibilitySubMenu_ToolTip", "Selected actor visibility options"),
+				FNewMenuDelegate::CreateStatic(&FLevelEditorContextMenuImpl::FillActorVisibilityMenu));
+
+			// Build the menu for grouping actors
+			BuildGroupMenu(MenuBuilder, SelectionInfo);
+
+			MenuBuilder.AddSubMenu(
+				LOCTEXT("LevelSubMenu", "Level"),
+				LOCTEXT("LevelSubMenu_ToolTip", "Options for interacting with this actor's level"),
+				FNewMenuDelegate::CreateStatic(&FLevelEditorContextMenuImpl::FillActorLevelMenu));
 		}
-	}
-
-	MenuBuilder.BeginSection("ActorSelectVisibilityLevels");
-	{
-		// Add a sub-menu for "Select"
-		MenuBuilder.AddSubMenu( 
-			LOCTEXT("SelectSubMenu", "Select"),
-			LOCTEXT("SelectSubMenu_ToolTip", "Opens the actor selection menu"),
-			FNewMenuDelegate::CreateStatic( &FLevelEditorContextMenuImpl::FillSelectActorMenu ) );
-
-		MenuBuilder.AddSubMenu( 
-			LOCTEXT("EditSubMenu", "Edit"),
-			FText::GetEmpty(),
-			FNewMenuDelegate::CreateStatic( &FLevelEditorContextMenuImpl::FillEditMenu, ContextType ) );
-
-		MenuBuilder.AddSubMenu( 
-			LOCTEXT("VisibilitySubMenu", "Visibility"),
-			LOCTEXT("VisibilitySubMenu_ToolTip", "Selected actor visibility options"),
-			FNewMenuDelegate::CreateStatic( &FLevelEditorContextMenuImpl::FillActorVisibilityMenu ) );
-
-		// Build the menu for grouping actors
-		BuildGroupMenu( MenuBuilder, SelectionInfo );
-
-		MenuBuilder.AddSubMenu( 
-			LOCTEXT("LevelSubMenu", "Level"),
-			LOCTEXT("LevelSubMenu_ToolTip", "Options for interacting with this actor's level"),
-			FNewMenuDelegate::CreateStatic( &FLevelEditorContextMenuImpl::FillActorLevelMenu ) );
-	}
-	MenuBuilder.EndSection();
+		MenuBuilder.EndSection();
 
-	if (ContextType == LevelEditorMenuContext::Viewport)
-	{
-		LevelEditorCreateActorMenu::FillAddReplaceViewportContextMenuSections( MenuBuilder );
-	}
+		if (ContextType == LevelEditorMenuContext::Viewport)
+		{
+			LevelEditorCreateActorMenu::FillAddReplaceViewportContextMenuSections(MenuBuilder);
+		}
 
-	if( GEditor->PlayWorld != NULL )
-	{
-		if( SelectionInfo.NumSelected > 0 )
+		if (GEditor->PlayWorld != NULL)
 		{
-			MenuBuilder.BeginSection( "Simulation", NSLOCTEXT( "LevelViewportContextMenu", "SimulationHeading", "Simulation" ) );
+			if (SelectionInfo.NumSelected > 0)
 			{
-				MenuBuilder.AddMenuEntry( FLevelEditorCommands::Get().KeepSimulationChanges );
+				MenuBuilder.BeginSection("Simulation", NSLOCTEXT("LevelViewportContextMenu", "SimulationHeading", "Simulation"));
+				{
+					MenuBuilder.AddMenuEntry(FLevelEditorCommands::Get().KeepSimulationChanges);
+				}
+				MenuBuilder.EndSection();
 			}
-			MenuBuilder.EndSection();
 		}
-	}
 
-	MenuBuilder.BeginSection("LevelViewportAttach");
-	{
-		// Only display the attach menu if we have actors selected
-		if ( GEditor->GetSelectedActorCount() )
+		MenuBuilder.BeginSection("LevelViewportAttach");
 		{
-			if(SelectionInfo.bHaveAttachedActor)
+			// Only display the attach menu if we have actors selected
+			if (GEditor->GetSelectedActorCount())
 			{
-				MenuBuilder.AddMenuEntry( FLevelEditorCommands::Get().DetachFromParent );
-			}
+				if (SelectionInfo.bHaveAttachedActor)
+				{
+					MenuBuilder.AddMenuEntry(FLevelEditorCommands::Get().DetachFromParent);
+				}
 
-			MenuBuilder.AddSubMenu( 
-				LOCTEXT( "ActorAttachToSubMenu", "Attach To" ), 
-				LOCTEXT( "ActorAttachToSubMenu_ToolTip", "Attach Actor as child" ),
-				FNewMenuDelegate::CreateStatic( &FLevelEditorContextMenuImpl::FillActorMenu ) );
-		}
+				MenuBuilder.AddSubMenu(
+					LOCTEXT("ActorAttachToSubMenu", "Attach To"),
+					LOCTEXT("ActorAttachToSubMenu_ToolTip", "Attach Actor as child"),
+					FNewMenuDelegate::CreateStatic(&FLevelEditorContextMenuImpl::FillActorMenu));
+			}
 
-		// Add a heading for "Movement" if an actor is selected
-		if ( GEditor->GetSelectedActorIterator() )
-		{
-			// Add a sub-menu for "Transform"
-			MenuBuilder.AddSubMenu( 
-				LOCTEXT("TransformSubMenu", "Transform"), 
-				LOCTEXT("TransformSubMenu_ToolTip", "Actor transform utils"),
-				FNewMenuDelegate::CreateStatic( &FLevelEditorContextMenuImpl::FillTransformMenu ) );
-		}
+			// Add a heading for "Movement" if an actor is selected
+			if (GEditor->GetSelectedActorIterator())
+			{
+				// Add a sub-menu for "Transform"
+				MenuBuilder.AddSubMenu(
+					LOCTEXT("TransformSubMenu", "Transform"),
+					LOCTEXT("TransformSubMenu_ToolTip", "Actor transform utils"),
+					FNewMenuDelegate::CreateStatic(&FLevelEditorContextMenuImpl::FillTransformMenu));
+			}
 
-		// @todo UE4: The current pivot options only work for brushes
-		if( SelectionInfo.bHaveBrush )
-		{
-			// You can only move the pivot in ortho viewports, but you can reset it in any viewport
-			if( GCurrentLevelEditingViewportClient->ViewportType != LVT_Perspective )
+			// @todo UE4: The current pivot options only work for brushes
+			if (SelectionInfo.bHaveBrush)
 			{
-				// Add a sub-menu for "Pivot"
-				MenuBuilder.AddSubMenu( 
-					LOCTEXT("PivotSubMenu", "Pivot"), 
-					LOCTEXT("PivotSubMenu_ToolTip", "Actor pivoting utils"),
-					FNewMenuDelegate::CreateStatic( &FLevelEditorContextMenuImpl::FillPivotMenu ) );
+				// You can only move the pivot in ortho viewports, but you can reset it in any viewport
+				if (GCurrentLevelEditingViewportClient->ViewportType != LVT_Perspective)
+				{
+					// Add a sub-menu for "Pivot"
+					MenuBuilder.AddSubMenu(
+						LOCTEXT("PivotSubMenu", "Pivot"),
+						LOCTEXT("PivotSubMenu_ToolTip", "Actor pivoting utils"),
+						FNewMenuDelegate::CreateStatic(&FLevelEditorContextMenuImpl::FillPivotMenu));
+				}
+				else
+				{
+					MenuBuilder.AddMenuEntry(FLevelEditorCommands::Get().ResetPivot);
+				}
 			}
-			else
+
+			if (GetDefault<UEditorExperimentalSettings>()->bActorMerging &&
+				( SelectionInfo.bHaveStaticMeshComponent || SelectionInfo.bHaveLandscape ))
 			{
-				MenuBuilder.AddMenuEntry( FLevelEditorCommands::Get().ResetPivot );
+				MenuBuilder.AddSubMenu(
+					LOCTEXT("MergeActorsSubMenu", "Merge"),
+					LOCTEXT("MergeActorsSubMenu_ToolTip", "Actor merging utils"),
+					FNewMenuDelegate::CreateStatic(&FLevelEditorContextMenuImpl::FillMergeActorsMenu));
 			}
 		}
+		MenuBuilder.EndSection();
 
-		if (GetDefault<UEditorExperimentalSettings>()->bActorMerging && 
-			(SelectionInfo.bHaveStaticMeshComponent || SelectionInfo.bHaveLandscape))
-		{
-			MenuBuilder.AddSubMenu( 
-				LOCTEXT("MergeActorsSubMenu", "Merge"), 
-				LOCTEXT("MergeActorsSubMenu_ToolTip", "Actor merging utils"),
-				FNewMenuDelegate::CreateStatic( &FLevelEditorContextMenuImpl::FillMergeActorsMenu ) );
-		}
-	}
-	MenuBuilder.EndSection();
+		FLevelScriptEventMenuHelper::FillLevelBlueprintEventsMenu(MenuBuilder, SelectedActors);
 
-	FLevelScriptEventMenuHelper::FillLevelBlueprintEventsMenu(MenuBuilder, SelectedActors);
+		MenuBuilder.PopExtender();
+	}	
 
 	MenuBuilder.PopCommandList();
-	MenuBuilder.PopExtender();
 }
 
 
diff --git a/Engine/Source/Editor/LevelEditor/Public/LevelEditorActions.h b/Engine/Source/Editor/LevelEditor/Public/LevelEditorActions.h
--- a/Engine/Source/Editor/LevelEditor/Public/LevelEditorActions.h
+++ b/Engine/Source/Editor/LevelEditor/Public/LevelEditorActions.h
@@ -325,6 +325,9 @@ public:
 	/** Selects all actors of the same class and archetype as the current selection */
 	TSharedPtr< FUICommandInfo > SelectAllActorsOfSameClassWithArchetype;
 
+	/** Selects the actor that owns the currently selected component(s) */
+	TSharedPtr< FUICommandInfo > SelectComponentOwnerActor;
+
 	/** Selects all lights relevant to the current selection */
 	TSharedPtr< FUICommandInfo > SelectRelevantLights;
 
@@ -870,6 +873,12 @@ public:
 	 */
 	static void OnSelectAllActorsOfClass( bool bArchetype );
 
+	/** Called when selecting the actor that owns the currently selected component(s) */
+	static void OnSelectComponentOwnerActor();
+
+	/** Called to see if any components are selected */
+	static bool CanSelectComponentOwnerActor();
+
 	/**
 	 * Called to select all lights
 	 */
diff --git a/Engine/Source/Editor/UnrealEd/Private/EditorSelectUtils.cpp b/Engine/Source/Editor/UnrealEd/Private/EditorSelectUtils.cpp
--- a/Engine/Source/Editor/UnrealEd/Private/EditorSelectUtils.cpp
+++ b/Engine/Source/Editor/UnrealEd/Private/EditorSelectUtils.cpp
@@ -608,7 +608,7 @@ void UUnrealEdEngine::SelectActor(AActor* Actor, bool bInSelected, bool bNotify,
 
 void UUnrealEdEngine::SelectComponent(UActorComponent* Component, bool bInSelected, bool bNotify, bool bSelectEvenIfHidden)
 {
-	// Don't do any work if the component is already selected
+	// Don't do any work if the component's selection state matches the target selection state
 	const bool bComponentSelected = Component->IsSelected();
 	if (( bComponentSelected && !bInSelected ) || ( !bComponentSelected && bInSelected ))
 	{
@@ -633,7 +633,6 @@ void UUnrealEdEngine::SelectComponent(UActorComponent* Component, bool bInSelect
 		AActor* ComponentOwner = Component->GetOwner();
 		if (ComponentOwner != nullptr)
 		{
-			//@todo dhertzka - need selection override delegates for the primitive components
 			TInlineComponentArray<UPrimitiveComponent*> PrimitiveComponents;
 			ComponentOwner->GetComponents(PrimitiveComponents);
 
@@ -648,14 +647,6 @@ void UUnrealEdEngine::SelectComponent(UActorComponent* Component, bool bInSelect
 			NoteSelectionChange(true);
 		}
 	}
-	else
-	{
-		if (bNotify)
-		{
-			//reset the property windows.  In case something has changed since previous selection
-			UpdateFloatingPropertyWindows();
-		}
-	}
 }
 
 bool UUnrealEdEngine::IsComponentSelected(const UPrimitiveComponent* PrimComponent)
diff --git a/Engine/Source/Editor/UnrealEd/Private/LevelEditorViewport.cpp b/Engine/Source/Editor/UnrealEd/Private/LevelEditorViewport.cpp
--- a/Engine/Source/Editor/UnrealEd/Private/LevelEditorViewport.cpp
+++ b/Engine/Source/Editor/UnrealEd/Private/LevelEditorViewport.cpp
@@ -18,6 +18,7 @@
 #include "Engine.h"
 #include "Editor/LevelEditor/Public/LevelEditor.h"
 #include "Editor/LevelEditor/Public/LevelViewportActions.h"
+#include "Editor/LevelEditor/Public/SLevelViewport.h"
 #include "Editor/PropertyEditor/Public/PropertyEditorModule.h"
 #include "AssetSelection.h"
 #include "BlueprintUtilities.h"
@@ -47,7 +48,6 @@
 #include "Editor/ActorPositioning.h"
 #include "NotificationManager.h"
 #include "SNotificationList.h"
-#include "SLevelViewport.h"
 
 DEFINE_LOG_CATEGORY(LogEditorViewport);
 
@@ -1133,12 +1133,12 @@ FDropQuery FLevelEditorViewportClient::CanDropObjectsAtCoordinates(int32 MouseX,
 		if ( AssetObj->IsA( AActor::StaticClass() ) || bHasActorFactory )
 		{
 			Result.bCanDrop = true;
-			bPivotMovedIndependantly = false;
+			bPivotMovedIndependently = false;
 		}
 		else if( AssetObj->IsA( UBrushBuilder::StaticClass()) )
 		{
 			Result.bCanDrop = true;
-			bPivotMovedIndependantly = false;
+			bPivotMovedIndependently = false;
 		}
 		else
 		{
@@ -1149,7 +1149,7 @@ FDropQuery FLevelEditorViewportClient::CanDropObjectsAtCoordinates(int32 MouseX,
 				{
 					// If our asset is a material and the target is a valid recipient
 					Result.bCanDrop = true;
-					bPivotMovedIndependantly = false;
+					bPivotMovedIndependently = false;
 
 					//if ( HitProxy->IsA(HActor::StaticGetType()) )
 					//{
@@ -1444,6 +1444,7 @@ FLevelEditorViewportClient::FLevelEditorViewportClient(const TSharedPtr<SLevelVi
 	, bDuplicateActorsInProgress( false )
 	, bIsTrackingBrushModification( false )
 	, bLockedCameraView(true)
+	, bReceivedFocusRecently(false)
 	, SpriteCategoryVisibility()
 	, World(nullptr)
 	, TrackingTransaction()
@@ -1696,15 +1697,32 @@ void FLevelEditorViewportClient::RestoreCameraFromPIE()
 	}
 }
 
+void FLevelEditorViewportClient::ReceivedFocus(FViewport* InViewport)
+{
+	if (!bReceivedFocusRecently)
+	{
+		bReceivedFocusRecently = true;
+
+		// A few frames can pass between receiving focus and processing a click, so we use a timer to track whether we have recently received focus.
+		FTimerHandle DummyHandle;
+		FTimerDelegate ResetFocusReceivedTimer;
+		ResetFocusReceivedTimer.BindLambda([&] ()
+		{
+			bReceivedFocusRecently = false;
+		});
+		GEditor->GetTimerManager()->SetTimer(DummyHandle, ResetFocusReceivedTimer, 0.1f, false);
+	}
+
+	FEditorViewportClient::ReceivedFocus(InViewport);
+}
 
 //
 //	FLevelEditorViewportClient::ProcessClick
 //
-
 void FLevelEditorViewportClient::ProcessClick(FSceneView& View, HHitProxy* HitProxy, FKey Key, EInputEvent Event, uint32 HitX, uint32 HitY)
 {
 	// We clicked, allow the pivot to reposition itself.
-	bPivotMovedIndependantly = false;
+	bPivotMovedIndependently = false;
 
 	static FName ProcessClickTrace = FName(TEXT("ProcessClickTrace"));
 
@@ -1763,11 +1781,23 @@ void FLevelEditorViewportClient::ProcessClick(FSceneView& View, HHitProxy* HitPr
 		{
 			auto ActorHitProxy = (HActor*)HitProxy;
 
-			auto SelectedActors = GEditor->GetSelectedActors();
-			const bool bSelectComponents = SelectedActors->IsSelected(ActorHitProxy->Actor) && GEditor->GetSelectedActorCount() == 1;
-			if (bSelectComponents && GetDefault<UEditorExperimentalSettings>()->bInWorldBPEditing)
+			// We want to process the click on the component only if:
+			// 1. The actor clicked is already selected
+			// 2. The actor selected is the only actor selected
+			// 3. The LMB was pressed or we already have a component selected (so that right-clicking a selected actor won't select a component)
+			// 4. The click was not a double click (unless a component has already been selected)
+			// 5. The level viewport didn't just receive focus this frame (again unless a component has already been selected)
+			const bool bActorAlreadySelectedExclusively = GEditor->GetSelectedActors()->IsSelected(ActorHitProxy->Actor) && ( GEditor->GetSelectedActorCount() == 1 );
+			const bool bComponentAlreadySelected = GEditor->GetSelectedComponentCount() > 0;
+	
+			const bool bCanBeginSelectingComponents =  (Click.GetKey() == EKeys::LeftMouseButton) 
+													&& (Click.GetEvent() != IE_DoubleClick) 
+													&& !bReceivedFocusRecently;
+
+			const bool bSelectComponent = bActorAlreadySelectedExclusively && ( bComponentAlreadySelected || bCanBeginSelectingComponents );
+
+			if (bSelectComponent && GetDefault<UEditorExperimentalSettings>()->bInWorldBPEditing)
 			{
-				// The hit actor is already selected and is the only selected actor, so select the hit component
 				ClickHandlers::ClickComponent(this, ActorHitProxy, Click);
 			}
 			else
@@ -1850,7 +1880,7 @@ void FLevelEditorViewportClient::Tick(float DeltaTime)
 {
 	FEditorViewportClient::Tick(DeltaTime);
 
-	if( !bPivotMovedIndependantly && GCurrentLevelEditingViewportClient == this &&
+	if( !bPivotMovedIndependently && GCurrentLevelEditingViewportClient == this &&
 		bIsRealtime &&
 		( Widget == NULL || !Widget->IsDragging() ) )
 	{
@@ -2157,7 +2187,7 @@ bool FLevelEditorViewportClient::InputWidgetDelta(FViewport* Viewport, EAxisList
 				else
 				{
 					FSnappingUtils::SnapDragLocationToNearestVertex( ModeTools->PivotLocation, Drag, this );
-					bPivotMovedIndependantly = true;
+					bPivotMovedIndependently = true;
 				}
 
 				ModeTools->PivotLocation += Drag;
@@ -2528,7 +2558,7 @@ void FLevelEditorViewportClient::TrackingStopped()
 			GEditor->BroadcastEndObjectMovement( *Actor );
 		}
 
-		if (!bPivotMovedIndependantly)
+		if (!bPivotMovedIndependently)
 		{
 			GUnrealEd->UpdatePivotLocationForSelection();
 		}
diff --git a/Engine/Source/Editor/UnrealEd/Private/LevelViewportClickHandlers.cpp b/Engine/Source/Editor/UnrealEd/Private/LevelViewportClickHandlers.cpp
--- a/Engine/Source/Editor/UnrealEd/Private/LevelViewportClickHandlers.cpp
+++ b/Engine/Source/Editor/UnrealEd/Private/LevelViewportClickHandlers.cpp
@@ -266,39 +266,81 @@ namespace ClickHandlers
 
 	bool ClickComponent(FLevelEditorViewportClient* ViewportClient, HActor* ActorHitProxy, const FViewportClick& Click)
 	{
+		//@todo hotkeys for component placement?
+		
 		bool bComponentClicked = false;
-		const FScopedTransaction Transaction(NSLOCTEXT("UnrealEd", "ClickingOnComponents", "Clicking on Components"));
-
-		UActorComponent* SelectedComponent = const_cast<UPrimitiveComponent*>(ActorHitProxy->PrimComponent); 
-		GEditor->GetSelectedComponents()->Modify();
+		UActorComponent* Component = const_cast<UPrimitiveComponent*>(ActorHitProxy->PrimComponent); 
+		
+		// Pivot snapping
+		if (Click.GetKey() == EKeys::MiddleMouseButton && Click.IsAltDown())
+		{
+			GEditor->SetPivot(GEditor->ClickLocation, true, false, true);
 
-		if (Click.IsControlDown())
+			return true;
+		}
+		// Selection + context menu
+		else if (Click.GetKey() == EKeys::RightMouseButton && !Click.IsControlDown() && !ViewportClient->Viewport->KeyState(EKeys::LeftMouseButton))
 		{
-			const bool bSelect = !SelectedComponent->IsSelected();
-			if (bSelect)
+			const FScopedTransaction Transaction(NSLOCTEXT("UnrealEd", "ClickingOnComponentContextMenu", "Clicking on Component (context menu)"));
+			UE_LOG(LogEditorViewport, Log, TEXT("Clicking on Component (context menu): %s (%s)"), *Component->GetClass()->GetName(), *Component->GetName());
+
+			auto const EditorComponentSelection = GEditor->GetSelectedComponents();
+			EditorComponentSelection->Modify();
+
+			// If the component the user clicked on was already selected, then we won't bother clearing the selection
+			bool bNeedViewportRefresh = false;
+			if (!EditorComponentSelection->IsSelected(Component))
 			{
-				UE_LOG(LogEditorViewport, Log, TEXT("Clicking on Component (CTRL LMB): %s (%s)"), *SelectedComponent->GetClass()->GetName(), *SelectedComponent->GetName());
+				EditorComponentSelection->DeselectAll();
+				bNeedViewportRefresh = true;
 			}
-			GEditor->SelectComponent(SelectedComponent, bSelect, true, true);
-			bComponentClicked = true;
-		}
-		else if (Click.IsShiftDown())
-		{
-			if (!SelectedComponent->IsSelected())
+
+			GEditor->SelectComponent(Component, true, true);
+			
+			if (bNeedViewportRefresh)
 			{
-				const bool bSelect = true;
-				GEditor->SelectComponent(SelectedComponent, bSelect, true, true);
+				// Redraw the viewport so the user can see which object was right clicked on
+				ViewportClient->Viewport->Draw();
+				FlushRenderingCommands();
 			}
+
+			PrivateSummonContextMenu(ViewportClient);
 			bComponentClicked = true;
 		}
-		else
+		// Selection only
+		else if (Click.GetKey() == EKeys::LeftMouseButton)
 		{
-			GEditor->GetSelectedComponents()->DeselectAll();
-			UE_LOG(LogEditorViewport, Log, TEXT("Clicking on Component (LMB): %s (%s)"), *SelectedComponent->GetClass()->GetName(), *SelectedComponent->GetName());
-			GEditor->SelectComponent(SelectedComponent, true, true, true);
-			bComponentClicked = true;
-		}
+			const FScopedTransaction Transaction(NSLOCTEXT("UnrealEd", "ClickingOnComponents", "Clicking on Components"));
+			GEditor->GetSelectedComponents()->Modify();
 
+			if (Click.IsControlDown())
+			{
+				const bool bSelect = !Component->IsSelected();
+				if (bSelect)
+				{
+					UE_LOG(LogEditorViewport, Log, TEXT("Clicking on Component (CTRL LMB): %s (%s)"), *Component->GetClass()->GetName(), *Component->GetName());
+				}
+				GEditor->SelectComponent(Component, bSelect, true, true);
+				bComponentClicked = true;
+			}
+			else if (Click.IsShiftDown())
+			{
+				if (!Component->IsSelected())
+				{
+					UE_LOG(LogEditorViewport, Log, TEXT("Clicking on Component (SHIFT LMB): %s (%s)"), *Component->GetClass()->GetName(), *Component->GetName());
+					GEditor->SelectComponent(Component, true, true, true);
+				}
+				bComponentClicked = true;
+			}
+			else
+			{
+				GEditor->GetSelectedComponents()->DeselectAll();
+				UE_LOG(LogEditorViewport, Log, TEXT("Clicking on Component (LMB): %s (%s)"), *Component->GetClass()->GetName(), *Component->GetName());
+				GEditor->SelectComponent(Component, true, true, true);
+				bComponentClicked = true;
+			}
+		}
+		
 		return bComponentClicked;
 	}
 
diff --git a/Engine/Source/Editor/UnrealEd/Public/LevelEditorViewport.h b/Engine/Source/Editor/UnrealEd/Public/LevelEditorViewport.h
--- a/Engine/Source/Editor/UnrealEd/Public/LevelEditorViewport.h
+++ b/Engine/Source/Editor/UnrealEd/Public/LevelEditorViewport.h
@@ -153,6 +153,7 @@ public:
 	virtual FLinearColor GetBackgroundColor() const override;
 	virtual int32 GetCameraSpeedSetting() const override;
 	virtual void SetCameraSpeedSetting(int32 SpeedSetting) override;
+	virtual void ReceivedFocus(FViewport* Viewport) override;
 	virtual void ProcessClick(FSceneView& View, HHitProxy* HitProxy, FKey Key, EInputEvent Event, uint32 HitX, uint32 HitY) override;
 	virtual UWorld* GetWorld() const override;
 
@@ -523,7 +524,6 @@ private:
 	 * Moves the locked actor according to the viewport cameras location and rotation
 	 */
 	void MoveLockedActorToCamera();
-
 	
 	/** @return	Returns true if the delta tracker was used to modify any selected actors or BSP.  Must be called before EndTracking(). */
 	bool HaveSelectedObjectsBeenChanged() const;
@@ -600,8 +600,6 @@ private:
 	 */
 	bool DropObjectsOnWidget(FSceneView* View, struct FViewportCursorLocation& Cursor, const TArray<UObject*>& DroppedObjects, bool bCreateDropPreview = false);
 
-	
-
 	/** Helper functions for ApplyDeltaTo* functions - modifies scale based on grid settings */
 	void ModifyScale( AActor* InActor, FVector& ScaleDelta, bool bCheckSmallExtent = false ) const;
 	void ValidateScale( const FVector& CurrentScale, const FVector& BoxExtent, FVector& ScaleDelta, bool bCheckSmallExtent = false ) const;
@@ -615,7 +613,7 @@ private:
 public:
 	/** Static: List of objects we're hovering over */
 	static TSet< FViewportHoverTarget > HoveredObjects;
-		
+	
 	/** Parent level editor that owns this viewport.  Currently, this may be null if the parent doesn't happen to be a level editor. */
 	TWeakPtr< class ILevelEditor > ParentLevelEditor;
 
@@ -634,15 +632,14 @@ public:
 
 	FColor					FadeColor;
 
-	
 	float					FadeAmount;
 
 	bool					bEnableFading;
 
 	bool					bEnableColorScaling;
 
-	/** If true then the pivot has been moved independantly of the actor and position updates should not occur when the actor is moved. */
-	bool					bPivotMovedIndependantly;
+	/** If true, the pivot has been moved independently of the actor and position updates should not occur when the actor is moved. */
+	bool					bPivotMovedIndependently;
 
 	/** If true, we switched between two different cameras. Set by matinee, used by the motion blur to invalidate this frames motion vectors */
 	bool					bEditorCameraCut;
@@ -670,6 +667,9 @@ public:
 	/** True if this viewport is to change its view (aspect ratio, post processing, FOV etc) to match that of the currently locked camera, if applicable */
 	bool					bLockedCameraView;
 
+	/** Whether this viewport recently received focus. Used to determine whether component selection is permissible. */
+	bool bReceivedFocusRecently;
+
 private:
 	/** The actors that are currently being placed in the viewport via dragging */
 	static TArray< TWeakObjectPtr< AActor > > DropPreviewActors;
