In-World Component Editing
- SCSEditor now shows up for any selected actor
- The details of a selected component now correctly displays the component template details
- Undo (at least for translations) is operational in the world (retains selection and properly modifies/reverts the template in addition to the instance in the world)
- Note undo does not function properly in the BP editor for component translation - selection is lost and the preview instance is not updated; will fix

diff --git a/Engine/Source/Editor/Kismet/Private/SSCSEditor.cpp b/Engine/Source/Editor/Kismet/Private/SSCSEditor.cpp
--- a/Engine/Source/Editor/Kismet/Private/SSCSEditor.cpp
+++ b/Engine/Source/Editor/Kismet/Private/SSCSEditor.cpp
@@ -2075,11 +2075,6 @@ void SSCSEditor::Construct( const FArguments& InArgs )
 	}
 }
 
-void SSCSEditor::SetSCS(USimpleConstructionScript* InSCS)
-{
-
-}
-
 FReply SSCSEditor::OnKeyDown( const FGeometry& MyGeometry, const FKeyEvent& InKeyEvent )
 {
 	if ( CommandList->ProcessCommandBindings( InKeyEvent ) )
diff --git a/Engine/Source/Editor/Kismet/Public/SSCSEditor.h b/Engine/Source/Editor/Kismet/Public/SSCSEditor.h
--- a/Engine/Source/Editor/Kismet/Public/SSCSEditor.h
+++ b/Engine/Source/Editor/Kismet/Public/SSCSEditor.h
@@ -417,12 +417,6 @@ public:
 	/** Returns true if in editing mode */
 	bool InEditingMode() const;
 
-	/**
-	 * Sets the Simple Construction Script this tree is visualizing
-	 * @param InSCS The SCS to edit
-	 */
-	void SetSCS(USimpleConstructionScript* InSCS);
-
 	/** Adds a component to the SCS Table
 	   @param NewComponentClass	(In) The class to add
 	   @param Asset				(In) Optional asset to assign to the component
diff --git a/Engine/Source/Editor/LevelEditor/Private/SLevelEditor.cpp b/Engine/Source/Editor/LevelEditor/Private/SLevelEditor.cpp
--- a/Engine/Source/Editor/LevelEditor/Private/SLevelEditor.cpp
+++ b/Engine/Source/Editor/LevelEditor/Private/SLevelEditor.cpp
@@ -519,6 +519,7 @@ public:
 
 		// Event subscriptions
 		USelection::SelectionChangedEvent.AddRaw(this, &SActorDetails::OnEditorSelectionChanged);
+		FCoreUObjectDelegates::OnObjectModified.AddRaw(this, &SActorDetails::OnObjectModified);
 		GEditor->OnComponentInstanceTransformed().AddRaw(this, &SActorDetails::OnComponentInstanceTransformed);
 		GEditor->OnPostTransformComponents().AddRaw(this, &SActorDetails::OnEditorPostTransformComponents);
 
@@ -554,6 +555,7 @@ public:
 	~SActorDetails()
 	{
 		USelection::SelectionChangedEvent.RemoveAll(this);
+		FCoreUObjectDelegates::OnObjectModified.RemoveAll(this);
 		GEditor->OnComponentInstanceTransformed().RemoveAll(this);
 		GEditor->OnPostTransformComponents().RemoveAll(this);
 	}
@@ -571,22 +573,18 @@ public:
 				auto Actor = GetSelectedActor();
 				if (Actor)
 				{
-					auto Blueprint = Cast<UBlueprint>(Actor->GetClass()->ClassGeneratedBy);
-					if (Blueprint && Blueprint->SimpleConstructionScript)
-					{
-						bShowingComponents = true;
-
-						ComponentsBox->ClearChildren();
-						ComponentsBox->AddSlot()
-						[
-							SAssignNew(SCSEditor, SSCSEditor)
-							.InEditingMode(true)
-							.ActorContext(this, &SActorDetails::GetSelectedActor)												// Get the instance of the actor in the world
-							.OnTreeViewSelectionChanged(this, &SActorDetails::OnSCSEditorTreeViewSelectionChanged)				// A selection has been made in the tree view, so inform the level editor
-							//.OnUpdateSelectionFromNodes(this, &SLevelEditor::OnSCSEditorUpdateSelectionFromNodes)				// Unsure, don't think it's needed
-							//.OnHighlightPropertyInDetailsView(this, &SLevelEditor::OnSCSEditorHighlightPropertyInDetailsView)	// Also unsure and don't think it's needed
-						];
-					}
+					bShowingComponents = true;
+
+					ComponentsBox->ClearChildren();
+					ComponentsBox->AddSlot()
+					[
+						SAssignNew(SCSEditor, SSCSEditor)
+						.InEditingMode(true)
+						.ActorContext(this, &SActorDetails::GetSelectedActor)												// Get the instance of the actor in the world
+						.OnTreeViewSelectionChanged(this, &SActorDetails::OnSCSEditorTreeViewSelectionChanged)				// A selection has been made in the tree view, so inform the level editor
+						//.OnUpdateSelectionFromNodes(this, &SLevelEditor::OnSCSEditorUpdateSelectionFromNodes)				// Unsure, don't think it's needed
+						//.OnHighlightPropertyInDetailsView(this, &SLevelEditor::OnSCSEditorHighlightPropertyInDetailsView)	// Also unsure and don't think it's needed
+					];
 				}
 			}
 
@@ -601,6 +599,42 @@ public:
 
 private:
 
+	void OnObjectModified(UObject* Object)
+	{
+		// If the modified object is a selected component, modify its corresponding template
+		auto Component = Cast<UActorComponent>(Object);
+		if (Component && GEditor->GetSelectedComponents()->IsSelected(Component))
+		{
+			auto SCSTreeNode = SCSEditor->GetNodeFromActorComponent(Component);
+			if (SCSTreeNode.IsValid())
+			{
+				auto ComponentTemplate = Cast<USceneComponent>(SCSTreeNode->GetComponentTemplate());
+				check(ComponentTemplate);
+
+				ComponentTemplate->SetFlags(RF_Transactional);
+				ComponentTemplate->Modify();
+			}
+		}
+	}
+
+	void OnEditorBeginObjectMovement(UObject& Object)
+	{
+		// If the object is a selected component, find the node and Modify() the corresponding template
+		auto Component = Cast<UActorComponent>(&Object);
+		if (Component)
+		{
+			auto SCSTreeNode = SCSEditor->GetNodeFromActorComponent(Component);
+			if (SCSTreeNode.IsValid())
+			{
+				auto ComponentTemplate = Cast<USceneComponent>(SCSTreeNode->GetComponentTemplate());
+				check(ComponentTemplate);
+
+				ComponentTemplate->SetFlags(RF_Transactional);
+				ComponentTemplate->Modify();
+			}
+		}
+	}
+
 	void OnComponentInstanceTransformed(USceneComponent& Component, const FVector& InDeltaDrag, const FRotator& InDeltaRot, const FVector& InDeltaScale)
 	{
 		// Locate and update the template component the corresponds to the transformed component instance
@@ -610,6 +644,15 @@ private:
 			auto ComponentTemplate = Cast<USceneComponent>(SCSTreeNode->GetComponentTemplate());
 			check(ComponentTemplate);
 
+			TSet<USceneComponent*> UpdatedComponents;
+			UpdatedComponents.Add(&Component);
+			UpdatedComponents.Add(ComponentTemplate);
+
+			// Cache the current default values for propagation
+			FVector OldRelativeLocation = ComponentTemplate->RelativeLocation;
+			FRotator OldRelativeRotation = ComponentTemplate->RelativeRotation;
+			FVector OldRelativeScale3D = ComponentTemplate->RelativeScale3D;
+
 			GEditor->ApplyDeltaToComponent(
 				ComponentTemplate,
 				true,
@@ -624,7 +667,7 @@ private:
 			auto Blueprint = Cast<UBlueprint>(PreviewActor->GetClass()->ClassGeneratedBy);
 			if (Blueprint != nullptr)
 			{
-				//@todo this doesn't include any of the "PropagateChange" logic/calls that are in FSCSEditorViewportClient::InputWidgetDelta()
+				//@todo this is virtually identical to FSCSEditorViewportClient::InputWidgetDelta(), should share more of it in FComponentEditorUtils
 
 				// Like PostEditMove(), but we only need to re-run construction scripts
 				if (Blueprint->bRunConstructionScriptOnDrag)
@@ -632,6 +675,50 @@ private:
 					PreviewActor->RerunConstructionScripts();
 				}
 
+				// If a constraint, copy back updated constraint frames to template
+				UPhysicsConstraintComponent* ConstraintComp = Cast<UPhysicsConstraintComponent>(SceneComp);
+				UPhysicsConstraintComponent* TemplateComp = Cast<UPhysicsConstraintComponent>(SelectedTemplate);
+				if (ConstraintComp && TemplateComp)
+				{
+					TemplateComp->ConstraintInstance.CopyConstraintGeometryFrom(&ConstraintComp->ConstraintInstance);
+				}
+
+				// Get the Blueprint class default object
+				AActor* CDO = nullptr;
+				if (Blueprint->GeneratedClass)
+				{
+					CDO = Cast<AActor>(Blueprint->GeneratedClass->ClassDefaultObject);
+				}
+				else if (Blueprint->ParentClass)
+				{
+					CDO = Cast<AActor>(Blueprint->ParentClass->ClassDefaultObject);
+				}
+
+				if (CDO != NULL)
+				{
+					// Iterate over all the active archetype instances and propagate the change(s) to the matching component instance
+					TArray<UObject*> ArchetypeInstances;
+					CDO->GetArchetypeInstances(ArchetypeInstances);
+					for (int32 InstanceIndex = 0; InstanceIndex < ArchetypeInstances.Num(); ++InstanceIndex)
+					{
+						AActor* ArchetypeInstance = Cast<AActor>(ArchetypeInstances[InstanceIndex]);
+						if (ArchetypeInstance != NULL)
+						{
+							/*const bool bIsProcessingPreviewActor = ( ArchetypeInstance == PreviewActor );*/
+							USceneComponent* SceneComp = Cast<USceneComponent>(SCSTreeNode->FindComponentInstanceInActor(ArchetypeInstance/*, bIsProcessingPreviewActor*/));
+							if (/*!bIsProcessingPreviewActor && */SceneComp != nullptr && !UpdatedComponents.Contains(SceneComp))
+							{
+								FComponentEditorUtils::PropagateTransformPropertyChange(SceneComp, SceneComp->RelativeLocation, OldRelativeLocation, ComponentTemplate->RelativeLocation, UpdatedComponents);
+								FComponentEditorUtils::PropagateTransformPropertyChange(SceneComp, SceneComp->RelativeRotation, OldRelativeRotation, ComponentTemplate->RelativeRotation, UpdatedComponents);
+								FComponentEditorUtils::PropagateTransformPropertyChange(SceneComp, SceneComp->RelativeScale3D, OldRelativeScale3D, ComponentTemplate->RelativeScale3D, UpdatedComponents);
+							}
+						}
+					}
+				}
+
+				// Enable the selection guard to prevent OnEditorSelectionChanged() from altering the contents of the SCSTreeWidget
+				TGuardValue<bool> SelectionGuard(bSelectionGuard, true);
+
 				// Update the component selection to the new reconstructed component
 				USelection* EditorComponentSelection = GEditor->GetSelectedComponents();
 				EditorComponentSelection->Modify();
@@ -640,7 +727,6 @@ private:
 				UActorComponent* ReconstructedComponent = SCSTreeNode->FindComponentInstanceInActor(PreviewActor);
 				check(ReconstructedComponent);
 				EditorComponentSelection->Select(ReconstructedComponent);
-				DetailsView->SetObject(ReconstructedComponent);
 			}
 		}
 	}
@@ -654,7 +740,6 @@ private:
 			TGuardValue<bool> SelectionGuard(bSelectionGuard, true);
 
 			// Update the editor's component selection
-			TArray<UObject*> Objects;
 			USelection* SelectedComponents = GEditor->GetSelectedComponents();
 
 			SelectedComponents->Modify();
@@ -669,16 +754,14 @@ private:
 			{
 				if (SelectedNode.IsValid())
 				{
-					UActorComponent* Component = SelectedNode->FindComponentInstanceInActor(SelectedActor);
-					if (Component)
+					UActorComponent* ComponentInstance = SelectedNode->FindComponentInstanceInActor(SelectedActor);
+					if (ComponentInstance)
 					{
-						Objects.Add(Component);
-						SelectedComponents->Select(Component);
+						// The instance is selected in the level editor
+						SelectedComponents->Select(ComponentInstance);
 					}
 				}
 			}
-			
-			DetailsView->SetObjects(Objects, false);
 			SelectedComponents->EndBatchSelectOperation();
 		}
 	}
@@ -687,9 +770,9 @@ private:
 	{
 		if (!bSelectionGuard && SCSEditor.IsValid())
 		{
-			TArray<UObject*> Objects;
+			TArray<UObject*> DetailsObjects;
 			auto Selection = Cast<USelection>(Object);
-			if (Selection == GEditor->GetSelectedComponents() || Selection == GEditor->GetSelectedActors())
+			if (Selection == GEditor->GetSelectedComponents() || (Selection == GEditor->GetSelectedActors() && GEditor->GetSelectedComponentCount() == 0))
 			{
 				// Enable the selection guard to prevent OnTreeSelectionChanged() from altering the editor's component selection
 				TGuardValue<bool> SelectionGuard(bSelectionGuard, true);
@@ -702,12 +785,16 @@ private:
 					UActorComponent* Component = CastChecked<UActorComponent>(*It);
 					if (Component)
 					{
-						SCSTreeWidget->SetItemSelection(SCSEditor->GetNodeFromActorComponent(Component), true);
-						Objects.Add(Component);
+						auto SCSTreeNode = SCSEditor->GetNodeFromActorComponent(Component);
+						if (SCSTreeNode.IsValid())
+						{
+							SCSTreeWidget->SetItemSelection(SCSTreeNode, true);
+							DetailsObjects.Add(SCSTreeNode->GetComponentTemplate());
+						}
 					}
 				}
 
-				DetailsView->SetObjects(Objects);
+				DetailsView->SetObjects(DetailsObjects);
 			}
 		}
 	}
@@ -720,7 +807,7 @@ private:
 
 	void OnSCSEditorTreeViewSelectionChanged(const TArray<FSCSEditorTreeNodePtrType>& SelectedNodes)
 	{
-		if (!bSelectionGuard/* && Actor != nullptr*/)
+		if (!bSelectionGuard)
 		{
 			auto Actor = GetSelectedActor();
 			if (Actor)
@@ -729,7 +816,7 @@ private:
 				TGuardValue<bool> SelectionGuard(bSelectionGuard, true);
 
 				// Update the editor's component selection
-				TArray<UObject*> Objects;
+				TArray<UObject*> DetailsObjects;
 				USelection* SelectedComponents = GEditor->GetSelectedComponents();
 
 				SelectedComponents->BeginBatchSelectOperation();
@@ -739,15 +826,19 @@ private:
 				{
 					if (SelectedNode.IsValid())
 					{
-						UActorComponent* Component = SelectedNode->FindComponentInstanceInActor(Actor);
-						if (Component)
+						UActorComponent* ComponentTemplate = SelectedNode->GetComponentTemplate();
+						UActorComponent* ComponentInstance = SelectedNode->FindComponentInstanceInActor(Actor);
+						if (ComponentTemplate && ComponentInstance)
 						{
-							Objects.Add(Component);
-							SelectedComponents->Select(Component);
+							// The details panel should edit the template
+							DetailsObjects.Add(ComponentTemplate);
+
+							// The instance is selected in the level editor
+							SelectedComponents->Select(ComponentInstance);
 						}
 					}
 				}
-				DetailsView->SetObjects(Objects, false);
+				DetailsView->SetObjects(DetailsObjects, false);
 
 				SelectedComponents->EndBatchSelectOperation();
 
diff --git a/Engine/Source/Editor/UnrealEd/Classes/Editor/EditorEngine.h b/Engine/Source/Editor/UnrealEd/Classes/Editor/EditorEngine.h
--- a/Engine/Source/Editor/UnrealEd/Classes/Editor/EditorEngine.h
+++ b/Engine/Source/Editor/UnrealEd/Classes/Editor/EditorEngine.h
@@ -2648,6 +2648,9 @@ private:
 	/** List of actors that were selected before Undo/redo */
 	TArray<AActor*> OldSelectedActors;
 
+	/** List of components that were selected before Undo/redo */
+	TArray<UActorComponent*> OldSelectedComponents;
+
 	/** The notification item to use for undo/redo */
 	TSharedPtr<class SNotificationItem> UndoRedoNotificationItem;
 
diff --git a/Engine/Source/Editor/UnrealEd/Private/EditorServer.cpp b/Engine/Source/Editor/UnrealEd/Private/EditorServer.cpp
--- a/Engine/Source/Editor/UnrealEd/Private/EditorServer.cpp
+++ b/Engine/Source/Editor/UnrealEd/Private/EditorServer.cpp
@@ -1048,12 +1048,19 @@ void UEditorEngine::HandleTransactorBeforeRedoUndo( FUndoSessionContext SessionC
 {
 	//Get the list of all selected actors before the undo/redo is performed
 	OldSelectedActors.Empty();
-
-	for ( FSelectionIterator It( GEditor->GetSelectedActorIterator() ) ; It ; ++It )
+	for ( FSelectionIterator It( GetSelectedActorIterator() ) ; It ; ++It )
 	{
 		AActor* Actor = CastChecked<AActor>( *It );
 		OldSelectedActors.Add( Actor);
 	}
+
+	// Get the list of selected components as well
+	OldSelectedComponents.Empty();
+	for (FSelectionIterator It(GetSelectedComponentIterator()); It; ++It)
+	{
+		auto Component = CastChecked<UActorComponent>(*It);
+		OldSelectedComponents.Add(Component);
+	}
 }
 
 void UEditorEngine::HandleTransactorRedo( FUndoSessionContext SessionContext, bool Succeeded )
@@ -1094,54 +1101,105 @@ UTransactor* UEditorEngine::CreateTrans()
 void UEditorEngine::PostUndo (bool bSuccess)
 {
 	//Make sure that the proper objects display as selected
-	//Get the list of all selected actors after the operation
-	TArray<AActor*> SelectedActors;
-	for ( FSelectionIterator It( GEditor->GetSelectedActorIterator() ) ; It ; ++It )
+	if (GetSelectedComponentCount() > 0)
 	{
-		AActor* Actor = CastChecked<AActor>( *It );
-		//if this actor is NOT in a hidden level add it to the list - otherwise de-select it
-		if( FLevelUtils::IsLevelLocked( Actor ) == false )
+		//@todo Check to see if component owner is in a hidden level
+		
+		// Get a list of all selected components after the operation
+		TArray<UActorComponent*> SelectedComponents;
+		for (FSelectionIterator It(GetSelectedComponentIterator()); It; ++It)
 		{
-			SelectedActors.Add( Actor);
+			SelectedComponents.Add(CastChecked<UActorComponent>(*It));
 		}
-		else
+		
+		USelection* Selection = GetSelectedComponents();
+		Selection->BeginBatchSelectOperation();
+
+		//Deselect all of the actors that were selected prior to the operation
+		for (int32 OldSelectedComponentIndex = OldSelectedComponents.Num() - 1; OldSelectedComponentIndex >= 0; --OldSelectedComponentIndex)
 		{
-			GetSelectedActors()->Select( Actor, false );				
+			UActorComponent* Component = OldSelectedComponents[OldSelectedComponentIndex];
+
+			//To stop us from unselecting and then reselecting again (causing two force update components, we will remove (from both lists) any object that was selected and should continue to be selected
+			int32 FoundIndex;
+			if (SelectedComponents.Find(Component, FoundIndex))
+			{
+				OldSelectedComponents.RemoveAt(OldSelectedComponentIndex);
+				SelectedComponents.RemoveAt(FoundIndex);
+			}
+			else
+			{
+				// Deselect without any notification
+				SelectComponent(Component, false, false);
+				//Actor->UpdateComponentTransforms();
+			}
 		}
-	}
 
-	USelection* Selection = GetSelectedActors();
-	Selection->BeginBatchSelectOperation();
+		//Select all of the components left in SelectedComponents
+		for (int32 SelectedComponentIndex = 0; SelectedComponentIndex < SelectedComponents.Num(); ++SelectedComponentIndex)
+		{
+			UActorComponent* Component = SelectedComponents[SelectedComponentIndex];
+			SelectComponent(Component, true, false);	//false is to stop notify which is done below if bOpWasSuccessful
+			//Actor->UpdateComponentTransforms();
+		}
 
-	//Deselect all of the actors that were selected prior to the operation
-	for ( int32 OldSelectedActorIndex = OldSelectedActors.Num()-1 ; OldSelectedActorIndex >= 0 ; --OldSelectedActorIndex )
+		OldSelectedComponents.Empty();
+
+		// We want to broadcast the component SelectionChangedEvent even if the selection didn't actually change
+		Selection->MarkBatchDirty();
+		Selection->EndBatchSelectOperation();
+	}
+	else
 	{
-		AActor* Actor = OldSelectedActors[OldSelectedActorIndex];
+		//Get the list of all selected actors after the operation
+		TArray<AActor*> SelectedActors;
+		for (FSelectionIterator It(GEditor->GetSelectedActorIterator()); It; ++It)
+		{
+			AActor* Actor = CastChecked<AActor>(*It);
+			//if this actor is NOT in a hidden level add it to the list - otherwise de-select it
+			if (FLevelUtils::IsLevelLocked(Actor) == false)
+			{
+				SelectedActors.Add(Actor);
+			}
+			else
+			{
+				GetSelectedActors()->Select(Actor, false);
+			}
+		}
+
+		USelection* Selection = GetSelectedActors();
+		Selection->BeginBatchSelectOperation();
 
-		//To stop us from unselecting and then reselecting again (causing two force update components, we will remove (from both lists) any object that was selected and should continue to be selected
-		int32 FoundIndex;
-		if (SelectedActors.Find(Actor, FoundIndex))
+		//Deselect all of the actors that were selected prior to the operation
+		for (int32 OldSelectedActorIndex = OldSelectedActors.Num() - 1; OldSelectedActorIndex >= 0; --OldSelectedActorIndex)
 		{
-			OldSelectedActors.RemoveAt(OldSelectedActorIndex);
-			SelectedActors.RemoveAt(FoundIndex);
+			AActor* Actor = OldSelectedActors[OldSelectedActorIndex];
+
+			//To stop us from unselecting and then reselecting again (causing two force update components, we will remove (from both lists) any object that was selected and should continue to be selected
+			int32 FoundIndex;
+			if (SelectedActors.Find(Actor, FoundIndex))
+			{
+				OldSelectedActors.RemoveAt(OldSelectedActorIndex);
+				SelectedActors.RemoveAt(FoundIndex);
+			}
+			else
+			{
+				SelectActor(Actor, false, false);//First false is to deselect, 2nd is to notify
+				Actor->UpdateComponentTransforms();
+			}
 		}
-		else
+
+		//Select all of the actors in SelectedActors
+		for (int32 SelectedActorIndex = 0; SelectedActorIndex < SelectedActors.Num(); ++SelectedActorIndex)
 		{
-			SelectActor( Actor, false, false );//First false is to deselect, 2nd is to notify
+			AActor* Actor = SelectedActors[SelectedActorIndex];
+			SelectActor(Actor, true, false);	//false is to stop notify which is done below if bOpWasSuccessful
 			Actor->UpdateComponentTransforms();
 		}
-	}
 
-	//Select all of the actors in SelectedActors
-	for ( int32 SelectedActorIndex = 0 ; SelectedActorIndex < SelectedActors.Num() ; ++SelectedActorIndex )
-	{
-		AActor* Actor = SelectedActors[SelectedActorIndex];
-		SelectActor( Actor, true, false );	//false is to stop notify which is done below if bOpWasSuccessful
-		Actor->UpdateComponentTransforms();
+		OldSelectedActors.Empty();
+		Selection->EndBatchSelectOperation();
 	}
-
-	OldSelectedActors.Empty();
-	Selection->EndBatchSelectOperation();
 }
 
 bool UEditorEngine::UndoTransaction()
diff --git a/Engine/Source/Editor/UnrealEd/Private/LevelEditorViewport.cpp b/Engine/Source/Editor/UnrealEd/Private/LevelEditorViewport.cpp
--- a/Engine/Source/Editor/UnrealEd/Private/LevelEditorViewport.cpp
+++ b/Engine/Source/Editor/UnrealEd/Private/LevelEditorViewport.cpp
@@ -1385,6 +1385,8 @@ void FTrackingTransaction::Begin(const FText& Description)
 	{
 		GroupActor->Modify();
 	}
+
+	// Modify selected components
 	for (FSelectionIterator It(GEditor->GetSelectedComponentIterator()); It; ++It)
 	{
 		CastChecked<UActorComponent>(*It)->Modify();
@@ -2392,49 +2394,58 @@ void FLevelEditorViewportClient::TrackingStarted( const FInputEventState& InInpu
 	}
 
 	PreDragActorTransforms.Empty();
-	for (FSelectionIterator It(GEditor->GetSelectedActorIterator()); It && !bIsTrackingBrushModification; ++It)
+	if (GEditor->GetSelectedComponentCount() > 0)
 	{
-		AActor* Actor = static_cast<AActor*>(*It);
-		checkSlow( Actor->IsA(AActor::StaticClass()) );
-
-		if( bIsDraggingWidget )
+		if (bIsDraggingWidget)
 		{
-			// Notify that this actor is beginning to move
-			GEditor->BroadcastBeginObjectMovement( *Actor );
+			Widget->SetSnapEnabled(true);
 		}
+	}
+	else
+	{
+		for (FSelectionIterator It(GEditor->GetSelectedActorIterator()); It && !bIsTrackingBrushModification; ++It)
+		{
+			AActor* Actor = static_cast<AActor*>( *It );
+			checkSlow(Actor->IsA(AActor::StaticClass()));
 
-		Widget->SetSnapEnabled(true);
+			if (bIsDraggingWidget)
+			{
+				// Notify that this actor is beginning to move
+				GEditor->BroadcastBeginObjectMovement(*Actor);
+			}
 
-		// See if any brushes are about to be transformed via their Widget
-		TArray<AActor*> AttachedActors;
-		Actor->GetAttachedActors( AttachedActors );
-		const bool bExactClass = true;
-		// First, check for selected brush actors, check the actors attached actors for brush actors as well.  If a parent actor moves, the bsp needs to be rebuilt
-		ABrush* Brush = Cast< ABrush >( Actor );
-		if (Brush && (!Brush->IsVolumeBrush() && !FActorEditorUtils::IsABuilderBrush(Actor)))
-		{
-			bIsTrackingBrushModification = true;
-		}
-		else // Next, check for selected groups actors that contain brushes
-		{
-			AGroupActor* GroupActor = Cast<AGroupActor>(Actor);
-			if (GroupActor)
+			Widget->SetSnapEnabled(true);
+
+			// See if any brushes are about to be transformed via their Widget
+			TArray<AActor*> AttachedActors;
+			Actor->GetAttachedActors(AttachedActors);
+			const bool bExactClass = true;
+			// First, check for selected brush actors, check the actors attached actors for brush actors as well.  If a parent actor moves, the bsp needs to be rebuilt
+			ABrush* Brush = Cast< ABrush >(Actor);
+			if (Brush && ( !Brush->IsVolumeBrush() && !FActorEditorUtils::IsABuilderBrush(Actor) ))
 			{
-				TArray<AActor*> GroupMembers;
-				GroupActor->GetAllChildren(GroupMembers, true);
-				for (int32 GroupMemberIdx = 0; GroupMemberIdx < GroupMembers.Num(); ++GroupMemberIdx)
+				bIsTrackingBrushModification = true;
+			}
+			else // Next, check for selected groups actors that contain brushes
+			{
+				AGroupActor* GroupActor = Cast<AGroupActor>(Actor);
+				if (GroupActor)
 				{
-					Brush = Cast< ABrush >( GroupMembers[GroupMemberIdx] );
-					if ( Brush && (!Brush->IsVolumeBrush() && !FActorEditorUtils::IsABuilderBrush(Actor)))
+					TArray<AActor*> GroupMembers;
+					GroupActor->GetAllChildren(GroupMembers, true);
+					for (int32 GroupMemberIdx = 0; GroupMemberIdx < GroupMembers.Num(); ++GroupMemberIdx)
 					{
-						bIsTrackingBrushModification = true;
+						Brush = Cast< ABrush >(GroupMembers[GroupMemberIdx]);
+						if (Brush && ( !Brush->IsVolumeBrush() && !FActorEditorUtils::IsABuilderBrush(Actor) ))
+						{
+							bIsTrackingBrushModification = true;
+						}
 					}
 				}
 			}
 		}
 	}
 
-
 	// Start a transformation transaction if required
 	if( !bTrackingHandledExternally )
 	{
