Fixing build failure

diff --git a/Engine/Source/Editor/LevelEditor/Private/SLevelEditor.cpp b/Engine/Source/Editor/LevelEditor/Private/SLevelEditor.cpp
--- a/Engine/Source/Editor/LevelEditor/Private/SLevelEditor.cpp
+++ b/Engine/Source/Editor/LevelEditor/Private/SLevelEditor.cpp
@@ -584,6 +584,8 @@ public:
 						.OnTreeViewSelectionChanged(this, &SActorDetails::OnSCSEditorTreeViewSelectionChanged)				// A selection has been made in the tree view, so inform the level editor
 						//.OnUpdateSelectionFromNodes(this, &SLevelEditor::OnSCSEditorUpdateSelectionFromNodes)				// Unsure, don't think it's needed
 						//.OnHighlightPropertyInDetailsView(this, &SLevelEditor::OnSCSEditorHighlightPropertyInDetailsView)	// Also unsure and don't think it's needed
+						//.OnAddNewComponent()
+						//.OnAddExistingComponent()
 					];
 				}
 			}
@@ -676,8 +678,8 @@ private:
 				}
 
 				// If a constraint, copy back updated constraint frames to template
-				UPhysicsConstraintComponent* ConstraintComp = Cast<UPhysicsConstraintComponent>(SceneComp);
-				UPhysicsConstraintComponent* TemplateComp = Cast<UPhysicsConstraintComponent>(SelectedTemplate);
+				UPhysicsConstraintComponent* ConstraintComp = Cast<UPhysicsConstraintComponent>(&Component);
+				UPhysicsConstraintComponent* TemplateComp = Cast<UPhysicsConstraintComponent>(ComponentTemplate);
 				if (ConstraintComp && TemplateComp)
 				{
 					TemplateComp->ConstraintInstance.CopyConstraintGeometryFrom(&ConstraintComp->ConstraintInstance);
@@ -848,6 +850,52 @@ private:
 			}
 		}
 	}
+	
+
+	//USCS_Node* OnSCSEditorAddNewComponent(UClass* InComponentClass)
+	//{
+	//	check(InComponentClass != nullptr);
+
+	//	//UBlueprint* Blueprint = GetBlueprintObj();
+	//	//@todo If the selected actor is only an actor, convert to a BP now and use the SCS to create the node
+	//	
+	//	// Use the creation portion of CreateAnonymousBlueprintAndSpawnActor():
+	//	/*
+	//	UClass* ActorClass = ActorClassType::StaticClass();
+	//	FName UniqueActorName = MakeUniqueObjectName(InLevel, ActorClass, Name);
+
+	//	// Spawn a blank actor with an anonymous blueprint
+	//	// Create and init a new Blueprint
+	//	const FName BPName = *FString::Printf(TEXT("%s_BPClass"), *UniqueActorName.ToString());
+
+	//	const FName UniqueBPName = MakeUniqueObjectName(InLevel, UBlueprint::StaticClass(), BPName);
+
+
+	//	UBlueprint* NewBP = FKismetEditorUtilities::CreateBlueprint(ActorClass, InLevel, UniqueBPName, BPTYPE_Normal, UBlueprint::StaticClass(), UBlueprintGeneratedClass::StaticClass());
+	//	*/
+
+	//	check(Blueprint != nullptr && Blueprint->SimpleConstructionScript != nullptr);
+
+	//	return Blueprint->SimpleConstructionScript->CreateNode(InComponentClass);
+	//}
+
+	//USCS_Node* OnSCSEditorAddExistingComponent(UActorComponent* InComponentInstance)
+	//{
+	//	check(InComponentInstance != nullptr);
+
+	//	//UBlueprint* Blueprint = GetBlueprintObj();
+	//	//@todo If the selected actor is only an actor, convert to a BP now and use the SCS to create the node
+
+	//	check(Blueprint != nullptr && Blueprint->SimpleConstructionScript != nullptr);
+
+	//	if (InComponentInstance->GetOuter() == GetTransientPackage())
+	//	{
+	//		// Relocate the instance from the transient package to the BPGC and assign it a unique object name
+	//		InComponentInstance->Rename(NULL, Blueprint->GeneratedClass, REN_DontCreateRedirectors | REN_DoNotDirty);
+	//	}
+
+	//	return Blueprint->SimpleConstructionScript->CreateNode(InComponentInstance);
+	//}
 
 	TSharedPtr<SSplitter> DetailsSplitter;
 	TSharedPtr<IDetailsView> DetailsView;
