Editor - Continuing to adjust padding in several places to slim down the editor some.

diff --git a/Engine/Source/Editor/LevelEditor/Private/LevelEditor.cpp b/Engine/Source/Editor/LevelEditor/Private/LevelEditor.cpp
--- a/Engine/Source/Editor/LevelEditor/Private/LevelEditor.cpp
+++ b/Engine/Source/Editor/LevelEditor/Private/LevelEditor.cpp
@@ -42,7 +42,10 @@ FLevelEditorModule::FLevelEditorModule()
 
 TSharedRef<SDockTab> FLevelEditorModule::SpawnLevelEditor( const FSpawnTabArgs& InArgs )
 {
-	TSharedRef<SDockTab> LevelEditorTab = SNew(SDockTab) .TabRole(ETabRole::MajorTab) .ContentPadding( FMargin(0,2,0,0) );
+	TSharedRef<SDockTab> LevelEditorTab = SNew(SDockTab)
+		.TabRole(ETabRole::MajorTab)
+		.ContentPadding( FMargin(0) );
+
 	SetLevelEditorInstanceTab(LevelEditorTab);
 	TSharedPtr< SWindow > OwnerWindow = InArgs.GetOwnerWindow();
 	
diff --git a/Engine/Source/Editor/LevelEditor/Private/SLevelEditor.cpp b/Engine/Source/Editor/LevelEditor/Private/SLevelEditor.cpp
--- a/Engine/Source/Editor/LevelEditor/Private/SLevelEditor.cpp
+++ b/Engine/Source/Editor/LevelEditor/Private/SLevelEditor.cpp
@@ -1131,7 +1131,6 @@ TSharedRef<SDockTab> SLevelEditor::SpawnLevelEditorTab( const FSpawnTabArgs& Arg
 		return SNew( SDockTab )
 			.Icon( FEditorStyle::GetBrush( "LevelEditor.Tabs.StatsViewer" ) )
 			.Label( NSLOCTEXT("LevelEditor", "StatsViewerTabTitle", "Statistics") )
-			.ContentPadding( 5 )
 			[
 				StatsViewerModule.CreateStatsViewer()					
 			];
diff --git a/Engine/Source/Editor/UnrealEd/Private/Toolkits/ToolkitManager.cpp b/Engine/Source/Editor/UnrealEd/Private/Toolkits/ToolkitManager.cpp
--- a/Engine/Source/Editor/UnrealEd/Private/Toolkits/ToolkitManager.cpp
+++ b/Engine/Source/Editor/UnrealEd/Private/Toolkits/ToolkitManager.cpp
@@ -11,7 +11,7 @@ static TSharedRef<SDockTab> SpawnStandaloneToolkitHost( const FSpawnTabArgs& Arg
 	return
 		SNew(SDockTab)
 		.TabRole(ETabRole::MajorTab)
-		.ContentPadding(FMargin(0,2,0,0));
+		.ContentPadding(FMargin(0));
 }
 
 static TSharedPtr<SDockTab> NeverReuse(const FTabId&)
