- Fixed undo bug when translating components in the BP editor (instance on preview actor will now properly return to the old transform and stay in sync with the component template)
- Removed CDO-editing capabilities of in-world component editing
- Level details SCS tree now refreshes instead of being reconstructed
- Updated SCSEditor::GetNodeFromActorComponent to work in ActorInstance mode
- Cannot select hidden editor-only components in the world

diff --git a/Engine/Source/Editor/Kismet/Private/SCSEditorViewportClient.cpp b/Engine/Source/Editor/Kismet/Private/SCSEditorViewportClient.cpp
--- a/Engine/Source/Editor/Kismet/Private/SCSEditorViewportClient.cpp
+++ b/Engine/Source/Editor/Kismet/Private/SCSEditorViewportClient.cpp
@@ -832,11 +832,12 @@ void FSCSEditorViewportClient::BeginTransaction(const FText& Description)
 	{
 		ScopedTransaction = new FScopedTransaction(Description);
 
-		if(PreviewBlueprint != NULL)
+		if (PreviewBlueprint != nullptr)
 		{
 			FBlueprintEditorUtils::MarkBlueprintAsModified(PreviewBlueprint);
 		}
 
+		AActor* PreviewActor = GetPreviewActor();
 		TArray<FSCSEditorTreeNodePtrType> SelectedNodes = BlueprintEditorPtr.Pin()->GetSelectedSCSEditorTreeNodes();
 		if(SelectedNodes.Num() > 0)
 		{
@@ -850,12 +851,23 @@ void FSCSEditorViewportClient::BeginTransaction(const FText& Description)
 						SCS_Node->Modify();
 					}
 
+					// Modify both the component template and the instance in the preview actor (provided there is one)
 					UActorComponent* ComponentTemplate = Node->GetComponentTemplate();
-					if(ComponentTemplate != NULL)
+					if (ComponentTemplate != nullptr)
 					{
 						ComponentTemplate->SetFlags(RF_Transactional);
 						ComponentTemplate->Modify();
 					}
+
+					if (PreviewActor)
+					{
+						UActorComponent* ComponentPreviewInstance = Node->FindComponentInstanceInActor(PreviewActor);
+						if (ComponentPreviewInstance != nullptr)
+						{
+							ComponentPreviewInstance->SetFlags(RF_Transactional);
+							ComponentPreviewInstance->Modify();
+						}
+					}
 				}
 			}
 		}
diff --git a/Engine/Source/Editor/Kismet/Private/SSCSEditor.cpp b/Engine/Source/Editor/Kismet/Private/SSCSEditor.cpp
--- a/Engine/Source/Editor/Kismet/Private/SSCSEditor.cpp
+++ b/Engine/Source/Editor/Kismet/Private/SSCSEditor.cpp
@@ -2446,54 +2446,57 @@ FSCSEditorTreeNodePtrType SSCSEditor::GetNodeFromActorComponent(const UActorComp
 
 	if(ActorComponent)
 	{
-		// If the given component instance is not already an archetype object
-		if(!ActorComponent->IsTemplate())
+		if (EditorMode == EEditorMode::BlueprintSCS)
 		{
-			// Get the component owner's class object
-			check(ActorComponent->GetOwner() != NULL);
-			UClass* OwnerClass = ActorComponent->GetOwner()->GetActorClass();
-
-			// If the given component is one that's created during Blueprint construction
-			if(ActorComponent->bCreatedByConstructionScript)
+			// If the given component instance is not already an archetype object
+			if (!ActorComponent->IsTemplate())
 			{
-				// Get the Blueprint object associated with the owner's class
-				UBlueprint* Blueprint = UBlueprint::GetBlueprintFromClass(OwnerClass);
-				if(Blueprint && Blueprint->SimpleConstructionScript)
+				// Get the component owner's class object
+				check(ActorComponent->GetOwner() != NULL);
+				UClass* OwnerClass = ActorComponent->GetOwner()->GetActorClass();
+
+				// If the given component is one that's created during Blueprint construction
+				if (ActorComponent->bCreatedByConstructionScript)
 				{
-					// Attempt to locate an SCS node with a variable name that matches the name of the given component
-					TArray<USCS_Node*> AllNodes = Blueprint->SimpleConstructionScript->GetAllNodes();
-					for(int32 i = 0; i < AllNodes.Num(); ++i)
+					// Get the Blueprint object associated with the owner's class
+					UBlueprint* Blueprint = UBlueprint::GetBlueprintFromClass(OwnerClass);
+					if (Blueprint && Blueprint->SimpleConstructionScript)
 					{
-						USCS_Node* SCS_Node = AllNodes[i];
-						
-						check(SCS_Node != NULL);
-						if(SCS_Node->VariableName == ActorComponent->GetFName())
+						// Attempt to locate an SCS node with a variable name that matches the name of the given component
+						TArray<USCS_Node*> AllNodes = Blueprint->SimpleConstructionScript->GetAllNodes();
+						for (int32 i = 0; i < AllNodes.Num(); ++i)
 						{
-							// We found a match; redirect to the component archetype instance that may be associated with a tree node
-							ActorComponent = SCS_Node->ComponentTemplate;
-							break;
+							USCS_Node* SCS_Node = AllNodes[i];
+
+							check(SCS_Node != NULL);
+							if (SCS_Node->VariableName == ActorComponent->GetFName())
+							{
+								// We found a match; redirect to the component archetype instance that may be associated with a tree node
+								ActorComponent = SCS_Node->ComponentTemplate;
+								break;
+							}
 						}
 					}
 				}
-			}
-			else
-			{
-				// Get the class default object
-				const AActor* CDO = Cast<AActor>(OwnerClass->GetDefaultObject());
-				if(CDO)
+				else
 				{
-					// Iterate over the Components array and attempt to find a component with a matching name
-					TInlineComponentArray<UActorComponent*> Components;
-					CDO->GetComponents(Components);
-
-					for(auto It = Components.CreateConstIterator(); It; ++It)
+					// Get the class default object
+					const AActor* CDO = Cast<AActor>(OwnerClass->GetDefaultObject());
+					if (CDO)
 					{
-						UActorComponent* ComponentTemplate = *It;
-						if(ComponentTemplate->GetFName() == ActorComponent->GetFName())
+						// Iterate over the Components array and attempt to find a component with a matching name
+						TInlineComponentArray<UActorComponent*> Components;
+						CDO->GetComponents(Components);
+
+						for (auto It = Components.CreateConstIterator(); It; ++It)
 						{
-							// We found a match; redirect to the component archetype instance that may be associated with a tree node
-							ActorComponent = ComponentTemplate;
-							break;
+							UActorComponent* ComponentTemplate = *It;
+							if (ComponentTemplate->GetFName() == ActorComponent->GetFName())
+							{
+								// We found a match; redirect to the component archetype instance that may be associated with a tree node
+								ActorComponent = ComponentTemplate;
+								break;
+							}
 						}
 					}
 				}
@@ -2501,7 +2504,7 @@ FSCSEditorTreeNodePtrType SSCSEditor::GetNodeFromActorComponent(const UActorComp
 		}
 
 		// If we have a valid component archetype instance, attempt to find a tree node that corresponds to it
-		if(ActorComponent->IsTemplate())
+		if((EditorMode == EEditorMode::BlueprintSCS && ActorComponent->IsTemplate()) || EditorMode == EEditorMode::ActorInstance)
 		{
 			for(int32 i = 0; i < RootNodes.Num() && !NodePtr.IsValid(); i++)
 			{
diff --git a/Engine/Source/Editor/LevelEditor/Private/SLevelEditor.cpp b/Engine/Source/Editor/LevelEditor/Private/SLevelEditor.cpp
--- a/Engine/Source/Editor/LevelEditor/Private/SLevelEditor.cpp
+++ b/Engine/Source/Editor/LevelEditor/Private/SLevelEditor.cpp
@@ -519,9 +519,6 @@ public:
 
 		// Event subscriptions
 		USelection::SelectionChangedEvent.AddRaw(this, &SActorDetails::OnEditorSelectionChanged);
-		FCoreUObjectDelegates::OnObjectModified.AddRaw(this, &SActorDetails::OnObjectModified);
-		GEditor->OnComponentInstanceTransformed().AddRaw(this, &SActorDetails::OnComponentInstanceTransformed);
-		GEditor->OnPostTransformComponents().AddRaw(this, &SActorDetails::OnEditorPostTransformComponents);
 
 		FPropertyEditorModule& PropPlugin = FModuleManager::LoadModuleChecked<FPropertyEditorModule>("PropertyEditor");
 		FDetailsViewArgs DetailsViewArgs;
@@ -543,6 +540,16 @@ public:
 		SAssignNew(ComponentsBox, SBox)
 		.Visibility(EVisibility::Collapsed);
 
+		ComponentsBox->SetContent
+		(
+			SAssignNew(SCSEditor, SSCSEditor)
+			.EditorMode(SSCSEditor::EEditorMode::ActorInstance)
+			.ActorContext(this, &SActorDetails::GetSelectedActor)												// Get the instance of the actor in the world
+			.OnTreeViewSelectionChanged(this, &SActorDetails::OnSCSEditorTreeViewSelectionChanged)				// A selection has been made in the tree view, so inform the level editor
+			//.OnUpdateSelectionFromNodes(this, &SLevelEditor::OnSCSEditorUpdateSelectionFromNodes)				// Unsure, don't think it's needed
+			//.OnHighlightPropertyInDetailsView(this, &SLevelEditor::OnSCSEditorHighlightPropertyInDetailsView)	// Also unsure and don't think it's needed
+		);
+
 		ChildSlot
 		[
 			SNew( SVerticalBox )
@@ -579,9 +586,6 @@ public:
 	~SActorDetails()
 	{
 		USelection::SelectionChangedEvent.RemoveAll(this);
-		FCoreUObjectDelegates::OnObjectModified.RemoveAll(this);
-		GEditor->OnComponentInstanceTransformed().RemoveAll(this);
-		GEditor->OnPostTransformComponents().RemoveAll(this);
 	}
 
 	void SetObjects( const TArray<UObject*>& InObjects )
@@ -599,17 +603,11 @@ public:
 				{
 					bShowingComponents = true;
 
-					ComponentsBox->SetContent
-					(
-						SAssignNew(SCSEditor, SSCSEditor)
-						.EditorMode(SSCSEditor::EEditorMode::ActorInstance)
-						.ActorContext(this, &SActorDetails::GetSelectedActor)												// Get the instance of the actor in the world
-						.OnTreeViewSelectionChanged(this, &SActorDetails::OnSCSEditorTreeViewSelectionChanged)				// A selection has been made in the tree view, so inform the level editor
-						//.OnUpdateSelectionFromNodes(this, &SLevelEditor::OnSCSEditorUpdateSelectionFromNodes)				// Unsure, don't think it's needed
-						//.OnHighlightPropertyInDetailsView(this, &SLevelEditor::OnSCSEditorHighlightPropertyInDetailsView)	// Also unsure and don't think it's needed
-						//.OnAddNewComponent()
-						//.OnAddExistingComponent()
-					);
+					// Update the tree if a new actor is selected
+					if (GEditor->GetSelectedComponentCount() == 0)
+					{
+						SCSEditor->UpdateTree();
+					}
 				}
 			}
 
@@ -617,180 +615,8 @@ public:
 		}
 	}
 
-	const TSharedPtr<SSCSEditor>& GetSCSEditor() const
-	{
-		return SCSEditor;
-	}
-
 private:
 
-	void OnObjectModified(UObject* Object)
-	{
-		// If the modified object is a selected component, modify its corresponding template
-		auto Component = Cast<UActorComponent>(Object);
-		if (Component && GEditor->GetSelectedComponents()->IsSelected(Component))
-		{
-			auto SCSTreeNode = SCSEditor->GetNodeFromActorComponent(Component);
-			if (SCSTreeNode.IsValid())
-			{
-				auto ComponentTemplate = Cast<USceneComponent>(SCSTreeNode->GetComponentTemplate());
-				check(ComponentTemplate);
-
-				ComponentTemplate->SetFlags(RF_Transactional);
-				ComponentTemplate->Modify();
-			}
-		}
-	}
-
-	void OnEditorBeginObjectMovement(UObject& Object)
-	{
-		// If the object is a selected component, find the node and Modify() the corresponding template
-		auto Component = Cast<UActorComponent>(&Object);
-		if (Component)
-		{
-			auto SCSTreeNode = SCSEditor->GetNodeFromActorComponent(Component);
-			if (SCSTreeNode.IsValid())
-			{
-				auto ComponentTemplate = Cast<USceneComponent>(SCSTreeNode->GetComponentTemplate());
-				check(ComponentTemplate);
-
-				ComponentTemplate->SetFlags(RF_Transactional);
-				ComponentTemplate->Modify();
-			}
-		}
-	}
-
-	void OnComponentInstanceTransformed(USceneComponent& Component, const FVector& InDeltaDrag, const FRotator& InDeltaRot, const FVector& InDeltaScale)
-	{
-		// Locate and update the template component the corresponds to the transformed component instance
-		auto SCSTreeNode = SCSEditor->GetNodeFromActorComponent(&Component);
-		if (SCSTreeNode.IsValid())
-		{
-			auto ComponentTemplate = Cast<USceneComponent>(SCSTreeNode->GetComponentTemplate());
-			check(ComponentTemplate);
-
-			TSet<USceneComponent*> UpdatedComponents;
-			UpdatedComponents.Add(&Component);
-			UpdatedComponents.Add(ComponentTemplate);
-
-			// Cache the current default values for propagation
-			FVector OldRelativeLocation = ComponentTemplate->RelativeLocation;
-			FRotator OldRelativeRotation = ComponentTemplate->RelativeRotation;
-			FVector OldRelativeScale3D = ComponentTemplate->RelativeScale3D;
-
-			GEditor->ApplyDeltaToComponent(
-				ComponentTemplate,
-				true,
-				&InDeltaDrag,
-				&InDeltaRot,
-				&InDeltaScale,
-				GEditor->GetPivotLocation());
-
-			auto PreviewActor = GetSelectedActor();
-			check(PreviewActor);
-
-			auto Blueprint = Cast<UBlueprint>(PreviewActor->GetClass()->ClassGeneratedBy);
-			if (Blueprint != nullptr)
-			{
-				//@todo this is virtually identical to FSCSEditorViewportClient::InputWidgetDelta(), should share more of it in FComponentEditorUtils
-
-				// Like PostEditMove(), but we only need to re-run construction scripts
-				if (Blueprint->bRunConstructionScriptOnDrag)
-				{
-					PreviewActor->RerunConstructionScripts();
-				}
-
-				// If a constraint, copy back updated constraint frames to template
-				UPhysicsConstraintComponent* ConstraintComp = Cast<UPhysicsConstraintComponent>(&Component);
-				UPhysicsConstraintComponent* TemplateComp = Cast<UPhysicsConstraintComponent>(ComponentTemplate);
-				if (ConstraintComp && TemplateComp)
-				{
-					TemplateComp->ConstraintInstance.CopyConstraintGeometryFrom(&ConstraintComp->ConstraintInstance);
-				}
-
-				// Get the Blueprint class default object
-				AActor* CDO = nullptr;
-				if (Blueprint->GeneratedClass)
-				{
-					CDO = Cast<AActor>(Blueprint->GeneratedClass->ClassDefaultObject);
-				}
-				else if (Blueprint->ParentClass)
-				{
-					CDO = Cast<AActor>(Blueprint->ParentClass->ClassDefaultObject);
-				}
-
-				if (CDO != NULL)
-				{
-					// Iterate over all the active archetype instances and propagate the change(s) to the matching component instance
-					TArray<UObject*> ArchetypeInstances;
-					CDO->GetArchetypeInstances(ArchetypeInstances);
-					for (int32 InstanceIndex = 0; InstanceIndex < ArchetypeInstances.Num(); ++InstanceIndex)
-					{
-						AActor* ArchetypeInstance = Cast<AActor>(ArchetypeInstances[InstanceIndex]);
-						if (ArchetypeInstance != NULL)
-						{
-							/*const bool bIsProcessingPreviewActor = ( ArchetypeInstance == PreviewActor );*/
-							USceneComponent* SceneComp = Cast<USceneComponent>(SCSTreeNode->FindComponentInstanceInActor(ArchetypeInstance/*, bIsProcessingPreviewActor*/));
-							if (/*!bIsProcessingPreviewActor && */SceneComp != nullptr && !UpdatedComponents.Contains(SceneComp))
-							{
-								FComponentEditorUtils::PropagateTransformPropertyChange(SceneComp, SceneComp->RelativeLocation, OldRelativeLocation, ComponentTemplate->RelativeLocation, UpdatedComponents);
-								FComponentEditorUtils::PropagateTransformPropertyChange(SceneComp, SceneComp->RelativeRotation, OldRelativeRotation, ComponentTemplate->RelativeRotation, UpdatedComponents);
-								FComponentEditorUtils::PropagateTransformPropertyChange(SceneComp, SceneComp->RelativeScale3D, OldRelativeScale3D, ComponentTemplate->RelativeScale3D, UpdatedComponents);
-							}
-						}
-					}
-				}
-
-				// Enable the selection guard to prevent OnEditorSelectionChanged() from altering the contents of the SCSTreeWidget
-				TGuardValue<bool> SelectionGuard(bSelectionGuard, true);
-
-				// Update the component selection to the new reconstructed component
-				USelection* EditorComponentSelection = GEditor->GetSelectedComponents();
-				EditorComponentSelection->Modify();
-				EditorComponentSelection->Deselect(&Component);
-				
-				UActorComponent* ReconstructedComponent = SCSTreeNode->FindComponentInstanceInActor(PreviewActor);
-				check(ReconstructedComponent);
-				EditorComponentSelection->Select(ReconstructedComponent);
-			}
-		}
-	}
-
-	void OnEditorPostTransformComponents()
-	{
-		// Update the component selection with the reconstructed components
-		if (!bSelectionGuard)
-		{
-			// Enable the selection guard to prevent OnEditorSelectionChanged() from altering the contents of the SCSTreeWidget
-			TGuardValue<bool> SelectionGuard(bSelectionGuard, true);
-
-			// Update the editor's component selection
-			USelection* SelectedComponents = GEditor->GetSelectedComponents();
-
-			SelectedComponents->Modify();
-			SelectedComponents->BeginBatchSelectOperation();
-			SelectedComponents->DeselectAll();
-
-			auto SelectedActor = GetSelectedActor();
-			check(SelectedActor);
-
-			auto SelectedNodes = SCSEditor->GetSelectedNodes();
-			for (auto& SelectedNode : SelectedNodes)
-			{
-				if (SelectedNode.IsValid())
-				{
-					UActorComponent* ComponentInstance = SelectedNode->FindComponentInstanceInActor(SelectedActor);
-					if (ComponentInstance)
-					{
-						// The instance is selected in the level editor
-						SelectedComponents->Select(ComponentInstance);
-					}
-				}
-			}
-			SelectedComponents->EndBatchSelectOperation();
-		}
-	}
-
 	void OnEditorSelectionChanged(UObject* Object)
 	{
 		if (!bSelectionGuard && SCSEditor.IsValid())
@@ -814,7 +640,8 @@ private:
 						if (SCSTreeNode.IsValid())
 						{
 							SCSTreeWidget->SetItemSelection(SCSTreeNode, true);
-							DetailsObjects.Add(SCSTreeNode->GetComponentTemplate());
+							check(Component == SCSTreeNode->GetComponentTemplate());
+							DetailsObjects.Add(Component);
 						}
 					}
 				}
@@ -851,14 +678,10 @@ private:
 				{
 					if (SelectedNode.IsValid())
 					{
-						UActorComponent* ComponentTemplate = SelectedNode->GetComponentTemplate();
 						UActorComponent* ComponentInstance = SelectedNode->FindComponentInstanceInActor(Actor);
-						if (ComponentTemplate && ComponentInstance)
+						if (ComponentInstance)
 						{
-							// The details panel should edit the template
-							DetailsObjects.Add(ComponentTemplate);
-
-							// The instance is selected in the level editor
+							DetailsObjects.Add(ComponentInstance);
 							SelectedComponents->Select(ComponentInstance);
 						}
 					}
@@ -874,60 +697,12 @@ private:
 		}
 	}
 	
-
-	//USCS_Node* OnSCSEditorAddNewComponent(UClass* InComponentClass)
-	//{
-	//	check(InComponentClass != nullptr);
-
-	//	//UBlueprint* Blueprint = GetBlueprintObj();
-	//	//@todo If the selected actor is only an actor, convert to a BP now and use the SCS to create the node
-	//	
-	//	// Use the creation portion of CreateAnonymousBlueprintAndSpawnActor():
-	//	/*
-	//	UClass* ActorClass = ActorClassType::StaticClass();
-	//	FName UniqueActorName = MakeUniqueObjectName(InLevel, ActorClass, Name);
-
-	//	// Spawn a blank actor with an anonymous blueprint
-	//	// Create and init a new Blueprint
-	//	const FName BPName = *FString::Printf(TEXT("%s_BPClass"), *UniqueActorName.ToString());
-
-	//	const FName UniqueBPName = MakeUniqueObjectName(InLevel, UBlueprint::StaticClass(), BPName);
-
-
-	//	UBlueprint* NewBP = FKismetEditorUtilities::CreateBlueprint(ActorClass, InLevel, UniqueBPName, BPTYPE_Normal, UBlueprint::StaticClass(), UBlueprintGeneratedClass::StaticClass());
-	//	*/
-
-	//	check(Blueprint != nullptr && Blueprint->SimpleConstructionScript != nullptr);
-
-	//	return Blueprint->SimpleConstructionScript->CreateNode(InComponentClass);
-	//}
-
-	//USCS_Node* OnSCSEditorAddExistingComponent(UActorComponent* InComponentInstance)
-	//{
-	//	check(InComponentInstance != nullptr);
-
-	//	//UBlueprint* Blueprint = GetBlueprintObj();
-	//	//@todo If the selected actor is only an actor, convert to a BP now and use the SCS to create the node
-
-	//	check(Blueprint != nullptr && Blueprint->SimpleConstructionScript != nullptr);
-
-	//	if (InComponentInstance->GetOuter() == GetTransientPackage())
-	//	{
-	//		// Relocate the instance from the transient package to the BPGC and assign it a unique object name
-	//		InComponentInstance->Rename(NULL, Blueprint->GeneratedClass, REN_DontCreateRedirectors | REN_DoNotDirty);
-	//	}
-
-	//	return Blueprint->SimpleConstructionScript->CreateNode(InComponentInstance);
-	//}
-
+private:
 	TSharedPtr<SSplitter> DetailsSplitter;
 	TSharedPtr<IDetailsView> DetailsView;
 	TSharedPtr<SBox> ComponentsBox;
-	
-	/** SCS editor */
 	TSharedPtr<SSCSEditor> SCSEditor;
 
-	bool bTreeVisible;
 	bool bSelectionGuard;
 };
 
diff --git a/Engine/Source/Editor/UnrealEd/Classes/Editor/EditorEngine.h b/Engine/Source/Editor/UnrealEd/Classes/Editor/EditorEngine.h
--- a/Engine/Source/Editor/UnrealEd/Classes/Editor/EditorEngine.h
+++ b/Engine/Source/Editor/UnrealEd/Classes/Editor/EditorEngine.h
@@ -589,14 +589,6 @@ public:
 	DECLARE_EVENT_OneParam( UEditorEngine, FOnEndTransformObject, UObject& );
 	FOnEndTransformObject& OnEndObjectMovement() { return OnEndObjectTransformEvent; }
 
-	/** Editor-only event triggered when a component instance in the world is transformed */
-	DECLARE_EVENT_FourParams(UEditorEngine, FOnComponentInstanceTransformed, USceneComponent&, const FVector&, const FRotator&, const FVector&);
-	FOnComponentInstanceTransformed& OnComponentInstanceTransformed() { return OnComponentInstanceTransformedEvent; }
-
-	/** Editor-only event triggered when transformation of a component in the world finishes */
-	DECLARE_EVENT(UEditorEngine, FOnPostTransformComponents);
-	FOnPostTransformComponents& OnPostTransformComponents() { return OnPostTransformComponentsEvent; }
-
 	/** Delegate broadcast by the engine every tick when PIE/SIE is active, to check to see whether we need to
 		be able to capture state for simulating actor (for Sequencer recording features).  The single bool parameter
 		should be set to true if recording features are needed. */
@@ -617,26 +609,6 @@ public:
 	*/
 	void BroadcastEndObjectMovement(UObject& Object) const { OnEndObjectTransformEvent.Broadcast(Object); }
 
-	//@todo dhertzka - this will fail when multiple components are selected; should just send the transform data and update each component based on GetSelectedComponents
-	/**
-	* Called when a component instance in the world has been translated, rotated, or scaled
-	*
-	* @param Component	The component instance the world that was transformed
-	* @param InDeltaDrag
-	* @param InDeltaRot
-	* @param InDeltaScale
-	*/
-	void BroadcastComponentInstanceTransformed(USceneComponent& Component, const FVector& InDeltaDrag, const FRotator& InDeltaRot, const FVector& InDeltaScale) const
-	{
-		OnComponentInstanceTransformedEvent.Broadcast(Component, InDeltaDrag, InDeltaRot, InDeltaScale);
-	}
-
-	/**
-	* Called after transformation of any component(s) in the world is complete
-	*/
-	void BroadcastPostTransformComponents() { OnPostTransformComponentsEvent.Broadcast(); }
-
-
 	/**	Broadcasts that an object has been reimported. THIS SHOULD NOT BE PUBLIC */
 	void BroadcastObjectReimported(UObject* InObject);
 
@@ -2626,12 +2598,6 @@ private:
 	/** Delegate broadcast when an actor or component has been moved, rotated, or scaled */
 	FOnEndTransformObject OnEndObjectTransformEvent;
 
-	/** Delegate broadcast when a component instance in the world has been moved, rotated, or scaled */
-	FOnComponentInstanceTransformed OnComponentInstanceTransformedEvent;
-
-	/** Delegate broadcast when transformation of any component(s) in the world is complete */
-	FOnPostTransformComponents OnPostTransformComponentsEvent;
-
 	/** Delegate broadcast by the engine every tick when PIE/SIE is active, to check to see whether we need to
 		be able to capture state for simulating actor (for Sequencer recording features) */
 	FGetActorRecordingState GetActorRecordingStateEvent;
diff --git a/Engine/Source/Editor/UnrealEd/Private/LevelEditorViewport.cpp b/Engine/Source/Editor/UnrealEd/Private/LevelEditorViewport.cpp
--- a/Engine/Source/Editor/UnrealEd/Private/LevelEditorViewport.cpp
+++ b/Engine/Source/Editor/UnrealEd/Private/LevelEditorViewport.cpp
@@ -1789,6 +1789,7 @@ void FLevelEditorViewportClient::ProcessClick(FSceneView& View, HHitProxy* HitPr
 			// 3. The LMB was pressed or we already have a component selected (so that right-clicking a selected actor won't select a component)
 			// 4. The click was not a double click (unless a component has already been selected)
 			// 5. The level viewport didn't just receive focus this frame (again unless a component has already been selected)
+			// 6. The component selected is not a hidden editor-only component
 			const bool bActorAlreadySelectedExclusively = GEditor->GetSelectedActors()->IsSelected(ActorHitProxy->Actor) && ( GEditor->GetSelectedActorCount() == 1 );
 			const bool bComponentAlreadySelected = GEditor->GetSelectedComponentCount() > 0;
 	
@@ -1796,7 +1797,10 @@ void FLevelEditorViewportClient::ProcessClick(FSceneView& View, HHitProxy* HitPr
 													&& (Click.GetEvent() != IE_DoubleClick) 
 													&& !bReceivedFocusRecently;
 
-			const bool bSelectComponent = bActorAlreadySelectedExclusively && ( bComponentAlreadySelected || bCanBeginSelectingComponents );
+			auto HitComponent = ActorHitProxy->PrimComponent;
+			const bool bHitSelectableComponent = HitComponent != nullptr && HitComponent->AlwaysLoadOnClient && HitComponent->AlwaysLoadOnServer;
+
+			const bool bSelectComponent = bActorAlreadySelectedExclusively && bHitSelectableComponent && ( bComponentAlreadySelected || bCanBeginSelectingComponents );
 
 			if (bSelectComponent && GetDefault<UEditorExperimentalSettings>()->bInWorldBPEditing)
 			{
@@ -2393,8 +2397,9 @@ void FLevelEditorViewportClient::TrackingStarted( const FInputEventState& InInpu
 		}
 	}
 
+	const bool bIsDraggingComponents = GEditor->GetSelectedComponentCount() > 0;
 	PreDragActorTransforms.Empty();
-	if (GEditor->GetSelectedComponentCount() > 0)
+	if (bIsDraggingComponents)
 	{
 		if (bIsDraggingWidget)
 		{
@@ -2453,26 +2458,27 @@ void FLevelEditorViewportClient::TrackingStarted( const FInputEventState& InInpu
 		{
 			TrackingTransaction.TransCount++;
 
+			FText ObjectTypeBeingTracked = bIsDraggingComponents ? LOCTEXT("TransactionFocus_Components", "Components") : LOCTEXT("TransactionFocus_Actors", "Actors");
 			FText TrackingDescription;
 
 			switch( GetWidgetMode() )
 			{
 			case FWidget::WM_Translate:
-				TrackingDescription = LOCTEXT("MoveActorsTransaction", "Move Actors");
+				TrackingDescription = FText::Format(LOCTEXT("MoveTransaction", "Move {0}"), ObjectTypeBeingTracked);
 				break;
 			case FWidget::WM_Rotate:
-				TrackingDescription = LOCTEXT("RotateActorsTransaction", "Rotate Actors");
+				TrackingDescription = FText::Format(LOCTEXT("RotateTransaction", "Rotate {0}"), ObjectTypeBeingTracked);
 				break;
 			case FWidget::WM_Scale:
-				TrackingDescription = LOCTEXT("ScaleActorsTransaction", "Scale Actors");
+				TrackingDescription = FText::Format(LOCTEXT("ScaleTransaction", "Scale {0}"), ObjectTypeBeingTracked);
 				break;
 			case FWidget::WM_TranslateRotateZ:
-				TrackingDescription = LOCTEXT("TranslateRotateZActorsTransaction", "Translate/RotateZ Actors");
+				TrackingDescription = FText::Format(LOCTEXT("TranslateRotateZTransaction", "Translate/RotateZ {0}"), ObjectTypeBeingTracked);
 				break;
 			default:
 				if( bNudge )
 				{
-					TrackingDescription = LOCTEXT("NudgeActorsTransaction", "Nudge Actors");
+					TrackingDescription = FText::Format(LOCTEXT("NudgeTransaction", "Nudge {0}"), ObjectTypeBeingTracked);
 				}
 			}
 
@@ -2543,9 +2549,6 @@ void FLevelEditorViewportClient::TrackingStopped()
 	// Don't do this if AddDelta was never called.
 	if( bDidAnythingActuallyChange && MouseDeltaTracker->HasReceivedDelta() )
 	{
-		// If components are selected, we'll need to notify after rerunning the construction script on the selected actor(s)
-		const bool bComponentsAreSelected = GEditor->GetSelectedComponentCount() > 0;
-
 		for (FSelectionIterator It(GEditor->GetSelectedActorIterator()); It; ++It)
 		{
 			AActor* Actor = static_cast<AActor*>( *It );
@@ -2572,11 +2575,6 @@ void FLevelEditorViewportClient::TrackingStopped()
 			GEditor->BroadcastEndObjectMovement(*Actor);
 		}
 
-		if (bComponentsAreSelected)
-		{
-			GEditor->BroadcastPostTransformComponents();
-		}
-
 		if (!bPivotMovedIndependently)
 		{
 			GUnrealEd->UpdatePivotLocationForSelection();
@@ -3059,8 +3057,6 @@ void FLevelEditorViewportClient::ApplyDeltaToComponent(USceneComponent* InCompon
 		&InDeltaRot,
 		&ModifiedDeltaScale,
 		GEditor->GetPivotLocation());
-
-	GEditor->BroadcastComponentInstanceTransformed(*InComponent, InDeltaDrag, InDeltaRot, ModifiedDeltaScale);
 }
 
 /** Helper function for ModifyScale - Convert the active Dragging Axis to per-axis flags */
diff --git a/Engine/Source/Editor/UnrealEd/Private/LevelViewportClickHandlers.cpp b/Engine/Source/Editor/UnrealEd/Private/LevelViewportClickHandlers.cpp
--- a/Engine/Source/Editor/UnrealEd/Private/LevelViewportClickHandlers.cpp
+++ b/Engine/Source/Editor/UnrealEd/Private/LevelViewportClickHandlers.cpp
@@ -269,7 +269,25 @@ namespace ClickHandlers
 		//@todo hotkeys for component placement?
 		
 		bool bComponentClicked = false;
-		UActorComponent* Component = const_cast<UPrimitiveComponent*>(ActorHitProxy->PrimComponent); 
+
+		USceneComponent* Component = nullptr;
+
+		// Find the component in the actor that matches the PrimComponent on the hit proxy
+		TInlineComponentArray<USceneComponent*> SceneComponents;
+		ActorHitProxy->Actor->GetComponents(SceneComponents);
+		for (auto CompIt = SceneComponents.CreateConstIterator(); CompIt; ++CompIt)
+		{
+			auto SceneComp = *CompIt;
+			if (SceneComp == ActorHitProxy->PrimComponent)
+			{
+				Component = SceneComp;
+			}
+		}
+		
+		if (!Component)
+		{
+			return false;
+		}
 		
 		// Pivot snapping
 		if (Click.GetKey() == EKeys::MiddleMouseButton && Click.IsAltDown())
