Fixed duplicate node creation/selection of newly added components for in-world component editing

diff --git a/Engine/Source/Editor/Kismet/Private/SSCSEditor.cpp b/Engine/Source/Editor/Kismet/Private/SSCSEditor.cpp
--- a/Engine/Source/Editor/Kismet/Private/SSCSEditor.cpp
+++ b/Engine/Source/Editor/Kismet/Private/SSCSEditor.cpp
@@ -2761,6 +2761,7 @@ void SSCSEditor::UpdateTree(bool bRegenerateTreeNodes)
 		TArray<FSCSEditorTreeNodePtrType> SelectedTreeNodes = SCSTreeWidget->GetSelectedItems();
 
 		// Clear the current tree
+		SCSTreeWidget->ClearSelection();
 		RootNodes.Empty();
 
 		// Reset the scene root node
diff --git a/Engine/Source/Editor/LevelEditor/Private/SLevelEditor.cpp b/Engine/Source/Editor/LevelEditor/Private/SLevelEditor.cpp
--- a/Engine/Source/Editor/LevelEditor/Private/SLevelEditor.cpp
+++ b/Engine/Source/Editor/LevelEditor/Private/SLevelEditor.cpp
@@ -623,7 +623,7 @@ private:
 		{
 			TArray<UObject*> DetailsObjects;
 			auto Selection = Cast<USelection>(Object);
-			if (Selection == GEditor->GetSelectedComponents() || (Selection == GEditor->GetSelectedActors() && GEditor->GetSelectedComponentCount() == 0))
+			if (Selection == GEditor->GetSelectedComponents())
 			{
 				// Enable the selection guard to prevent OnTreeSelectionChanged() from altering the editor's component selection
 				TGuardValue<bool> SelectionGuard(bSelectionGuard, true);
@@ -659,7 +659,7 @@ private:
 
 	void OnSCSEditorTreeViewSelectionChanged(const TArray<FSCSEditorTreeNodePtrType>& SelectedNodes)
 	{
-		if (!bSelectionGuard)
+		if (!bSelectionGuard && SelectedNodes.Num() > 0 )
 		{
 			auto Actor = GetSelectedActor();
 			if (Actor)
@@ -686,7 +686,7 @@ private:
 						}
 					}
 				}
-				DetailsView->SetObjects(DetailsObjects, false);
+				DetailsView->SetObjects(DetailsObjects);
 
 				SelectedComponents->EndBatchSelectOperation();
 
