Fixed bug in source control toolbar icon selection.

diff --git a/Engine/Source/Editor/LevelEditor/Private/LevelEditorToolBar.cpp b/Engine/Source/Editor/LevelEditor/Private/LevelEditorToolBar.cpp
--- a/Engine/Source/Editor/LevelEditor/Private/LevelEditorToolBar.cpp
+++ b/Engine/Source/Editor/LevelEditor/Private/LevelEditorToolBar.cpp
@@ -1211,12 +1211,7 @@ TSharedRef< SWidget > FLevelEditorToolBar::MakeLevelEditorToolBar( const TShared
 
 			static FText GetSourceControlTooltip()
 			{
-				if (QueryState == EQueryState::NotQueried)
-				{
-					CheckSourceControlStatus();
-					return LOCTEXT("CheckingSourceControl", "Checking source control status...");
-				}
-				else if (QueryState == EQueryState::Querying)
+				if (QueryState == EQueryState::Querying)
 				{
 					return LOCTEXT("SourceControlUnknown", "Source control status is unknown");
 				}
@@ -1228,7 +1223,11 @@ TSharedRef< SWidget > FLevelEditorToolBar::MakeLevelEditorToolBar( const TShared
 
 			static FSlateIcon GetSourceControlIcon()
 			{
-				if (QueryState == EQueryState::Queried)
+				if (QueryState == EQueryState::Querying)
+				{
+					return FSlateIcon(FEditorStyle::GetStyleSetName(), "LevelEditor.SourceControl.Unknown");
+				}
+				else
 				{
 					ISourceControlModule& SourceControlModule = ISourceControlModule::Get();
 					if (SourceControlModule.IsEnabled())
@@ -1247,15 +1246,6 @@ TSharedRef< SWidget > FLevelEditorToolBar::MakeLevelEditorToolBar( const TShared
 						return FSlateIcon(FEditorStyle::GetStyleSetName(), "LevelEditor.SourceControl.Off");
 					}
 				}
-				else
-				{
-					if (QueryState == EQueryState::NotQueried)
-					{
-						CheckSourceControlStatus();
-					}
-
-					return FSlateIcon(FEditorStyle::GetStyleSetName(), "LevelEditor.SourceControl.Unknown");
-				}
 			}
 		};
 
