Added a one pixel high row between the tab bar and tab contents in the level editor on Mac

diff --git a/Engine/Source/Editor/LevelEditor/Private/SLevelEditor.cpp b/Engine/Source/Editor/LevelEditor/Private/SLevelEditor.cpp
--- a/Engine/Source/Editor/LevelEditor/Private/SLevelEditor.cpp
+++ b/Engine/Source/Editor/LevelEditor/Private/SLevelEditor.cpp
@@ -180,6 +180,16 @@ void SLevelEditor::Initialize( const TSharedRef<SDockTab>& OwnerTab, const TShar
 #endif
 		]
 
+#if PLATFORM_MAC
+		// Without the in-window menu bar, we need some space between the tab bar and tab contents
+		+SVerticalBox::Slot()
+		.AutoHeight()
+		[
+			SNew( SBox )
+			.HeightOverride( 1.0f )
+		]
+#endif
+
 		+SVerticalBox::Slot()
 		.FillHeight( 1.0f )
 		[
