Editor: Add a tooltip with version information, branch name, etc... to the level editor watermark in the main frame

diff --git a/Engine/Source/Editor/LevelEditor/Private/LevelEditor.cpp b/Engine/Source/Editor/LevelEditor/Private/LevelEditor.cpp
--- a/Engine/Source/Editor/LevelEditor/Private/LevelEditor.cpp
+++ b/Engine/Source/Editor/LevelEditor/Private/LevelEditor.cpp
@@ -37,8 +37,67 @@ FLevelEditorModule::FLevelEditorModule()
 	TEXT( "LevelEditor.ToggleImmersive" ),
 	TEXT( "Toggle 'Immersive Mode' for the active level editing viewport" ),
 	FConsoleCommandDelegate::CreateRaw( this, &FLevelEditorModule::ToggleImmersiveOnActiveLevelViewport ) )
+{
+}
+
+
+class SLevelEditorWatermark : public SCompoundWidget
+{
+public:
+	SLATE_BEGIN_ARGS(SLevelEditorWatermark) {}
+	SLATE_END_ARGS()
+
+	void Construct(const FArguments& InArgs)
 	{
+		FString OptionalBranchPrefix;
+		GConfig->GetString(TEXT("LevelEditor"), TEXT("ProjectNameWatermarkPrefix"), /*out*/ OptionalBranchPrefix, GEditorUserSettingsIni);
+
+		FFormatNamedArguments Args;
+		Args.Add(TEXT("Branch"), FText::FromString(OptionalBranchPrefix));
+		Args.Add(TEXT("GameName"), FText::FromString(FString(FApp::GetGameName())));
+
+		FText RightContentText;
+		FText RightContentTooltip;
+
+		const EBuildConfigurations::Type BuildConfig = FApp::GetBuildConfiguration();
+		if (BuildConfig != EBuildConfigurations::Shipping && BuildConfig != EBuildConfigurations::Development && BuildConfig != EBuildConfigurations::Unknown)
+		{
+			Args.Add(TEXT("Config"), EBuildConfigurations::ToText(BuildConfig));
+			RightContentText = FText::Format(NSLOCTEXT("UnrealEditor", "TitleBarRightContentAndConfig", "{Branch}{GameName} [{Config}]"), Args);
+		}
+		else
+		{
+			RightContentText = FText::Format(NSLOCTEXT("UnrealEditor", "TitleBarRightContent", "{Branch}{GameName}"), Args);
+		}
+
+		// Create the tooltip showing more detailed information
+		FFormatNamedArguments TooltipArgs;
+		const FString EngineVerisonString = GEngineVersion.ToString(GEngineVersion.IsPromotedBuild() ? EVersionComponent::Changelist : EVersionComponent::Patch);
+		TooltipArgs.Add(TEXT("Version"), FText::FromString(EngineVerisonString));
+		TooltipArgs.Add(TEXT("Branch"), FText::FromString(FApp::GetBranchName()));
+		TooltipArgs.Add(TEXT("BuildConfiguration"), EBuildConfigurations::ToText(BuildConfig));
+		TooltipArgs.Add(TEXT("BuildDate"), FText::FromString(FApp::GetBuildDate()));
+		RightContentTooltip = FText::Format(NSLOCTEXT("UnrealEditor", "TitleBarRightContentTooltip", "Version: {Version}\nBranch: {Branch}\nBuild Configuration: {BuildConfiguration}\nBuild Date: {BuildDate}"), TooltipArgs);
+
+		SetToolTipText(RightContentTooltip);
+
+		ChildSlot
+		[
+			SNew(STextBlock)
+			.Text(RightContentText)
+			.Visibility(EVisibility::HitTestInvisible)
+			.Font(FSlateFontInfo(FPaths::EngineContentDir() / TEXT("Slate/Fonts/Roboto-Bold.ttf"), 14))
+			.ColorAndOpacity(FLinearColor(1.0f, 1.0f, 1.0f, 0.3f))
+		];
+	}
+	
+	// SWidget interface
+	virtual EWindowZone::Type GetWindowZoneOverride() const override
+	{
+		return EWindowZone::TitleBar;
 	}
+	// End of SWidget interface
+};
 
 TSharedRef<SDockTab> FLevelEditorModule::SpawnLevelEditor( const FSpawnTabArgs& InArgs )
 {
@@ -74,40 +133,13 @@ TSharedRef<SDockTab> FLevelEditorModule::SpawnLevelEditor( const FSpawnTabArgs&
 
 	TSharedPtr< SWidget > RightContent;
 	{
-		FString OptionalBranchPrefix;
-		GConfig->GetString(TEXT("LevelEditor"), TEXT("ProjectNameWatermarkPrefix"), /*out*/ OptionalBranchPrefix, GEditorUserSettingsIni);
-
-		FFormatNamedArguments Args;
-		Args.Add( TEXT("Branch"), FText::FromString(OptionalBranchPrefix) );
-		Args.Add( TEXT("GameName"), FText::FromString(FString(FApp::GetGameName())) );
-
-		FText RightContentText;
-
-		const EBuildConfigurations::Type BuildConfig = FApp::GetBuildConfiguration();
-		if (BuildConfig != EBuildConfigurations::Shipping && BuildConfig != EBuildConfigurations::Development && BuildConfig != EBuildConfigurations::Unknown)
-		{
-			Args.Add( TEXT("Config"), EBuildConfigurations::ToText(BuildConfig) );
-			RightContentText = FText::Format(NSLOCTEXT("UnrealEditor", "TitleBarRightContentAndConfig", "{Branch}{GameName} [{Config}]"), Args);
-		}
-		else
-		{
-			RightContentText = FText::Format(NSLOCTEXT("UnrealEditor", "TitleBarRightContent", "{Branch}{GameName}"), Args);
-		}
-
 		RightContent =
 				SNew( SHorizontalBox )
 
 				+SHorizontalBox::Slot()
 				.AutoWidth()
 				[
-					SNew(SBox)
-					.Visibility( EVisibility::HitTestInvisible )
-					[
-						SNew( STextBlock )
-						.Text( RightContentText )
-						.Font( FSlateFontInfo( FPaths::EngineContentDir() / TEXT("Slate/Fonts/Roboto-Bold.ttf"), 14 ) )
-						.ColorAndOpacity( FLinearColor( 1.0f, 1.0f, 1.0f, 0.3f ) )
-					]
+					SNew(SLevelEditorWatermark)
 				]
 // Put the level editor stats/notification widgets on the main window title bar since we don't have a menu bar on OS X
 #if PLATFORM_MAC
