Changes to how favorite levels list is organized.
#change Changed favorites menu to not have an upper limit.  Favorites now organized so the most recently added/opened is at the head of the list.  Only the first five favorite levels are displayed on the main menu, the others are on a submenu.

diff --git a/Engine/Source/Editor/EditorStyle/Private/SlateEditorStyle.cpp b/Engine/Source/Editor/EditorStyle/Private/SlateEditorStyle.cpp
--- a/Engine/Source/Editor/EditorStyle/Private/SlateEditorStyle.cpp
+++ b/Engine/Source/Editor/EditorStyle/Private/SlateEditorStyle.cpp
@@ -1304,6 +1304,7 @@ void FSlateEditorStyle::FStyle::SetupGeneralStyles()
 		Set( "MainFrame.PackageProject", new IMAGE_BRUSH( "Icons/icon_package_16x", Icon16x16 ) );
 		Set( "MainFrame.RecentProjects", new IMAGE_BRUSH( "Icons/icon_file_ProjectsRecent_16px", Icon16x16 ) );
 		Set( "MainFrame.RecentLevels", new IMAGE_BRUSH( "Icons/icon_file_LevelsRecent_16px", Icon16x16 ) );
+		Set( "MainFrame.FavoriteLevels", new IMAGE_BRUSH( "Icons/icon_file_LevelsRecent_16px", Icon16x16 ) );
 
 		Set( "MainFrame.DebugTools.SmallFont", TTF_CORE_FONT( "Fonts/Roboto-Regular", 8 ) );
 		Set( "MainFrame.DebugTools.NormalFont", TTF_CORE_FONT( "Fonts/Roboto-Regular", 9 ) );
diff --git a/Engine/Source/Editor/LevelEditor/Private/LevelEditor.cpp b/Engine/Source/Editor/LevelEditor/Private/LevelEditor.cpp
--- a/Engine/Source/Editor/LevelEditor/Private/LevelEditor.cpp
+++ b/Engine/Source/Editor/LevelEditor/Private/LevelEditor.cpp
@@ -561,16 +561,6 @@ void FLevelEditorModule::BindGlobalLevelEditorCommands()
 		ActionList.MapAction( Commands.OpenRecentFileCommands[ CurRecentIndex ], FExecuteAction::CreateStatic( &FLevelEditorActionCallbacks::OpenRecentFile, CurRecentIndex ), DefaultExecuteAction );
 	}
 
-	for( int32 CurFavoriteIndex = 0; CurFavoriteIndex < FLevelEditorCommands::MaxFavoriteFiles; ++CurFavoriteIndex )
-	{
-		ActionList.MapAction( Commands.OpenFavoriteFileCommands[ CurFavoriteIndex ], FExecuteAction::CreateStatic( &FLevelEditorActionCallbacks::OpenFavoriteFile, CurFavoriteIndex ), DefaultExecuteAction );
-	}
-
-	for( int32 CurFavoriteIndex = 0; CurFavoriteIndex < FLevelEditorCommands::MaxFavoriteFiles; ++CurFavoriteIndex )
-	{
-		ActionList.MapAction( Commands.RemoveFavoriteCommands[ CurFavoriteIndex ], FExecuteAction::CreateStatic( &FLevelEditorActionCallbacks::RemoveFavorite, CurFavoriteIndex ), DefaultExecuteAction );
-	}
-
 	ActionList.MapAction( Commands.Import,
 		FExecuteAction::CreateStatic( &FLevelEditorActionCallbacks::Import_Clicked ) );
 
diff --git a/Engine/Source/Editor/LevelEditor/Private/LevelEditorActions.cpp b/Engine/Source/Editor/LevelEditor/Private/LevelEditorActions.cpp
--- a/Engine/Source/Editor/LevelEditor/Private/LevelEditorActions.cpp
+++ b/Engine/Source/Editor/LevelEditor/Private/LevelEditorActions.cpp
@@ -250,6 +250,9 @@ void FLevelEditorActionCallbacks::OpenFavoriteFile( int32 FavoriteFileIndex )
 		{
 			// Load the requested level.
 			FEditorFileUtils::LoadMap( FileName );
+
+			// Move the item to the head of the list
+			MRUFavoritesList->MoveFavoritesItemToHead( FileName );
 		}
 		else
 		{
@@ -2658,36 +2661,6 @@ void FLevelEditorCommands::RegisterCommands()
 		OpenRecentFileCommands.Add( OpenRecentFile );
 	}
 
-	for( int32 CurFavoriteIndex = 0; CurFavoriteIndex < FLevelEditorCommands::MaxFavoriteFiles; ++CurFavoriteIndex )
-	{
-		// NOTE: The actual label and tool-tip will be overridden at runtime when the command is bound to a menu item, however
-		// we still need to set one here so that the key bindings UI can function properly
-		TSharedRef< FUICommandInfo > OpenFavoriteFile =
-			FUICommandInfoDecl(
-				this->AsShared(),
-				FName( *FString::Printf( TEXT( "OpenFavoriteFile%i" ), CurFavoriteIndex ) ),
-				FText::Format( NSLOCTEXT( "LevelEditorCommands", "OpenFavoriteFile", "Open Favorite File {0}" ), FText::AsNumber( CurFavoriteIndex ) ),
-				NSLOCTEXT( "LevelEditorCommands", "OpenFavoriteFileTaggedToolTip", "Opens a file that was tagged as a favorite" ) )
-			.UserInterfaceType( EUserInterfaceActionType::Button )
-			.DefaultGesture( FInputGesture() );
-		OpenFavoriteFileCommands.Add( OpenFavoriteFile );
-	}
-
-	for( int32 CurFavoriteIndex = 0; CurFavoriteIndex < FLevelEditorCommands::MaxFavoriteFiles; ++CurFavoriteIndex )
-	{
-		// NOTE: The actual label and tool-tip will be overridden at runtime when the command is bound to a menu item, however
-		// we still need to set one here so that the key bindings UI can function properly
-		TSharedRef< FUICommandInfo > RemoveFavorite =
-			FUICommandInfoDecl(
-			this->AsShared(),
-			FName( *FString::Printf( TEXT( "RemoveFavorite%i" ), CurFavoriteIndex ) ),
-			FText::Format( NSLOCTEXT( "LevelEditorCommands", "RemoveFavorite", "Remove Favorite {0}" ), FText::AsNumber( CurFavoriteIndex ) ),
-			NSLOCTEXT( "LevelEditorCommands", "RemoveFavoriteToolTip", "Removes an entry from the favorites list" ) )
-			.UserInterfaceType( EUserInterfaceActionType::Button )
-			.DefaultGesture( FInputGesture() );
-		RemoveFavoriteCommands.Add( RemoveFavorite );
-	}
-
 	UI_COMMAND( Import, "Import...", "Imports objects and actors from a T3D format into the current level", EUserInterfaceActionType::Button, FInputGesture() );
 	UI_COMMAND( ExportAll, "Export All...", "Exports the entire level to a file on disk (multiple formats are supported.)", EUserInterfaceActionType::Button, FInputGesture() );
 	UI_COMMAND( ExportSelected, "Export Selected...", "Exports currently-selected objects to a file on disk (multiple formats are supported.)", EUserInterfaceActionType::Button, FInputGesture() );
diff --git a/Engine/Source/Editor/LevelEditor/Private/LevelEditorMenu.cpp b/Engine/Source/Editor/LevelEditor/Private/LevelEditorMenu.cpp
--- a/Engine/Source/Editor/LevelEditor/Private/LevelEditorMenu.cpp
+++ b/Engine/Source/Editor/LevelEditor/Private/LevelEditorMenu.cpp
@@ -143,34 +143,6 @@ TSharedRef< SWidget > FLevelEditorMenu::MakeLevelEditorMenu( const TSharedPtr<FU
 			MenuBuilder.AddMenuEntry( FLevelEditorCommands::Get().SaveAllLevels );
 		}
 
-		static void FillFavoriteLevelSubMenu( FMenuBuilder& MenuBuilder, const int32 CurFavoriteIndex )
-		{
-			IMainFrameModule& MainFrameModule = FModuleManager::LoadModuleChecked<IMainFrameModule>( "MainFrame" );
-			const FMainMRUFavoritesList& RecentsAndFavorites = *MainFrameModule.GetMRUFavoritesList();
-
-			MenuBuilder.BeginSection("FavoritesOptions", LOCTEXT("FavoriteOptionsHeading", "Level Favorite Options") );
-			{
-				TSharedPtr< FUICommandInfo > OpenFavoriteFile = FLevelEditorCommands::Get().OpenFavoriteFileCommands[ CurFavoriteIndex ];
-				const FString CurFavorite = RecentsAndFavorites.GetFavoritesItem( CurFavoriteIndex );
-				const FText CurFavoriteText = FText::FromString( RecentsAndFavorites.GetFavoritesItem( CurFavoriteIndex ) );
-				const FText CurBasename = FText::FromString( FPaths::GetBaseFilename(CurFavorite) );
-				MenuBuilder.AddMenuEntry( 
-					OpenFavoriteFile, 
-					NAME_None, 
-					CurBasename, 
-					FText::Format( LOCTEXT("OpenFavoriteFileToolTip", "Open favorite file: {0}"), CurFavoriteText ) );
-
-				TSharedPtr< FUICommandInfo > RemoveFavoriteFile = FLevelEditorCommands::Get().RemoveFavoriteCommands[ CurFavoriteIndex ];
-				MenuBuilder.AddMenuEntry( 
-					RemoveFavoriteFile, 
-					NAME_None, 
-					FText::Format( LOCTEXT("ToggleFavorite_Remove", "Remove %s from Favorites"), CurBasename ), 
-					FText::Format( LOCTEXT("RemoveFavoriteToolTip", "Remove the level %s from your list of Favorites"), CurFavoriteText )  );
-			}
-			MenuBuilder.EndSection();
-		}
-
-
 		static void FillFileRecentAndFavoriteFileItems( FMenuBuilder& MenuBuilder )
 		{
 			IMainFrameModule& MainFrameModule = FModuleManager::LoadModuleChecked<IMainFrameModule>( "MainFrame" );
@@ -194,20 +166,19 @@ TSharedRef< SWidget > FLevelEditorMenu::MakeLevelEditorMenu( const TSharedPtr<FU
 
 			// Favorite files
 			{
+				static const int32 FavoritesToDisplayInMainMenu = 5;
 				const int32 NumFavorites = RecentsAndFavorites.GetNumFavorites();
 
 				MenuBuilder.BeginSection("FileFavorites", LOCTEXT("FavoriteFilesHeading", "Favorites") );
 
 				if( NumFavorites > 0 )
 				{
-					// Our UI only supports displaying a certain number of favorite items
-					const int32 AllowedFavorites = FMath::Min( NumFavorites, FLevelEditorCommands::Get().OpenFavoriteFileCommands.Num() );
-					for( int32 CurFavoriteIndex = 0; CurFavoriteIndex < AllowedFavorites; ++CurFavoriteIndex )
+					// Display the most recent favorites on the main menu
+					const int32 MainMenuFavorites = FMath::Min( NumFavorites, FavoritesToDisplayInMainMenu );
+					for (int32 CurFavoriteIndex = 0; CurFavoriteIndex < MainMenuFavorites; ++CurFavoriteIndex)
 					{
-						TSharedPtr< FUICommandInfo > OpenFavoriteFile = FLevelEditorCommands::Get().OpenFavoriteFileCommands[ CurFavoriteIndex ];
-
 						const FString CurFavorite = FPaths::GetBaseFilename(RecentsAndFavorites.GetFavoritesItem( CurFavoriteIndex ));
-						const bool bNoIndent = true;
+						const bool bNoIndent = false;
 
 						MenuBuilder.AddWidget(
 							SNew(SFavouriteMenuEntry)
@@ -216,6 +187,41 @@ TSharedRef< SWidget > FLevelEditorMenu::MakeLevelEditorMenu( const TSharedPtr<FU
 								.OnRemoveClickedDelegate(FExecuteAction::CreateStatic(&FLevelEditorActionCallbacks::RemoveFavorite, CurFavoriteIndex)), 
 								FText(), bNoIndent );
 					}
+
+					// Any remaining favorites go into a submenu
+					if (NumFavorites > FavoritesToDisplayInMainMenu)
+					{
+						struct Local
+						{
+							static void MakeFavoriteLevelMenu(FMenuBuilder& MenuBuilder)
+							{
+								IMainFrameModule& MainFrameModule = FModuleManager::LoadModuleChecked<IMainFrameModule>("MainFrame");
+								const FMainMRUFavoritesList& RecentsAndFavorites = *MainFrameModule.GetMRUFavoritesList();
+								const int32 NumFavorites = RecentsAndFavorites.GetNumFavorites();
+
+								for (int32 CurFavoriteIndex = FavoritesToDisplayInMainMenu; CurFavoriteIndex < NumFavorites; ++CurFavoriteIndex)
+								{
+									const FString CurFavorite = FPaths::GetBaseFilename(RecentsAndFavorites.GetFavoritesItem(CurFavoriteIndex));
+									const bool bNoIndent = true;
+
+									MenuBuilder.AddWidget(
+										SNew(SFavouriteMenuEntry)
+										.LabelOverride(FText::FromString(FPaths::GetBaseFilename(CurFavorite)))
+										.OnOpenClickedDelegate(FExecuteAction::CreateStatic(&FLevelEditorActionCallbacks::OpenFavoriteFile, CurFavoriteIndex))
+										.OnRemoveClickedDelegate(FExecuteAction::CreateStatic(&FLevelEditorActionCallbacks::RemoveFavorite, CurFavoriteIndex)),
+										FText(), bNoIndent);
+								}
+							}
+						};
+
+						MenuBuilder.AddSubMenu(
+							LOCTEXT("MoreFavoriteFilesSubMenu", "More Favorites"),
+							LOCTEXT("MoreFavoriteFilesTooltip", "Select other favorite levels"),
+							FNewMenuDelegate::CreateStatic(&Local::MakeFavoriteLevelMenu),
+							false,
+							FSlateIcon(FEditorStyle::GetStyleSetName(), "MainFrame.FavoriteLevels")
+							);
+					}
 				}
 
 				MenuBuilder.EndSection();
diff --git a/Engine/Source/Editor/LevelEditor/Public/LevelEditorActions.h b/Engine/Source/Editor/LevelEditor/Public/LevelEditorActions.h
--- a/Engine/Source/Editor/LevelEditor/Public/LevelEditorActions.h
+++ b/Engine/Source/Editor/LevelEditor/Public/LevelEditorActions.h
@@ -47,11 +47,7 @@ public:
 	static const int32 MaxRecentFiles = 10;
 	TArray< TSharedPtr< FUICommandInfo > > OpenRecentFileCommands;
 
-	static const int32 MaxFavoriteFiles = 10;
-	TArray< TSharedPtr< FUICommandInfo > > OpenFavoriteFileCommands;
-
 	TSharedPtr< FUICommandInfo > ToggleFavorite;
-	TArray< TSharedPtr< FUICommandInfo > > RemoveFavoriteCommands;
 
 	/** Import */
 	TSharedPtr< FUICommandInfo > Import;
diff --git a/Engine/Source/Editor/UnrealEd/Private/MRUFavoritesList.cpp b/Engine/Source/Editor/UnrealEd/Private/MRUFavoritesList.cpp
--- a/Engine/Source/Editor/UnrealEd/Private/MRUFavoritesList.cpp
+++ b/Engine/Source/Editor/UnrealEd/Private/MRUFavoritesList.cpp
@@ -46,7 +46,7 @@ void FMainMRUFavoritesList::AddFavoritesItem( const FString& Item )
 	const FString CleanedName = FPaths::ConvertRelativePathToFull(Item);
 	if ( !FavoriteItems.Contains( CleanedName ) )
 	{
-		FavoriteItems.Add( CleanedName );
+		FavoriteItems.Insert( CleanedName, 0 );
 		WriteToINI();
 	}
 }
@@ -67,6 +67,21 @@ void FMainMRUFavoritesList::RemoveFavoritesItem( const FString& Item )
 	}
 }
 
+/**
+ * Moves the specified favorites item to the head of the list
+ *
+ * @param	Item	Filename of the item to move
+ */
+void FMainMRUFavoritesList::MoveFavoritesItemToHead(const FString& Item)
+{
+	const FString CleanedName = FPaths::ConvertRelativePathToFull(Item);
+	if ( FavoriteItems.RemoveSingle(Item) == 1 )
+	{
+		FavoriteItems.Insert( CleanedName, 0 );
+		WriteToINI();
+	}
+}
+
 /**
  * Returns whether a filename is favorited or not
  *
diff --git a/Engine/Source/Editor/UnrealEd/Public/MRUFavoritesList.h b/Engine/Source/Editor/UnrealEd/Public/MRUFavoritesList.h
--- a/Engine/Source/Editor/UnrealEd/Public/MRUFavoritesList.h
+++ b/Engine/Source/Editor/UnrealEd/Public/MRUFavoritesList.h
@@ -77,6 +77,13 @@ public:
 	 */
 	UNREALED_API bool VerifyFavoritesFile( int32 ItemIndex );
 
+	/**
+	 * Moves the specified favorites item to the head of the list
+	 *
+	 * @param	Item	Filename of the item to move
+	 */
+	UNREALED_API void MoveFavoritesItemToHead( const FString& Item );
+
 private:
 
 	/** Favorited items */
