[INTEGRATE] Change 2411014 by [REDACTED_PERSON][REDACTED_HANDLE] on 2015/01/19 11:59:12

	In-World Component Editing - fixed up a variety of cases where selection/details didn't update properly

diff --git a/Engine/Source/Editor/LevelEditor/Private/LevelEditorActions.cpp b/Engine/Source/Editor/LevelEditor/Private/LevelEditorActions.cpp
--- a/Engine/Source/Editor/LevelEditor/Private/LevelEditorActions.cpp
+++ b/Engine/Source/Editor/LevelEditor/Private/LevelEditorActions.cpp
@@ -1484,7 +1484,7 @@ void FLevelEditorActionCallbacks::OnSelectComponentOwnerActor()
 	auto ComponentOwner = Cast<AActor>(*GEditor->GetSelectedActorIterator());
 	check(ComponentOwner);
 
-	GEditor->SelectNone(false, true, false);
+	GEditor->SelectNone(true, true, false);
 	GEditor->SelectActor(ComponentOwner, true, true, true);
 }
 
diff --git a/Engine/Source/Editor/LevelEditor/Private/SLevelEditor.cpp b/Engine/Source/Editor/LevelEditor/Private/SLevelEditor.cpp
--- a/Engine/Source/Editor/LevelEditor/Private/SLevelEditor.cpp
+++ b/Engine/Source/Editor/LevelEditor/Private/SLevelEditor.cpp
@@ -631,33 +631,44 @@ private:
 	{
 		if (!bSelectionGuard && SCSEditor.IsValid())
 		{
+			// Enable the selection guard to prevent OnTreeSelectionChanged() from altering the editor's component selection
+			TGuardValue<bool> SelectionGuard(bSelectionGuard, true);
+
+			auto& SCSTreeWidget = SCSEditor->SCSTreeWidget;
+
 			TArray<UObject*> DetailsObjects;
 			auto Selection = Cast<USelection>(Object);
-			if (Selection == GEditor->GetSelectedComponents())
+			if (Selection == GEditor->GetSelectedComponents()) // A component was selected
 			{
-				// Enable the selection guard to prevent OnTreeSelectionChanged() from altering the editor's component selection
-				TGuardValue<bool> SelectionGuard(bSelectionGuard, true);
-
 				// Update the tree selection to match the level editor component selection
-				auto& SCSTreeWidget = SCSEditor->SCSTreeWidget;
 				SCSTreeWidget->ClearSelection();
 				for (FSelectionIterator It(GEditor->GetSelectedComponentIterator()); It; ++It)
 				{
 					UActorComponent* Component = CastChecked<UActorComponent>(*It);
-					if (Component)
+
+					auto SCSTreeNode = SCSEditor->GetNodeFromActorComponent(Component);
+					if (SCSTreeNode.IsValid())
 					{
-						auto SCSTreeNode = SCSEditor->GetNodeFromActorComponent(Component);
-						if (SCSTreeNode.IsValid())
-						{
-							SCSTreeWidget->SetItemSelection(SCSTreeNode, true);
-							check(Component == SCSTreeNode->GetComponentTemplate());
-							DetailsObjects.Add(Component);
-						}
+						SCSTreeWidget->SetItemSelection(SCSTreeNode, true);
+						check(Component == SCSTreeNode->GetComponentTemplate());
+						DetailsObjects.Add(Component);
 					}
 				}
 
 				DetailsView->SetObjects(DetailsObjects);
 			}
+			else if (Selection == GEditor->GetSelectedActors() && GEditor->GetSelectedComponentCount() == 0) // An actor was selected
+			{
+				// Wipe the tree selection
+				SCSTreeWidget->ClearSelection();
+
+				// Ensure the components appear properly selected
+				for (FSelectionIterator It(GEditor->GetSelectedActorIterator()); It; ++It)
+				{
+					auto Actor = CastChecked<AActor>(*It);
+					GUnrealEd->SetActorSelectionFlags(Actor);
+				}
+			}
 		}
 	}
 	
@@ -677,10 +688,13 @@ private:
 				// Enable the selection guard to prevent OnEditorSelectionChanged() from altering the contents of the SCSTreeWidget
 				TGuardValue<bool> SelectionGuard(bSelectionGuard, true);
 
+				const FScopedTransaction Transaction(NSLOCTEXT("UnrealEd", "ClickingOnComponentInTree", "Clicking on Component (tree view)"));
+
 				// Update the editor's component selection
 				TArray<UObject*> DetailsObjects;
 				USelection* SelectedComponents = GEditor->GetSelectedComponents();
 
+				SelectedComponents->Modify();
 				SelectedComponents->BeginBatchSelectOperation();
 				SelectedComponents->DeselectAll();
 
@@ -693,13 +707,21 @@ private:
 						{
 							DetailsObjects.Add(ComponentInstance);
 							SelectedComponents->Select(ComponentInstance);
+
+							// Ensure the selection override is bound for all the nodes (since it's possible we just added one)
+							auto PrimComponent = Cast<UPrimitiveComponent>(ComponentInstance);
+							if (PrimComponent && !PrimComponent->SelectionOverrideDelegate.IsBound())
+							{
+								PrimComponent->SelectionOverrideDelegate = UPrimitiveComponent::FSelectionOverride::CreateUObject(GUnrealEd, &UUnrealEdEngine::IsComponentSelected);
+							}
 						}
 					}
 				}
-				DetailsView->SetObjects(DetailsObjects);
 
 				SelectedComponents->EndBatchSelectOperation();
 
+				DetailsView->SetObjects(DetailsObjects);
+
 				GUnrealEd->SetActorSelectionFlags(Actor);
 				GUnrealEd->UpdatePivotLocationForSelection(true);
 				GEditor->RedrawLevelEditingViewports();
diff --git a/Engine/Source/Editor/UnrealEd/Private/EditorSelectUtils.cpp b/Engine/Source/Editor/UnrealEd/Private/EditorSelectUtils.cpp
--- a/Engine/Source/Editor/UnrealEd/Private/EditorSelectUtils.cpp
+++ b/Engine/Source/Editor/UnrealEd/Private/EditorSelectUtils.cpp
@@ -609,7 +609,7 @@ void UUnrealEdEngine::SelectActor(AActor* Actor, bool bInSelected, bool bNotify,
 void UUnrealEdEngine::SelectComponent(UActorComponent* Component, bool bInSelected, bool bNotify, bool bSelectEvenIfHidden)
 {
 	// Don't do any work if the component's selection state matches the target selection state
-	const bool bComponentSelected = Component->IsSelected();
+	const bool bComponentSelected = GetSelectedComponents()->IsSelected(Component);
 	if (( bComponentSelected && !bInSelected ) || ( !bComponentSelected && bInSelected ))
 	{
 		if (bInSelected)
diff --git a/Engine/Source/Editor/UnrealEd/Private/LevelViewportClickHandlers.cpp b/Engine/Source/Editor/UnrealEd/Private/LevelViewportClickHandlers.cpp
--- a/Engine/Source/Editor/UnrealEd/Private/LevelViewportClickHandlers.cpp
+++ b/Engine/Source/Editor/UnrealEd/Private/LevelViewportClickHandlers.cpp
@@ -281,6 +281,7 @@ namespace ClickHandlers
 			if (SceneComp == ActorHitProxy->PrimComponent)
 			{
 				Component = SceneComp;
+				break;
 			}
 		}
 		
