[INTEGRATE] Change 2412472 by [REDACTED_PERSON][REDACTED_HANDLE] on 2015/01/20 13:14:31

	Moved new source control toolbar button next to File->Save button

diff --git a/Engine/Source/Editor/LevelEditor/Private/LevelEditorToolBar.cpp b/Engine/Source/Editor/LevelEditor/Private/LevelEditorToolBar.cpp
--- a/Engine/Source/Editor/LevelEditor/Private/LevelEditorToolBar.cpp
+++ b/Engine/Source/Editor/LevelEditor/Private/LevelEditorToolBar.cpp
@@ -1091,6 +1091,88 @@ TSharedRef< SWidget > FLevelEditorToolBar::MakeLevelEditorToolBar( const TShared
 	{
 		// Save All Levels
 		ToolbarBuilder.AddToolBarButton( FLevelEditorCommands::Get().Save, NAME_None, TAttribute<FText>(), TAttribute<FText>(), FSlateIcon(FEditorStyle::GetStyleSetName(), "AssetEditor.SaveAsset") );
+
+		// Source control buttons
+		{
+			enum EQueryState
+			{
+				NotQueried,
+				Querying,
+				Queried,
+			};
+
+			static EQueryState QueryState = EQueryState::NotQueried;
+
+			struct FSourceControlStatus
+			{
+				static void CheckSourceControlStatus()
+				{
+					ISourceControlModule& SourceControlModule = ISourceControlModule::Get();
+					if (SourceControlModule.IsEnabled())
+					{
+						SourceControlModule.GetProvider().Execute(ISourceControlOperation::Create<FConnect>(),
+																  EConcurrency::Asynchronous,
+																  FSourceControlOperationComplete::CreateStatic(&FSourceControlStatus::OnSourceControlOperationComplete));
+						QueryState = EQueryState::Querying;
+					}
+				}
+
+				static void OnSourceControlOperationComplete(const FSourceControlOperationRef& InOperation, ECommandResult::Type InResult)
+				{
+					QueryState = EQueryState::Queried;
+				}
+
+				static FText GetSourceControlTooltip()
+				{
+					if (QueryState == EQueryState::Querying)
+					{
+						return LOCTEXT("SourceControlUnknown", "Source control status is unknown");
+					}
+					else
+					{
+						return ISourceControlModule::Get().GetProvider().GetStatusText();
+					}
+				}
+
+				static FSlateIcon GetSourceControlIcon()
+				{
+					if (QueryState == EQueryState::Querying)
+					{
+						return FSlateIcon(FEditorStyle::GetStyleSetName(), "LevelEditor.SourceControl.Unknown");
+					}
+					else
+					{
+						ISourceControlModule& SourceControlModule = ISourceControlModule::Get();
+						if (SourceControlModule.IsEnabled())
+						{
+							if (!SourceControlModule.GetProvider().IsAvailable())
+							{
+								return FSlateIcon(FEditorStyle::GetStyleSetName(), "LevelEditor.SourceControl.Error");
+							}
+							else
+							{
+								return FSlateIcon(FEditorStyle::GetStyleSetName(), "LevelEditor.SourceControl.On");
+							}
+						}
+						else
+						{
+							return FSlateIcon(FEditorStyle::GetStyleSetName(), "LevelEditor.SourceControl.Off");
+						}
+					}
+				}
+			};
+
+			FSourceControlStatus::CheckSourceControlStatus();
+
+			ToolbarBuilder.AddComboButton(
+				FUIAction(),
+				FOnGetContent::CreateStatic(&FLevelEditorToolBar::GenerateSourceControlMenu, InCommandList),
+				LOCTEXT("SourceControl_Label", "Source Control"),
+				TAttribute<FText>::Create(&FSourceControlStatus::GetSourceControlTooltip),
+				TAttribute<FSlateIcon>::Create(&FSourceControlStatus::GetSourceControlIcon),
+				false
+				);
+		}
 	}
 	ToolbarBuilder.EndSection();
 
@@ -1179,88 +1261,6 @@ TSharedRef< SWidget > FLevelEditorToolBar::MakeLevelEditorToolBar( const TShared
 	}
 	ToolbarBuilder.EndSection();
 
-	ToolbarBuilder.BeginSection("SourceControl");
-	{
-		enum EQueryState
-		{
-			NotQueried,
-			Querying,
-			Queried,
-		};
-
-		static EQueryState QueryState = EQueryState::NotQueried;
-
-		struct FSourceControlStatus
-		{
-			static void CheckSourceControlStatus()
-			{
-				ISourceControlModule& SourceControlModule = ISourceControlModule::Get();
-				if (SourceControlModule.IsEnabled())
-				{
-					SourceControlModule.GetProvider().Execute(ISourceControlOperation::Create<FConnect>(),
-															  EConcurrency::Asynchronous,
-															  FSourceControlOperationComplete::CreateStatic(&FSourceControlStatus::OnSourceControlOperationComplete));
-					QueryState = EQueryState::Querying;
-				}
-			}
-
-			static void OnSourceControlOperationComplete(const FSourceControlOperationRef& InOperation, ECommandResult::Type InResult)
-			{
-				QueryState = EQueryState::Queried;
-			}
-
-			static FText GetSourceControlTooltip()
-			{
-				if (QueryState == EQueryState::Querying)
-				{
-					return LOCTEXT("SourceControlUnknown", "Source control status is unknown");
-				}
-				else
-				{
-					return ISourceControlModule::Get().GetProvider().GetStatusText();
-				}
-			}
-
-			static FSlateIcon GetSourceControlIcon()
-			{
-				if (QueryState == EQueryState::Querying)
-				{
-					return FSlateIcon(FEditorStyle::GetStyleSetName(), "LevelEditor.SourceControl.Unknown");
-				}
-				else
-				{
-					ISourceControlModule& SourceControlModule = ISourceControlModule::Get();
-					if (SourceControlModule.IsEnabled())
-					{
-						if (!SourceControlModule.GetProvider().IsAvailable())
-						{
-							return FSlateIcon(FEditorStyle::GetStyleSetName(), "LevelEditor.SourceControl.Error");
-						}
-						else
-						{
-							return FSlateIcon(FEditorStyle::GetStyleSetName(), "LevelEditor.SourceControl.On");
-						}
-					}
-					else
-					{
-						return FSlateIcon(FEditorStyle::GetStyleSetName(), "LevelEditor.SourceControl.Off");
-					}
-				}
-			}
-		};
-
-		FSourceControlStatus::CheckSourceControlStatus();
-
-		ToolbarBuilder.AddComboButton(
-			FUIAction(),
-			FOnGetContent::CreateStatic(&FLevelEditorToolBar::GenerateSourceControlMenu, InCommandList),
-			LOCTEXT("SourceControl_Label", "Source Control"),
-			TAttribute<FText>::Create(&FSourceControlStatus::GetSourceControlTooltip),
-			TAttribute<FSlateIcon>::Create(&FSourceControlStatus::GetSourceControlIcon),
-			false
-			);
-	}
-	ToolbarBuilder.EndSection();
 
 	// Create the tool bar!
 	return
