[INTEGRATE] Change 2412721 by [REDACTED_PERSON][REDACTED_HANDLE] on 2015/01/20 15:18:01

	In-World Component Editing
	- Undo/Redo add component works
	- Valid name is automatically assigned to a component when it's added
	- Selection in the component list is a transaction
	- Gave the SCSEditor the same style as the scene outliner (peach for selected+inactive)

diff --git a/Engine/Source/Editor/Kismet/Private/SSCSEditor.cpp b/Engine/Source/Editor/Kismet/Private/SSCSEditor.cpp
--- a/Engine/Source/Editor/Kismet/Private/SSCSEditor.cpp
+++ b/Engine/Source/Editor/Kismet/Private/SSCSEditor.cpp
@@ -235,6 +235,14 @@ FSCSEditorTreeNode::FSCSEditorTreeNode(UActorComponent* InComponentTemplate)
 	{
 		bIsInstanced = true;
 		bWasInstancedFromNativeClass = false;
+
+		// Make sure the component has a valid name
+		if (!FComponentEditorUtils::IsValidVariableNameString(InComponentTemplate, InComponentTemplate->GetName()))
+		{
+			ERenameFlags RenameFlags = REN_DontCreateRedirectors | REN_NonTransactional;
+			InComponentTemplate->Rename(*InComponentTemplate->GetName().Replace(TEXT("Component_"), TEXT("")), nullptr, RenameFlags);
+		}
+		
 		InstancedComponentName = InComponentTemplate->GetFName();
 
 		ComponentTemplatePtr.Reset();
@@ -1004,7 +1012,12 @@ void SSCS_RowWidget::Construct( const FArguments& InArgs, TSharedPtr<SSCSEditor>
 	SCSEditor = InSCSEditor;
 	NodePtr = InNodePtr;
 
-	SMultiColumnTableRow<FSCSEditorTreeNodePtrType>::Construct( FSuperRowType::FArguments().Padding(FMargin(0.f, 0.f, 0.f, 4.f)), InOwnerTableView.ToSharedRef() );
+	
+	auto Args = FSuperRowType::FArguments()
+		.Style(&FEditorStyle::Get().GetWidgetStyle<FTableRowStyle>("SceneOutliner.TableViewRow")) //@todo create editor style for the SCS tree
+		.Padding(FMargin(0.f, 0.f, 0.f, 4.f));
+
+	SMultiColumnTableRow<FSCSEditorTreeNodePtrType>::Construct( Args, InOwnerTableView.ToSharedRef() );
 }
 
 SSCS_RowWidget::~SSCS_RowWidget()
@@ -3037,12 +3050,12 @@ FText SSCSEditor::GetActorDisplayText() const
 		if (Blueprint != nullptr)
 		{
 			Blueprint->GetName(Name);
-			return FText::Format(LOCTEXT("DefaultActor_Name", "{0} (self)"), FText::FromString(Name));
+			return FText::Format(LOCTEXT("DefaultActor_Name", "{0} (Self)"), FText::FromString(Name));
 		}
 		else
 		{
 			Name = DefaultActor->GetActorLabel();
-			return FText::Format(LOCTEXT("DefaultActor_Name", "{0}"), FText::FromString(Name));
+			return FText::Format(LOCTEXT("DefaultActor_Name", "{0} (Instance)"), FText::FromString(Name));
 		}
 	}
 	return FText::GetEmpty();	
@@ -3573,28 +3586,28 @@ UActorComponent* SSCSEditor::AddNewNode(USCS_Node* NewNode,  UObject* Asset, boo
 	return NewNode->ComponentTemplate;
 }
 
-UActorComponent* SSCSEditor::AddNewNode(UActorComponent* NewNode,  UObject* Asset, bool bSetFocusToNewItem)
+UActorComponent* SSCSEditor::AddNewNode(UActorComponent* NewInstanceComponent,  UObject* Asset, bool bSetFocusToNewItem)
 {
-	check(NewNode != nullptr);
+	check(NewInstanceComponent != nullptr);
 
 	if(Asset)
 	{
-		FComponentAssetBrokerage::AssignAssetToComponent(NewNode, Asset);
+		FComponentAssetBrokerage::AssignAssetToComponent(NewInstanceComponent, Asset);
 	}
 
-	NewNode->RegisterComponent();
+	NewInstanceComponent->RegisterComponent();
 
 	FSCSEditorTreeNodePtrType NewNodePtr;
 
 	// Add the new node to the editor tree
-	USceneComponent* NewSceneComponent = Cast<USceneComponent>(NewNode);
+	USceneComponent* NewSceneComponent = Cast<USceneComponent>(NewInstanceComponent);
 	if(NewSceneComponent != nullptr)
 	{
 		NewNodePtr = AddTreeNode(NewSceneComponent);
 	}
 	else
 	{
-		NewNodePtr = MakeShareable(new FSCSEditorTreeNode(NewNode));
+		NewNodePtr = MakeShareable(new FSCSEditorTreeNode(NewInstanceComponent));
 
 		// Ensure that the root node ordering is what we assume it to be
 		check(!SceneRootNodePtr.IsValid() || (RootNodes.Num() > 0 && RootNodes[RootNodes.Num() - 1] == SceneRootNodePtr));
@@ -3619,7 +3632,7 @@ UActorComponent* SSCSEditor::AddNewNode(UActorComponent* NewNode,  UObject* Asse
 
 	UpdateTree();
 
-	return NewNode;
+	return NewInstanceComponent;
 }
 
 bool SSCSEditor::IsComponentSelected(const UPrimitiveComponent* PrimComponent) const
diff --git a/Engine/Source/Editor/Kismet/Public/SSCSEditor.h b/Engine/Source/Editor/Kismet/Public/SSCSEditor.h
--- a/Engine/Source/Editor/Kismet/Public/SSCSEditor.h
+++ b/Engine/Source/Editor/Kismet/Public/SSCSEditor.h
@@ -473,11 +473,11 @@ public:
 	UActorComponent* AddNewNode(USCS_Node* NewNode, UObject* Asset, bool bMarkBlueprintModified, bool bSetFocusToNewItem = true);
 
 	/** Adds a new component instance node to the component Table
-		@param NewNode	(In) The component instance node to add
+		@param NewInstanceComponent	(In) The component being added to the actor instance
 		@param Asset	(In) Optional asset to assign to the component
 		@param bSetFocusToNewItem (In) Select the new item and activate the inline rename widget (default is true)
 		@return The reference of the newly created ActorComponent */
-	UActorComponent* AddNewNode(UActorComponent* NewNode, UObject* Asset, bool bSetFocusToNewItem = true);
+	UActorComponent* AddNewNode(UActorComponent* NewInstanceComponent, UObject* Asset, bool bSetFocusToNewItem = true);
 	
 	/** Returns true if the specified component is currently selected */
 	bool IsComponentSelected(const UPrimitiveComponent* PrimComponent) const;
diff --git a/Engine/Source/Editor/LevelEditor/Private/SLevelEditor.cpp b/Engine/Source/Editor/LevelEditor/Private/SLevelEditor.cpp
--- a/Engine/Source/Editor/LevelEditor/Private/SLevelEditor.cpp
+++ b/Engine/Source/Editor/LevelEditor/Private/SLevelEditor.cpp
@@ -502,7 +502,7 @@ TSharedRef<FTabManager> SLevelEditor::GetTabManager() const
 	return LevelEditorTabManager.ToSharedRef();
 }
 
-class SActorDetails : public SCompoundWidget
+class SActorDetails : public SCompoundWidget, public FEditorUndoClient
 {
 public:
 	SLATE_BEGIN_ARGS( SActorDetails )
@@ -547,6 +547,8 @@ public:
 		// Set up a delegate to call to add generic details to the view
 		DetailsView->SetGenericLayoutDetailsDelegate( FOnGetDetailCustomizationInstance::CreateStatic( &FLevelEditorGenericDetails::MakeInstance ) );
 
+		GEditor->RegisterForUndo(this);
+
 		SAssignNew(ComponentsBox, SBox)
 		.Visibility(EVisibility::Collapsed);
 
@@ -625,49 +627,37 @@ public:
 		}
 	}
 
+	virtual void PostRedo(bool bSuccess)
+	{
+		// Enable the selection guard to prevent OnTreeSelectionChanged() from altering the editor's component selection
+		TGuardValue<bool> SelectionGuard(bSelectionGuard, true);
+
+		// Refresh the tree and update the selection to match the world
+		SCSEditor->UpdateTree();
+		UpdateComponentTreeFromEditorSelection();
+	}
+
 private:
 
 	void OnEditorSelectionChanged(UObject* Object)
 	{
 		if (!bSelectionGuard && SCSEditor.IsValid())
 		{
-			// Enable the selection guard to prevent OnTreeSelectionChanged() from altering the editor's component selection
-			TGuardValue<bool> SelectionGuard(bSelectionGuard, true);
-
-			auto& SCSTreeWidget = SCSEditor->SCSTreeWidget;
-
-			TArray<UObject*> DetailsObjects;
+			// Make sure the selection set that changed is relevant to us
 			auto Selection = Cast<USelection>(Object);
-			if (Selection == GEditor->GetSelectedComponents()) // A component was selected
+			if (Selection == GEditor->GetSelectedComponents() || Selection == GEditor->GetSelectedActors())
 			{
-				// Update the tree selection to match the level editor component selection
-				SCSTreeWidget->ClearSelection();
-				for (FSelectionIterator It(GEditor->GetSelectedComponentIterator()); It; ++It)
-				{
-					UActorComponent* Component = CastChecked<UActorComponent>(*It);
+				UpdateComponentTreeFromEditorSelection();
 
-					auto SCSTreeNode = SCSEditor->GetNodeFromActorComponent(Component);
-					if (SCSTreeNode.IsValid())
+				if (GEditor->GetSelectedComponentCount() == 0) // An actor was selected
+				{
+					// Ensure the selection flags are up to date for the components in the selected actor
+					for (FSelectionIterator It(GEditor->GetSelectedActorIterator()); It; ++It)
 					{
-						SCSTreeWidget->SetItemSelection(SCSTreeNode, true);
-						check(Component == SCSTreeNode->GetComponentTemplate());
-						DetailsObjects.Add(Component);
+						auto Actor = CastChecked<AActor>(*It);
+						GUnrealEd->SetActorSelectionFlags(Actor);
 					}
 				}
-
-				DetailsView->SetObjects(DetailsObjects);
-			}
-			else if (Selection == GEditor->GetSelectedActors() && GEditor->GetSelectedComponentCount() == 0) // An actor was selected
-			{
-				// Wipe the tree selection
-				SCSTreeWidget->ClearSelection();
-
-				// Ensure the components appear properly selected
-				for (FSelectionIterator It(GEditor->GetSelectedActorIterator()); It; ++It)
-				{
-					auto Actor = CastChecked<AActor>(*It);
-					GUnrealEd->SetActorSelectionFlags(Actor);
-				}
 			}
 		}
 	}
@@ -678,32 +668,34 @@ private:
 		return Cast<AActor>(*GEditor->GetSelectedActorIterator());
 	}
 
-	void OnSCSEditorRootSelected(AActor*)
+	void OnSCSEditorRootSelected(AActor* Actor)
 	{
 		if (!bSelectionGuard)
 		{
-			auto Actor = GetSelectedActor();
-			if (Actor)
-			{
-				// Enable the selection guard to prevent OnEditorSelectionChanged() from altering the contents of the SCSTreeWidget
-				TGuardValue<bool> SelectionGuard(bSelectionGuard, true);
-
-				// Update the editor's component selection
-				TArray<UObject*> DetailsObjects;
-				USelection* SelectedComponents = GEditor->GetSelectedComponents();
-				DetailsObjects.Add(Actor);
-
-				SelectedComponents->BeginBatchSelectOperation();
-				SelectedComponents->DeselectAll();
-
-				DetailsView->SetObjects(DetailsObjects, false);
-
-				SelectedComponents->EndBatchSelectOperation();
-
-				GUnrealEd->SetActorSelectionFlags(Actor);
-				GUnrealEd->UpdatePivotLocationForSelection(true);
-				GEditor->RedrawLevelEditingViewports();
-			}
+			GEditor->SelectNone(true, true, false);
+			GEditor->SelectActor(Actor, true, true, true);
+		//	auto Actor = GetSelectedActor();
+		//	if (Actor)
+		//	{
+		//		// Enable the selection guard to prevent OnEditorSelectionChanged() from altering the contents of the SCSTreeWidget
+		//		TGuardValue<bool> SelectionGuard(bSelectionGuard, true);
+
+		//		// Update the editor's component selection
+		//		TArray<UObject*> DetailsObjects;
+		//		USelection* SelectedComponents = GEditor->GetSelectedComponents();
+		//		DetailsObjects.Add(Actor);
+
+		//		SelectedComponents->BeginBatchSelectOperation();
+		//		SelectedComponents->DeselectAll();
+
+		//		DetailsView->SetObjects(DetailsObjects, false);
+
+		//		SelectedComponents->EndBatchSelectOperation();
+
+		//		GUnrealEd->SetActorSelectionFlags(Actor);
+		//		GUnrealEd->UpdatePivotLocationForSelection(true);
+		//		GEditor->RedrawLevelEditingViewports();
+		//	}
 		}
 	}
 
@@ -714,50 +706,105 @@ private:
 			auto Actor = GetSelectedActor();
 			if (Actor)
 			{
-				// Enable the selection guard to prevent OnEditorSelectionChanged() from altering the contents of the SCSTreeWidget
-				TGuardValue<bool> SelectionGuard(bSelectionGuard, true);
-
-				const FScopedTransaction Transaction(NSLOCTEXT("UnrealEd", "ClickingOnComponentInTree", "Clicking on Component (tree view)"));
-
-				// Update the editor's component selection
-				TArray<UObject*> DetailsObjects;
 				USelection* SelectedComponents = GEditor->GetSelectedComponents();
 
-				SelectedComponents->Modify();
-				SelectedComponents->BeginBatchSelectOperation();
-				SelectedComponents->DeselectAll();
-
-				for (auto& SelectedNode : SelectedNodes)
+				// Don't bother doing anything if the node selection already matches the current world selection
+				bool bSelectionChanged = GEditor->GetSelectedComponentCount() != SelectedNodes.Num();
+				if (!bSelectionChanged)
 				{
-					if (SelectedNode.IsValid())
+					// Check to see if any of the selected nodes aren't already selected in the world
+					for (auto& SelectedNode : SelectedNodes)
 					{
 						UActorComponent* ComponentInstance = SelectedNode->FindComponentInstanceInActor(Actor);
-						if (ComponentInstance)
+						if (ComponentInstance && !SelectedComponents->IsSelected(ComponentInstance))
 						{
-							DetailsObjects.Add(ComponentInstance);
-							SelectedComponents->Select(ComponentInstance);
+							// At least one of the selected nodes isn't selected in the world, so update the selection
+							bSelectionChanged = true;
+							break;
+						}
+					}
+				}
+				
+				if (bSelectionChanged)
+				{
+					// Enable the selection guard to prevent OnEditorSelectionChanged() from altering the contents of the SCSTreeWidget
+					TGuardValue<bool> SelectionGuard(bSelectionGuard, true);
+
+					// Update the editor's component selection to match the node selection
+					TArray<UObject*> DetailsObjects;
+					const FScopedTransaction Transaction(NSLOCTEXT("UnrealEd", "ClickingOnComponentInTree", "Clicking on Component (tree view)"));
 
-							// Ensure the selection override is bound for all the nodes (since it's possible we just added one)
-							auto PrimComponent = Cast<UPrimitiveComponent>(ComponentInstance);
-							if (PrimComponent && !PrimComponent->SelectionOverrideDelegate.IsBound())
+					SelectedComponents->Modify();
+					SelectedComponents->BeginBatchSelectOperation();
+					SelectedComponents->DeselectAll();
+
+					for (auto& SelectedNode : SelectedNodes)
+					{
+						if (SelectedNode.IsValid())
+						{
+							UActorComponent* ComponentInstance = SelectedNode->FindComponentInstanceInActor(Actor);
+							if (ComponentInstance)
 							{
-								PrimComponent->SelectionOverrideDelegate = UPrimitiveComponent::FSelectionOverride::CreateUObject(GUnrealEd, &UUnrealEdEngine::IsComponentSelected);
+								DetailsObjects.Add(ComponentInstance);
+								SelectedComponents->Select(ComponentInstance);
+
+								// Ensure the selection override is bound for all the nodes (since it's possible we just added one)
+								auto PrimComponent = Cast<UPrimitiveComponent>(ComponentInstance);
+								if (PrimComponent && !PrimComponent->SelectionOverrideDelegate.IsBound())
+								{
+									PrimComponent->SelectionOverrideDelegate = UPrimitiveComponent::FSelectionOverride::CreateUObject(GUnrealEd, &UUnrealEdEngine::IsComponentSelected);
+								}
 							}
 						}
 					}
-				}
 
-				SelectedComponents->EndBatchSelectOperation();
+					SelectedComponents->EndBatchSelectOperation();
 
-				DetailsView->SetObjects(DetailsObjects);
+					DetailsView->SetObjects(DetailsObjects);
 
-				GUnrealEd->SetActorSelectionFlags(Actor);
-				GUnrealEd->UpdatePivotLocationForSelection(true);
-				GEditor->RedrawLevelEditingViewports();
+					GUnrealEd->SetActorSelectionFlags(Actor);
+					GUnrealEd->UpdatePivotLocationForSelection(true);
+					GEditor->RedrawLevelEditingViewports();
+				}
 			}
 		}
 	}
 	
+	void UpdateComponentTreeFromEditorSelection()
+	{
+		// Enable the selection guard to prevent OnTreeSelectionChanged() from altering the editor's component selection
+		TGuardValue<bool> SelectionGuard(bSelectionGuard, true);
+
+		auto& SCSTreeWidget = SCSEditor->SCSTreeWidget;
+		TArray<UObject*> DetailsObjects;
+
+		// Update the tree selection to match the level editor component selection
+		SCSTreeWidget->ClearSelection();
+		for (FSelectionIterator It(GEditor->GetSelectedComponentIterator()); It; ++It)
+		{
+			UActorComponent* Component = CastChecked<UActorComponent>(*It);
+
+			auto SCSTreeNode = SCSEditor->GetNodeFromActorComponent(Component, false);
+			if (SCSTreeNode.IsValid() && SCSTreeNode->GetComponentTemplate())
+			{
+				SCSTreeWidget->SetItemSelection(SCSTreeNode, true);
+
+				auto ComponentTemplate = SCSTreeNode->GetComponentTemplate();
+				check(Component == ComponentTemplate);
+				DetailsObjects.Add(Component);
+			}
+		}
+
+		if (DetailsObjects.Num() > 0)
+		{
+			DetailsView->SetObjects(DetailsObjects);
+		}
+		else
+		{
+			//@todo Make sure the actor "root" is selected
+		}
+	}
+
 private:
 	TSharedPtr<SSplitter> DetailsSplitter;
 	TSharedPtr<IDetailsView> DetailsView;
diff --git a/Engine/Source/Editor/UnrealEd/Classes/Editor/EditorEngine.h b/Engine/Source/Editor/UnrealEd/Classes/Editor/EditorEngine.h
--- a/Engine/Source/Editor/UnrealEd/Classes/Editor/EditorEngine.h
+++ b/Engine/Source/Editor/UnrealEd/Classes/Editor/EditorEngine.h
@@ -947,7 +947,7 @@ public:
 	 */
 	virtual void TakeHighResScreenShots(){}
 
-	virtual void NoteSelectionChange(bool bComponentSelectionChanged = false) { check(0); }
+	virtual void NoteSelectionChange() { check(0); }
 
 	/**
 	 * Adds an actor to the world at the specified location.
diff --git a/Engine/Source/Editor/UnrealEd/Classes/Editor/UnrealEdEngine.h b/Engine/Source/Editor/UnrealEd/Classes/Editor/UnrealEdEngine.h
--- a/Engine/Source/Editor/UnrealEd/Classes/Editor/UnrealEdEngine.h
+++ b/Engine/Source/Editor/UnrealEd/Classes/Editor/UnrealEdEngine.h
@@ -184,7 +184,7 @@ class UNREALED_API UUnrealEdEngine : public UEditorEngine, public FNotifyHook
 	virtual void SelectComponent(class UActorComponent* Component, bool bInSelected, bool bNotify, bool bSelectEvenIfHidden = false) override;
 	virtual void SelectBSPSurf(UModel* InModel, int32 iSurf, bool bSelected, bool bNoteSelectionChange) override;
 	virtual void SelectNone(bool bNoteSelectionChange, bool bDeselectBSPSurfs, bool WarnAboutManyActors=true) override;
-	virtual void NoteSelectionChange(bool bComponentSelectionChanged = false) override;
+	virtual void NoteSelectionChange() override;
 	virtual void NoteActorMovement() override;
 	virtual void FinishAllSnaps() override;
 	virtual void Cleanse( bool ClearSelection, bool Redraw, const FText& Reason ) override;
diff --git a/Engine/Source/Editor/UnrealEd/Private/EditorSelectUtils.cpp b/Engine/Source/Editor/UnrealEd/Private/EditorSelectUtils.cpp
--- a/Engine/Source/Editor/UnrealEd/Private/EditorSelectUtils.cpp
+++ b/Engine/Source/Editor/UnrealEd/Private/EditorSelectUtils.cpp
@@ -351,7 +351,7 @@ void UUnrealEdEngine::UpdatePivotLocationForSelection( bool bOnChange )
 
 
 
-void UUnrealEdEngine::NoteSelectionChange(bool bComponentSelectionChanged)
+void UUnrealEdEngine::NoteSelectionChange()
 {
 	// The selection changed, so make sure the pivot (widget) is located in the right place
 	UpdatePivotLocationForSelection( true );
@@ -366,9 +366,9 @@ void UUnrealEdEngine::NoteSelectionChange(bool bComponentSelectionChanged)
 		ActiveModes[ModeIndex]->ActorSelectionChangeNotify();
 	}
 
+	bool bComponentSelectionChanged = GetSelectedComponentCount() > 0;
 	USelection* Selection = bComponentSelectionChanged ? GetSelectedComponents() : GetSelectedActors();
 	USelection::SelectionChangedEvent.Broadcast(Selection);
-
 	
 	if (!bComponentSelectionChanged)
 	{
@@ -644,7 +644,7 @@ void UUnrealEdEngine::SelectComponent(UActorComponent* Component, bool bInSelect
 
 		if (bNotify)
 		{
-			NoteSelectionChange(true);
+			NoteSelectionChange();
 		}
 	}
 }
