[INTEGRATE] Change 2413255 by [REDACTED_PERSON][REDACTED_HANDLE] on 2015/01/20 20:29:49

	In-world component editing
	- Proper name is given to a new component on creation and when pasting nodes for IWCE component tree
	UE-7864 - fixed as by-product of naming changes
	UE-7831 - numbering follows the same pattern as the blueprint SCS editor
	UE-7801 - Selection & selection overrides all work for "invisible" components that have an editor-only component

diff --git a/Engine/Source/Editor/Kismet/Private/SSCSEditor.cpp b/Engine/Source/Editor/Kismet/Private/SSCSEditor.cpp
--- a/Engine/Source/Editor/Kismet/Private/SSCSEditor.cpp
+++ b/Engine/Source/Editor/Kismet/Private/SSCSEditor.cpp
@@ -240,7 +240,7 @@ FSCSEditorTreeNode::FSCSEditorTreeNode(UActorComponent* InComponentTemplate)
 		if (!FComponentEditorUtils::IsValidVariableNameString(InComponentTemplate, InComponentTemplate->GetName()))
 		{
 			ERenameFlags RenameFlags = REN_DontCreateRedirectors | REN_NonTransactional;
-			FString NewComponentName = FComponentEditorUtils::GenerateValidVariableName(InComponentTemplate);
+			FString NewComponentName = FComponentEditorUtils::GenerateValidVariableName(InComponentTemplate->GetClass(), InComponentTemplate->GetOwner());
 
 			InComponentTemplate->Rename(*NewComponentName, nullptr, RenameFlags);
 		}
@@ -3533,7 +3533,9 @@ UActorComponent* SSCSEditor::AddNewComponent( UClass* NewComponentClass, UObject
 		ActorInstance->Modify();
 
 		// Create new component
-		UActorComponent* NewComponentInstance = ConstructObject<UActorComponent>(NewComponentClass, ActorInstance, NAME_None, RF_Transactional);
+		FName NewComponentName(*FComponentEditorUtils::GenerateValidVariableName(NewComponentClass, ActorInstance));
+
+		UActorComponent* NewComponentInstance = ConstructObject<UActorComponent>(NewComponentClass, ActorInstance, NewComponentName, RF_Transactional);
 		check(NewComponentInstance != nullptr);
 
 		// Add to SerializedComponents array so it gets saved
@@ -3841,7 +3843,8 @@ void SSCSEditor::PasteNodes()
 		else
 		{
 			// Relocate the instance from the transient package to the Actor and assign it a unique object name
-			NewActorComponent->Rename(NULL, ActorInstance, REN_DontCreateRedirectors|REN_DoNotDirty);
+			FString NewComponentName = FComponentEditorUtils::GenerateValidVariableName(NewActorComponent->GetClass(), ActorInstance);
+			NewActorComponent->Rename(*NewComponentName, ActorInstance, REN_DontCreateRedirectors | REN_DoNotDirty);
 
 			// Add to SerializedComponents array so it gets saved
 			ActorInstance->InstanceComponents.Add(NewActorComponent);
diff --git a/Engine/Source/Editor/LevelEditor/Private/SLevelEditor.cpp b/Engine/Source/Editor/LevelEditor/Private/SLevelEditor.cpp
--- a/Engine/Source/Editor/LevelEditor/Private/SLevelEditor.cpp
+++ b/Engine/Source/Editor/LevelEditor/Private/SLevelEditor.cpp
@@ -627,7 +627,7 @@ public:
 		}
 	}
 
-	virtual void PostRedo(bool bSuccess)
+	virtual void PostUndo(bool bSuccess)
 	{
 		// Enable the selection guard to prevent OnTreeSelectionChanged() from altering the editor's component selection
 		TGuardValue<bool> SelectionGuard(bSelectionGuard, true);
@@ -637,6 +637,11 @@ public:
 		UpdateComponentTreeFromEditorSelection();
 	}
 
+	virtual void PostRedo(bool bSuccess)
+	{
+		PostUndo(bSuccess);
+	}
+
 private:
 
 	void OnEditorSelectionChanged(UObject* Object)
@@ -748,12 +753,28 @@ private:
 								DetailsObjects.Add(ComponentInstance);
 								SelectedComponents->Select(ComponentInstance);
 
-								// Ensure the selection override is bound for all the nodes (since it's possible we just added one)
+								// Ensure the selection override is bound for this component (including any attached editor-only children)
 								auto PrimComponent = Cast<UPrimitiveComponent>(ComponentInstance);
 								if (PrimComponent && !PrimComponent->SelectionOverrideDelegate.IsBound())
 								{
 									PrimComponent->SelectionOverrideDelegate = UPrimitiveComponent::FSelectionOverride::CreateUObject(GUnrealEd, &UUnrealEdEngine::IsComponentSelected);
 								}
+								else
+								{
+									//@todo move the selection override binding check to FComponentEditorUtils
+									auto SceneComponent = Cast<USceneComponent>(ComponentInstance);
+									if (SceneComponent)
+									{
+										for (auto Component : SceneComponent->AttachChildren)
+										{
+											PrimComponent = Cast<UPrimitiveComponent>(Component);
+											if (PrimComponent && !PrimComponent->SelectionOverrideDelegate.IsBound())
+											{
+												PrimComponent->SelectionOverrideDelegate = UPrimitiveComponent::FSelectionOverride::CreateUObject(GUnrealEd, &UUnrealEdEngine::IsComponentSelected);
+											}
+										}
+									}
+								}
 							}
 						}
 					}
diff --git a/Engine/Source/Editor/UnrealEd/Private/EditorSelectUtils.cpp b/Engine/Source/Editor/UnrealEd/Private/EditorSelectUtils.cpp
--- a/Engine/Source/Editor/UnrealEd/Private/EditorSelectUtils.cpp
+++ b/Engine/Source/Editor/UnrealEd/Private/EditorSelectUtils.cpp
@@ -653,7 +653,7 @@ bool UUnrealEdEngine::IsComponentSelected(const UPrimitiveComponent* PrimCompone
 {
 	if (GetSelectedComponentCount() > 0)
 	{
-		return GetSelectedComponents()->IsSelected(PrimComponent);
+		return GetSelectedComponents()->IsSelected(PrimComponent->IsEditorOnly() ? PrimComponent->AttachParent : PrimComponent);
 	}
 
 	return GetSelectedActors()->IsSelected(PrimComponent->GetOwner());
diff --git a/Engine/Source/Editor/UnrealEd/Private/Kismet2/ComponentEditorUtils.cpp b/Engine/Source/Editor/UnrealEd/Private/Kismet2/ComponentEditorUtils.cpp
--- a/Engine/Source/Editor/UnrealEd/Private/Kismet2/ComponentEditorUtils.cpp
+++ b/Engine/Source/Editor/UnrealEd/Private/Kismet2/ComponentEditorUtils.cpp
@@ -107,15 +107,17 @@ bool FComponentEditorUtils::IsValidVariableNameString(const UActorComponent* InC
 	return bIsValid;
 }
 
-FString FComponentEditorUtils::GenerateValidVariableName(const UActorComponent* InComponent)
+FString FComponentEditorUtils::GenerateValidVariableName(TSubclassOf<UActorComponent> ComponentClass, AActor* ComponentOwner)
 {
+	check(ComponentOwner);
+
 	int32 Counter = 1;
-	FString ComponentTypeName = *InComponent->GetClass()->GetName().Replace(TEXT("Component"), TEXT(""));
+	FString ComponentTypeName = *ComponentClass->GetName().Replace(TEXT("Component"), TEXT(""));
 	FString ComponentInstanceName;
 
-	// Make sure that none of the current components have the same name
+	// Make sure that none of the components currently on the actor have the same name
 	TInlineComponentArray<UActorComponent*> Components;
-	InComponent->GetOwner()->GetComponents(Components);
+	ComponentOwner->GetComponents(Components);
 
 	bool bNameIsValid = false;
 	while (!bNameIsValid)
diff --git a/Engine/Source/Editor/UnrealEd/Private/LevelViewportClickHandlers.cpp b/Engine/Source/Editor/UnrealEd/Private/LevelViewportClickHandlers.cpp
--- a/Engine/Source/Editor/UnrealEd/Private/LevelViewportClickHandlers.cpp
+++ b/Engine/Source/Editor/UnrealEd/Private/LevelViewportClickHandlers.cpp
@@ -284,12 +284,18 @@ namespace ClickHandlers
 				break;
 			}
 		}
+
+		//If the component selected is editor-only, we want to select the non-editor-only component it's attached to
+		while (Component != nullptr && Component->IsEditorOnly())
+		{
+			Component = Component->AttachParent;
+		}
 		
-		if (!Component)
+		if (!ensure(Component != nullptr))
 		{
 			return false;
 		}
-		
+
 		// Pivot snapping
 		if (Click.GetKey() == EKeys::MiddleMouseButton && Click.IsAltDown())
 		{
diff --git a/Engine/Source/Editor/UnrealEd/Public/Kismet2/ComponentEditorUtils.h b/Engine/Source/Editor/UnrealEd/Public/Kismet2/ComponentEditorUtils.h
--- a/Engine/Source/Editor/UnrealEd/Public/Kismet2/ComponentEditorUtils.h
+++ b/Engine/Source/Editor/UnrealEd/Public/Kismet2/ComponentEditorUtils.h
@@ -17,7 +17,7 @@ public:
 	static bool IsValidVariableNameString(const UActorComponent* InComponent, const FString& InString);
 
 	/** Generate a valid variable name string for the given component instance */
-	static FString GenerateValidVariableName(const UActorComponent* InComponent);
+	static FString GenerateValidVariableName(TSubclassOf<UActorComponent> InComponentClass, AActor* ComponentOwner);
 
 	struct FTransformData
 	{
