[INTEGRATE] Change 2417399 by [REDACTED_PERSON][REDACTED_HANDLE] on 2015/01/23 17:45:16

	Disable component editing for non-blueprintable types and for multi-selection

diff --git a/Engine/Source/Editor/LevelEditor/Private/SLevelEditor.cpp b/Engine/Source/Editor/LevelEditor/Private/SLevelEditor.cpp
--- a/Engine/Source/Editor/LevelEditor/Private/SLevelEditor.cpp
+++ b/Engine/Source/Editor/LevelEditor/Private/SLevelEditor.cpp
@@ -610,8 +610,8 @@ public:
 			DetailsView->SetObjects( InObjects );
 			
 			bool bShowingComponents = false;
-			
-			if( GetDefault<UEditorExperimentalSettings>()->bInWorldBPEditing )
+
+			if( InObjects.Num() == 1 && FKismetEditorUtilities::CanCreateBlueprintOfClass( InObjects[0]->GetClass() ) && GetDefault<UEditorExperimentalSettings>()->bInWorldBPEditing )
 			{
 				auto Actor = GetSelectedActor();
 				if (Actor)
