Fixed crash when undoing converting BSP to Static Mesh.

diff --git a/Engine/Source/Editor/LevelEditor/Private/SLevelEditor.cpp b/Engine/Source/Editor/LevelEditor/Private/SLevelEditor.cpp
--- a/Engine/Source/Editor/LevelEditor/Private/SLevelEditor.cpp
+++ b/Engine/Source/Editor/LevelEditor/Private/SLevelEditor.cpp
@@ -733,7 +733,7 @@ private:
 					TGuardValue<bool> SelectionGuard(bSelectionGuard, true);
 
 					// Update the editor's component selection to match the node selection
-					const FScopedTransaction Transaction(NSLOCTEXT("UnrealEd", "ClickingOnComponentInTree", "Clicking on Component (tree view)"));
+					const FScopedTransaction Transaction(NSLOCTEXT("UnrealEd", "ClickingOnComponentInTree", "Clicking on Component (tree view)"), !GIsTransacting);
 
 					SelectedComponents->Modify();
 					SelectedComponents->BeginBatchSelectOperation();
