[INTEGRATE] Change 2420317 by [REDACTED_PERSON][REDACTED_HANDLE] on 2015/01/27 09:56:06

	Move SActorDetails to another file

diff --git a/Engine/Source/Editor/LevelEditor/Private/SActorDetails.cpp b/Engine/Source/Editor/LevelEditor/Private/SActorDetails.cpp
new file mode 100644
--- /dev/null
+++ b/Engine/Source/Editor/LevelEditor/Private/SActorDetails.cpp
@@ -0,0 +1,318 @@
+// Copyright 1998-2015 Epic Games, Inc. All Rights Reserved.
+
+#include "LevelEditor.h"
+#include "SActorDetails.h"
+#include "SSCSEditor.h"
+#include "ComponentEditorUtils.h"
+#include "Editor/PropertyEditor/Public/PropertyEditing.h"
+#include "LevelEditorGenericDetails.h"
+#include "ScopedTransaction.h"
+
+void SActorDetails::Construct(const FArguments& InArgs, const FName TabIdentifier)
+{
+	struct Local
+	{
+		static bool IsPropertyVisible(const FPropertyAndParent& PropertyAndParent)
+		{
+			// For details views in the level editor all properties are the instanced versions
+			if(PropertyAndParent.Property.HasAllPropertyFlags(CPF_DisableEditOnInstance))
+			{
+				return false;
+			}
+
+			return true;
+		}
+	};
+
+	bSelectionGuard = false;
+
+	// Event subscriptions
+	USelection::SelectionChangedEvent.AddRaw(this, &SActorDetails::OnEditorSelectionChanged);
+
+	FPropertyEditorModule& PropPlugin = FModuleManager::LoadModuleChecked<FPropertyEditorModule>("PropertyEditor");
+	FDetailsViewArgs DetailsViewArgs;
+	DetailsViewArgs.bUpdatesFromSelection = true;
+	DetailsViewArgs.bLockable = true;
+	DetailsViewArgs.NameAreaSettings = FDetailsViewArgs::ComponentsAndActorsUseNameArea;
+	DetailsViewArgs.NotifyHook = GUnrealEd;
+	DetailsViewArgs.ViewIdentifier = TabIdentifier;
+	DetailsViewArgs.bCustomNameAreaLocation = true;
+	DetailsViewArgs.bCustomFilterAreaLocation = true;
+
+	DetailsView = PropPlugin.CreateDetailView(DetailsViewArgs);
+
+	DetailsView->SetIsPropertyVisibleDelegate(FIsPropertyVisible::CreateStatic(&Local::IsPropertyVisible));
+
+	// Set up a delegate to call to add generic details to the view
+	DetailsView->SetGenericLayoutDetailsDelegate(FOnGetDetailCustomizationInstance::CreateStatic(&FLevelEditorGenericDetails::MakeInstance));
+
+	GEditor->RegisterForUndo(this);
+
+	SAssignNew(ComponentsBox, SBox)
+		.Visibility(EVisibility::Collapsed);
+
+	ComponentsBox->SetContent
+	(
+		SAssignNew(SCSEditor, SSCSEditor)
+		.EditorMode(EComponentEditorMode::ActorInstance)
+		.ActorContext(this, &SActorDetails::GetSelectedActor)												// Get the instance of the actor in the world
+		.OnSelectionUpdated(this, &SActorDetails::OnSCSEditorTreeViewSelectionChanged)						// A selection has been made in the tree view, so inform the level editor
+		//.OnHighlightPropertyInDetailsView(this, &SLevelEditor::OnSCSEditorHighlightPropertyInDetailsView)	// Also unsure and don't think it's needed
+	);
+
+	ChildSlot
+	[
+		SNew(SVerticalBox)
+		+SVerticalBox::Slot()
+		.Padding(0.0f, 0.0f, 0.0f, 2.0f)
+		.AutoHeight()
+		[
+			DetailsView->GetNameAreaWidget().ToSharedRef()
+		]
+		+ SVerticalBox::Slot()
+			.Padding(0.0f, 0.0f, 0.0f, 2.0f)
+			.AutoHeight()
+			[
+				DetailsView->GetFilterAreaWidget().ToSharedRef()
+			]
+		+SVerticalBox::Slot()
+			[
+				SAssignNew(DetailsSplitter, SSplitter)
+				.Orientation(Orient_Vertical)
+				+ SSplitter::Slot()
+				[
+					DetailsView.ToSharedRef()
+				]
+			]
+	];
+
+	DetailsSplitter->AddSlot(0)
+	.Value(.2f)
+	[
+		ComponentsBox.ToSharedRef()
+	];
+}
+
+SActorDetails::~SActorDetails()
+{
+	USelection::SelectionChangedEvent.RemoveAll(this);
+}
+
+void SActorDetails::SetObjects(const TArray<UObject*>& InObjects)
+{
+	if(!DetailsView->IsLocked())
+	{
+		DetailsView->SetObjects(InObjects);
+
+		bool bShowingComponents = false;
+
+		if(InObjects.Num() == 1 && FKismetEditorUtilities::CanCreateBlueprintOfClass(InObjects[0]->GetClass()) && GetDefault<UEditorExperimentalSettings>()->bInWorldBPEditing)
+		{
+			auto Actor = GetSelectedActor();
+			if(Actor)
+			{
+				bShowingComponents = true;
+
+				// Update the tree if a new actor is selected
+				if(GEditor->GetSelectedComponentCount() == 0)
+				{
+					SCSEditor->UpdateTree();
+				}
+			}
+		}
+
+		ComponentsBox->SetVisibility(bShowingComponents ? EVisibility::Visible : EVisibility::Collapsed);
+	}
+}
+
+void SActorDetails::PostUndo(bool bSuccess)
+{
+	// Enable the selection guard to prevent OnTreeSelectionChanged() from altering the editor's component selection
+	TGuardValue<bool> SelectionGuard(bSelectionGuard, true);
+
+	// Refresh the tree and update the selection to match the world
+	SCSEditor->UpdateTree();
+	UpdateComponentTreeFromEditorSelection();
+}
+
+void SActorDetails::PostRedo(bool bSuccess)
+{
+	PostUndo(bSuccess);
+}
+
+void SActorDetails::OnEditorSelectionChanged(UObject* Object)
+{
+	if(!bSelectionGuard && SCSEditor.IsValid())
+	{
+		// Make sure the selection set that changed is relevant to us
+		auto Selection = Cast<USelection>(Object);
+		if(Selection == GEditor->GetSelectedComponents() || Selection == GEditor->GetSelectedActors())
+		{
+			UpdateComponentTreeFromEditorSelection();
+
+			if(GEditor->GetSelectedComponentCount() == 0) // An actor was selected
+			{
+				// Ensure the selection flags are up to date for the components in the selected actor
+				for(FSelectionIterator It(GEditor->GetSelectedActorIterator()); It; ++It)
+				{
+					auto Actor = CastChecked<AActor>(*It);
+					GUnrealEd->SetActorSelectionFlags(Actor);
+				}
+			}
+		}
+	}
+}
+
+AActor* SActorDetails::GetSelectedActor() const
+{
+	//@todo this won't work w/ multi-select
+	return Cast<AActor>(*GEditor->GetSelectedActorIterator());
+}
+
+void SActorDetails::OnSCSEditorRootSelected(AActor* Actor)
+{
+	if(!bSelectionGuard)
+	{
+		GEditor->SelectNone(true, true, false);
+		GEditor->SelectActor(Actor, true, true, true);
+	}
+}
+
+void SActorDetails::OnSCSEditorTreeViewSelectionChanged(const TArray<FSCSEditorTreeNodePtrType>& SelectedNodes)
+{
+	if(!bSelectionGuard && SelectedNodes.Num() > 0)
+	{
+		auto Actor = GetSelectedActor();
+		if(Actor)
+		{
+			TArray<UObject*> DetailsObjects;
+
+			bool bActorSelected = false;
+			for(auto& SelectedNode : SelectedNodes)
+			{
+				if(SelectedNode.IsValid() && SelectedNode->GetNodeType() == FSCSEditorTreeNode::RootActorNode)
+				{
+					bActorSelected = true;
+					break;
+				}
+			}
+
+			if(bActorSelected)
+			{
+				DetailsObjects.Add(Actor);
+			}
+
+			USelection* SelectedComponents = GEditor->GetSelectedComponents();
+
+			// Don't bother doing anything if the node selection already matches the current world selection
+			bool bSelectionChanged = GEditor->GetSelectedComponentCount() != SelectedNodes.Num() - (bActorSelected ? 1 : 1);
+			if(!bSelectionChanged)
+			{
+				// Check to see if any of the selected nodes aren't already selected in the world
+				for(auto& SelectedNode : SelectedNodes)
+				{
+					UActorComponent* ComponentInstance = SelectedNode->FindComponentInstanceInActor(Actor);
+					if(ComponentInstance && !SelectedComponents->IsSelected(ComponentInstance))
+					{
+						// At least one of the selected nodes isn't selected in the world, so update the selection
+						bSelectionChanged = true;
+						break;
+					}
+				}
+			}
+
+			if(bActorSelected || bSelectionChanged)
+			{
+				// Enable the selection guard to prevent OnEditorSelectionChanged() from altering the contents of the SCSTreeWidget
+				TGuardValue<bool> SelectionGuard(bSelectionGuard, true);
+
+				// Update the editor's component selection to match the node selection
+				const FScopedTransaction Transaction(NSLOCTEXT("UnrealEd", "ClickingOnComponentInTree", "Clicking on Component (tree view)"), !GIsTransacting);
+
+				SelectedComponents->Modify();
+				SelectedComponents->BeginBatchSelectOperation();
+				SelectedComponents->DeselectAll();
+
+				for(auto& SelectedNode : SelectedNodes)
+				{
+					if(SelectedNode.IsValid())
+					{
+						UActorComponent* ComponentInstance = SelectedNode->FindComponentInstanceInActor(Actor);
+						if(ComponentInstance)
+						{
+							DetailsObjects.Add(ComponentInstance);
+							SelectedComponents->Select(ComponentInstance);
+
+							// Ensure the selection override is bound for this component (including any attached editor-only children)
+							auto PrimComponent = Cast<UPrimitiveComponent>(ComponentInstance);
+							if(PrimComponent && !PrimComponent->SelectionOverrideDelegate.IsBound())
+							{
+								PrimComponent->SelectionOverrideDelegate = UPrimitiveComponent::FSelectionOverride::CreateUObject(GUnrealEd, &UUnrealEdEngine::IsComponentSelected);
+							}
+							else
+							{
+								//@todo move the selection override binding check to FComponentEditorUtils
+								auto SceneComponent = Cast<USceneComponent>(ComponentInstance);
+								if(SceneComponent)
+								{
+									for(auto Component : SceneComponent->AttachChildren)
+									{
+										PrimComponent = Cast<UPrimitiveComponent>(Component);
+										if(PrimComponent && !PrimComponent->SelectionOverrideDelegate.IsBound())
+										{
+											PrimComponent->SelectionOverrideDelegate = UPrimitiveComponent::FSelectionOverride::CreateUObject(GUnrealEd, &UUnrealEdEngine::IsComponentSelected);
+										}
+									}
+								}
+							}
+						}
+					}
+				}
+
+				SelectedComponents->EndBatchSelectOperation();
+
+				DetailsView->SetObjects(DetailsObjects);
+
+				GUnrealEd->SetActorSelectionFlags(Actor);
+				GUnrealEd->UpdatePivotLocationForSelection(true);
+				GEditor->RedrawLevelEditingViewports();
+			}
+		}
+	}
+}
+
+void SActorDetails::UpdateComponentTreeFromEditorSelection()
+{
+	// Enable the selection guard to prevent OnTreeSelectionChanged() from altering the editor's component selection
+	TGuardValue<bool> SelectionGuard(bSelectionGuard, true);
+
+	auto& SCSTreeWidget = SCSEditor->SCSTreeWidget;
+	TArray<UObject*> DetailsObjects;
+
+	// Update the tree selection to match the level editor component selection
+	SCSTreeWidget->ClearSelection();
+	for(FSelectionIterator It(GEditor->GetSelectedComponentIterator()); It; ++It)
+	{
+		UActorComponent* Component = CastChecked<UActorComponent>(*It);
+
+		auto SCSTreeNode = SCSEditor->GetNodeFromActorComponent(Component, false);
+		if(SCSTreeNode.IsValid() && SCSTreeNode->GetComponentTemplate())
+		{
+			SCSTreeWidget->RequestScrollIntoView(SCSTreeNode);
+			SCSTreeWidget->SetItemSelection(SCSTreeNode, true);
+
+			auto ComponentTemplate = SCSTreeNode->GetComponentTemplate();
+			check(Component == ComponentTemplate);
+			DetailsObjects.Add(Component);
+		}
+	}
+
+	if(DetailsObjects.Num() > 0)
+	{
+		DetailsView->SetObjects(DetailsObjects);
+	}
+	else
+	{
+		SCSEditor->SelectRoot();
+	}
+}
diff --git a/Engine/Source/Editor/LevelEditor/Private/SActorDetails.h b/Engine/Source/Editor/LevelEditor/Private/SActorDetails.h
new file mode 100644
--- /dev/null
+++ b/Engine/Source/Editor/LevelEditor/Private/SActorDetails.h
@@ -0,0 +1,42 @@
+// Copyright 1998-2015 Epic Games, Inc. All Rights Reserved.
+
+#pragma once
+
+/**
+ * Wraps a details panel customized for viewing actors
+ */
+class SActorDetails : public SCompoundWidget, public FEditorUndoClient
+{
+public:
+	SLATE_BEGIN_ARGS(SActorDetails){}
+	SLATE_END_ARGS()
+
+	void Construct(const FArguments& InArgs, const FName TabIdentifier);
+	~SActorDetails();
+
+	/**
+	 * Sets the objects to be viewed by the details panel
+	 *
+	 * @param InObjects	The objects to set
+	 */
+	void SetObjects(const TArray<UObject*>& InObjects);
+
+	/** FEditorUndoClient Interface */
+	virtual void PostUndo(bool bSuccess) override;
+	virtual void PostRedo(bool bSuccess) override;
+
+private:
+	void OnEditorSelectionChanged(UObject* Object);
+	AActor* GetSelectedActor() const;
+	void OnSCSEditorRootSelected(AActor* Actor);
+	void OnSCSEditorTreeViewSelectionChanged(const TArray<TSharedPtr<class FSCSEditorTreeNode> >& SelectedNodes);
+	void UpdateComponentTreeFromEditorSelection();
+
+private:
+	TSharedPtr<SSplitter> DetailsSplitter;
+	TSharedPtr<class IDetailsView> DetailsView;
+	TSharedPtr<SBox> ComponentsBox;
+	TSharedPtr<class SSCSEditor> SCSEditor;
+
+	bool bSelectionGuard;
+};
\ No newline at end of file
diff --git a/Engine/Source/Editor/LevelEditor/Private/SLevelEditor.cpp b/Engine/Source/Editor/LevelEditor/Private/SLevelEditor.cpp
--- a/Engine/Source/Editor/LevelEditor/Private/SLevelEditor.cpp
+++ b/Engine/Source/Editor/LevelEditor/Private/SLevelEditor.cpp
@@ -13,7 +13,6 @@
 #include "AssetSelection.h"
 #include "LevelEditorContextMenu.h"
 #include "LevelEditorToolBar.h"
-#include "ScopedTransaction.h"
 #include "SLevelEditorToolBox.h"
 #include "SLevelEditorModeContent.h"
 #include "SLevelEditorBuildAndSubmit.h"
@@ -27,24 +26,19 @@
 #include "Toolkits/ToolkitManager.h"
 #include "Editor/PropertyEditor/Public/PropertyEditorModule.h"
 #include "Editor/ContentBrowser/Public/ContentBrowserModule.h"
-#include "LevelEditorGenericDetails.h"
 #include "Editor/MainFrame/Public/MainFrame.h"
 #include "Editor/WorkspaceMenuStructure/Public/WorkspaceMenuStructureModule.h"
 #include "Editor/Sequencer/Public/ISequencerModule.h"
 #include "Editor/StatsViewer/Public/StatsViewerModule.h"
-#include "Editor/UMGEditor/Public/UMGEditorModule.h"
 #include "EditorModes.h"
 #include "IDocumentation.h"
 #include "NewsFeed.h"
 #include "TutorialMetaData.h"
 #include "SDockTab.h"
+#include "SActorDetails.h"
 
 
 
-#include "SSCSEditor.h"
-#include "ComponentEditorUtils.h"
-#include "SExpandableArea.h"
-
 static const FName LevelEditorBuildAndSubmitTab("LevelEditorBuildAndSubmit");
 static const FName LevelEditorStatsViewerTab("LevelEditorStatsViewer");
 static const FName MainFrameModuleName("MainFrame");
@@ -505,332 +499,6 @@ TSharedRef<FTabManager> SLevelEditor::GetTabManager() const
 	return LevelEditorTabManager.ToSharedRef();
 }
 
-class SActorDetails : public SCompoundWidget, public FEditorUndoClient
-{
-public:
-	SLATE_BEGIN_ARGS( SActorDetails )
-	{}
-	SLATE_END_ARGS()
-
-	void Construct( const FArguments& InArgs, const FName TabIdentifier )
-	{
-		struct Local
-		{
-			static bool IsPropertyVisible( const FPropertyAndParent& PropertyAndParent )
-			{
-				// For details views in the level editor all properties are the instanced versions
-				if (PropertyAndParent.Property.HasAllPropertyFlags(CPF_DisableEditOnInstance))
-				{
-					return false;
-				}
-
-				return true;
-			}
-		};
-
-		bSelectionGuard = false;
-
-		// Event subscriptions
-		USelection::SelectionChangedEvent.AddRaw(this, &SActorDetails::OnEditorSelectionChanged);
-
-		FPropertyEditorModule& PropPlugin = FModuleManager::LoadModuleChecked<FPropertyEditorModule>("PropertyEditor");
-		FDetailsViewArgs DetailsViewArgs;
-		DetailsViewArgs.bUpdatesFromSelection = true;
-		DetailsViewArgs.bLockable = true;
-		DetailsViewArgs.NameAreaSettings = FDetailsViewArgs::ComponentsAndActorsUseNameArea;
-		DetailsViewArgs.NotifyHook = GUnrealEd;
-		DetailsViewArgs.ViewIdentifier = TabIdentifier;
-		DetailsViewArgs.bCustomNameAreaLocation = true;
-		DetailsViewArgs.bCustomFilterAreaLocation = true;
-
-		DetailsView = PropPlugin.CreateDetailView(DetailsViewArgs);
-
-		DetailsView->SetIsPropertyVisibleDelegate( FIsPropertyVisible::CreateStatic( &Local::IsPropertyVisible ) );
-
-		// Set up a delegate to call to add generic details to the view
-		DetailsView->SetGenericLayoutDetailsDelegate( FOnGetDetailCustomizationInstance::CreateStatic( &FLevelEditorGenericDetails::MakeInstance ) );
-
-		GEditor->RegisterForUndo(this);
-
-		SAssignNew(ComponentsBox, SBox)
-		.Visibility(EVisibility::Collapsed);
-
-		ComponentsBox->SetContent
-		(
-			SAssignNew(SCSEditor, SSCSEditor)
-			.EditorMode(EComponentEditorMode::ActorInstance)
-			.ActorContext(this, &SActorDetails::GetSelectedActor)												// Get the instance of the actor in the world
-			.OnSelectionUpdated(this, &SActorDetails::OnSCSEditorTreeViewSelectionChanged)						// A selection has been made in the tree view, so inform the level editor
-			//.OnHighlightPropertyInDetailsView(this, &SLevelEditor::OnSCSEditorHighlightPropertyInDetailsView)	// Also unsure and don't think it's needed
-		);
-
-		ChildSlot
-		[
-			SNew( SVerticalBox )
-			+SVerticalBox::Slot()
-			.Padding(0.0f, 0.0f, 0.0f, 2.0f)
-			.AutoHeight()
-			[
-				DetailsView->GetNameAreaWidget().ToSharedRef()
-			]
-			+ SVerticalBox::Slot()
-			.Padding(0.0f, 0.0f, 0.0f, 2.0f)
-			.AutoHeight()
-			[
-				DetailsView->GetFilterAreaWidget().ToSharedRef()
-			]
-			+SVerticalBox::Slot()
-			[
-				SAssignNew(DetailsSplitter, SSplitter)
-				.Orientation(Orient_Vertical)
-				+ SSplitter::Slot()
-				[
-					DetailsView.ToSharedRef()
-				]
-			]
-		];
-
-		DetailsSplitter->AddSlot(0)
-		.Value(.2f)
-		[
-			ComponentsBox.ToSharedRef()
-		];
-	}
-
-	~SActorDetails()
-	{
-		USelection::SelectionChangedEvent.RemoveAll(this);
-	}
-
-	void SetObjects( const TArray<UObject*>& InObjects )
-	{
-		if(!DetailsView->IsLocked() )
-		{
-			DetailsView->SetObjects( InObjects );
-			
-			bool bShowingComponents = false;
-
-			if( InObjects.Num() == 1 && FKismetEditorUtilities::CanCreateBlueprintOfClass( InObjects[0]->GetClass() ) && GetDefault<UEditorExperimentalSettings>()->bInWorldBPEditing )
-			{
-				auto Actor = GetSelectedActor();
-				if (Actor)
-				{
-					bShowingComponents = true;
-
-					// Update the tree if a new actor is selected
-					if (GEditor->GetSelectedComponentCount() == 0)
-					{
-						SCSEditor->UpdateTree();
-					}
-				}
-			}
-
-			ComponentsBox->SetVisibility(bShowingComponents ? EVisibility::Visible : EVisibility::Collapsed);
-		}
-	}
-
-	virtual void PostUndo(bool bSuccess)
-	{
-		// Enable the selection guard to prevent OnTreeSelectionChanged() from altering the editor's component selection
-		TGuardValue<bool> SelectionGuard(bSelectionGuard, true);
-
-		// Refresh the tree and update the selection to match the world
-		SCSEditor->UpdateTree();
-		UpdateComponentTreeFromEditorSelection();
-	}
-
-	virtual void PostRedo(bool bSuccess)
-	{
-		PostUndo(bSuccess);
-	}
-
-private:
-
-	void OnEditorSelectionChanged(UObject* Object)
-	{
-		if (!bSelectionGuard && SCSEditor.IsValid())
-		{
-			// Make sure the selection set that changed is relevant to us
-			auto Selection = Cast<USelection>(Object);
-			if (Selection == GEditor->GetSelectedComponents() || Selection == GEditor->GetSelectedActors())
-			{
-				UpdateComponentTreeFromEditorSelection();
-
-				if (GEditor->GetSelectedComponentCount() == 0) // An actor was selected
-				{
-					// Ensure the selection flags are up to date for the components in the selected actor
-					for (FSelectionIterator It(GEditor->GetSelectedActorIterator()); It; ++It)
-					{
-						auto Actor = CastChecked<AActor>(*It);
-						GUnrealEd->SetActorSelectionFlags(Actor);
-					}
-				}
-			}
-		}
-	}
-	
-	AActor* GetSelectedActor() const
-	{
-		//@todo this won't work w/ multi-select
-		return Cast<AActor>(*GEditor->GetSelectedActorIterator());
-	}
-
-	void OnSCSEditorRootSelected(AActor* Actor)
-	{
-		if (!bSelectionGuard)
-		{
-			GEditor->SelectNone(true, true, false);
-			GEditor->SelectActor(Actor, true, true, true);
-		}
-	}
-
-	void OnSCSEditorTreeViewSelectionChanged(const TArray<FSCSEditorTreeNodePtrType>& SelectedNodes)
-	{
-		if (!bSelectionGuard && SelectedNodes.Num() > 0 )
-		{
-			auto Actor = GetSelectedActor();
-			if (Actor)
-			{
-				TArray<UObject*> DetailsObjects;
-
-				bool bActorSelected = false;
-				for (auto& SelectedNode : SelectedNodes)
-				{
-					if (SelectedNode.IsValid() && SelectedNode->GetNodeType() == FSCSEditorTreeNode::RootActorNode)
-					{
-						bActorSelected = true;
-						break;
-					}
-				}
-
-				if (bActorSelected)
-				{
-					DetailsObjects.Add(Actor);
-				}
-
-				USelection* SelectedComponents = GEditor->GetSelectedComponents();
-
-				// Don't bother doing anything if the node selection already matches the current world selection
-				bool bSelectionChanged = GEditor->GetSelectedComponentCount() != SelectedNodes.Num() - (bActorSelected ? 1 : 1);
-				if (!bSelectionChanged)
-				{
-					// Check to see if any of the selected nodes aren't already selected in the world
-					for (auto& SelectedNode : SelectedNodes)
-					{
-						UActorComponent* ComponentInstance = SelectedNode->FindComponentInstanceInActor(Actor);
-						if (ComponentInstance && !SelectedComponents->IsSelected(ComponentInstance))
-						{
-							// At least one of the selected nodes isn't selected in the world, so update the selection
-							bSelectionChanged = true;
-							break;
-						}
-					}
-				}
-				
-				if (bActorSelected || bSelectionChanged)
-				{
-					// Enable the selection guard to prevent OnEditorSelectionChanged() from altering the contents of the SCSTreeWidget
-					TGuardValue<bool> SelectionGuard(bSelectionGuard, true);
-
-					// Update the editor's component selection to match the node selection
-					const FScopedTransaction Transaction(NSLOCTEXT("UnrealEd", "ClickingOnComponentInTree", "Clicking on Component (tree view)"), !GIsTransacting);
-
-					SelectedComponents->Modify();
-					SelectedComponents->BeginBatchSelectOperation();
-					SelectedComponents->DeselectAll();
-
-					for (auto& SelectedNode : SelectedNodes)
-					{
-						if (SelectedNode.IsValid())
-						{
-							UActorComponent* ComponentInstance = SelectedNode->FindComponentInstanceInActor(Actor);
-							if (ComponentInstance)
-							{
-								DetailsObjects.Add(ComponentInstance);
-								SelectedComponents->Select(ComponentInstance);
-
-								// Ensure the selection override is bound for this component (including any attached editor-only children)
-								auto PrimComponent = Cast<UPrimitiveComponent>(ComponentInstance);
-								if (PrimComponent && !PrimComponent->SelectionOverrideDelegate.IsBound())
-								{
-									PrimComponent->SelectionOverrideDelegate = UPrimitiveComponent::FSelectionOverride::CreateUObject(GUnrealEd, &UUnrealEdEngine::IsComponentSelected);
-								}
-								else
-								{
-									//@todo move the selection override binding check to FComponentEditorUtils
-									auto SceneComponent = Cast<USceneComponent>(ComponentInstance);
-									if (SceneComponent)
-									{
-										for (auto Component : SceneComponent->AttachChildren)
-										{
-											PrimComponent = Cast<UPrimitiveComponent>(Component);
-											if (PrimComponent && !PrimComponent->SelectionOverrideDelegate.IsBound())
-											{
-												PrimComponent->SelectionOverrideDelegate = UPrimitiveComponent::FSelectionOverride::CreateUObject(GUnrealEd, &UUnrealEdEngine::IsComponentSelected);
-											}
-										}
-									}
-								}
-							}
-						}
-					}
-
-					SelectedComponents->EndBatchSelectOperation();
-
-					DetailsView->SetObjects(DetailsObjects);
-
-					GUnrealEd->SetActorSelectionFlags(Actor);
-					GUnrealEd->UpdatePivotLocationForSelection(true);
-					GEditor->RedrawLevelEditingViewports();
-				}
-			}
-		}
-	}
-	
-	void UpdateComponentTreeFromEditorSelection()
-	{
-		// Enable the selection guard to prevent OnTreeSelectionChanged() from altering the editor's component selection
-		TGuardValue<bool> SelectionGuard(bSelectionGuard, true);
-
-		auto& SCSTreeWidget = SCSEditor->SCSTreeWidget;
-		TArray<UObject*> DetailsObjects;
-
-		// Update the tree selection to match the level editor component selection
-		SCSTreeWidget->ClearSelection();
-		for (FSelectionIterator It(GEditor->GetSelectedComponentIterator()); It; ++It)
-		{
-			UActorComponent* Component = CastChecked<UActorComponent>(*It);
-
-			auto SCSTreeNode = SCSEditor->GetNodeFromActorComponent(Component, false);
-			if (SCSTreeNode.IsValid() && SCSTreeNode->GetComponentTemplate())
-			{
-				SCSTreeWidget->RequestScrollIntoView(SCSTreeNode);
-				SCSTreeWidget->SetItemSelection(SCSTreeNode, true);
-
-				auto ComponentTemplate = SCSTreeNode->GetComponentTemplate();
-				check(Component == ComponentTemplate);
-				DetailsObjects.Add(Component);
-			}
-		}
-
-		if (DetailsObjects.Num() > 0)
-		{
-			DetailsView->SetObjects(DetailsObjects);
-		}
-		else
-		{
-			SCSEditor->SelectRoot();
-		}
-	}
-
-private:
-	TSharedPtr<SSplitter> DetailsSplitter;
-	TSharedPtr<IDetailsView> DetailsView;
-	TSharedPtr<SBox> ComponentsBox;
-	TSharedPtr<SSCSEditor> SCSEditor;
-
-	bool bSelectionGuard;
-};
 
 
 TSharedRef<SDockTab> SLevelEditor::SummonDetailsPanel( FName TabIdentifier )
