[INTEGRATE] Change 2420685 by [REDACTED_PERSON][REDACTED_HANDLE] on 2015/01/27 13:15:05

	[UE-8251] [UE-8216] - Can no longer have a component selected in addition to an actor that isn't the component owner
	- Fixed a number of undo issues with the new in-tree Actor root node
	- Fixed undo issue where going from select component->select different actor->undo would cause the different actor to incorrectly still be flagged as selected
	- Fixed component selection visualization undo issues (components were not being modified when override delegates were bound/unbound)

diff --git a/Engine/Source/Editor/LevelEditor/Private/SActorDetails.cpp b/Engine/Source/Editor/LevelEditor/Private/SActorDetails.cpp
--- a/Engine/Source/Editor/LevelEditor/Private/SActorDetails.cpp
+++ b/Engine/Source/Editor/LevelEditor/Private/SActorDetails.cpp
@@ -25,6 +25,7 @@ void SActorDetails::Construct(const FArguments& InArgs, const FName TabIdentifie
 	};
 
 	bSelectionGuard = false;
+	bShowingRootActorNodeSelected = false;
 
 	// Event subscriptions
 	USelection::SelectionChangedEvent.AddRaw(this, &SActorDetails::OnEditorSelectionChanged);
@@ -180,49 +181,53 @@ void SActorDetails::OnSCSEditorRootSelected(AActor* Actor)
 
 void SActorDetails::OnSCSEditorTreeViewSelectionChanged(const TArray<FSCSEditorTreeNodePtrType>& SelectedNodes)
 {
-	if(!bSelectionGuard && SelectedNodes.Num() > 0)
+	if (!bSelectionGuard && SelectedNodes.Num() > 0)
 	{
 		auto Actor = GetSelectedActor();
-		if(Actor)
+		if (Actor)
 		{
+			USelection* SelectedComponents = GEditor->GetSelectedComponents();
 			TArray<UObject*> DetailsObjects;
 
-			bool bActorSelected = false;
-			for(auto& SelectedNode : SelectedNodes)
+			// Determine if the root actor node is among the selected nodes
+			bool bActorNodeSelected = false;
+			for (auto& SelectedNode : SelectedNodes)
 			{
-				if(SelectedNode.IsValid() && SelectedNode->GetNodeType() == FSCSEditorTreeNode::RootActorNode)
+				if (SelectedNode.IsValid() && SelectedNode->GetNodeType() == FSCSEditorTreeNode::RootActorNode)
 				{
-					bActorSelected = true;
+					bActorNodeSelected = true;
 					break;
 				}
 			}
 
-			if(bActorSelected)
-			{
-				DetailsObjects.Add(Actor);
-			}
-
-			USelection* SelectedComponents = GEditor->GetSelectedComponents();
-
-			// Don't bother doing anything if the node selection already matches the current world selection
-			bool bSelectionChanged = GEditor->GetSelectedComponentCount() != SelectedNodes.Num() - (bActorSelected ? 1 : 1);
-			if(!bSelectionChanged)
+			// Determine if the selected non-root actor nodes differ from the editor component selection
+			bool bComponentSelectionChanged = GEditor->GetSelectedComponentCount() != SelectedNodes.Num() - ( bActorNodeSelected ? 1 : 0 );
+			if (!bComponentSelectionChanged)
 			{
 				// Check to see if any of the selected nodes aren't already selected in the world
-				for(auto& SelectedNode : SelectedNodes)
+				for (auto& SelectedNode : SelectedNodes)
 				{
-					UActorComponent* ComponentInstance = SelectedNode->FindComponentInstanceInActor(Actor);
-					if(ComponentInstance && !SelectedComponents->IsSelected(ComponentInstance))
+					if (SelectedNode.IsValid() && !SelectedNode->GetNodeType() == FSCSEditorTreeNode::RootActorNode)
 					{
-						// At least one of the selected nodes isn't selected in the world, so update the selection
-						bSelectionChanged = true;
-						break;
+						UActorComponent* ComponentInstance = SelectedNode->FindComponentInstanceInActor(Actor);
+						if (ComponentInstance && !SelectedComponents->IsSelected(ComponentInstance))
+						{
+							bComponentSelectionChanged = true;
+							break;
+						}
 					}
 				}
 			}
 
-			if(bActorSelected || bSelectionChanged)
+			// Does the actor selection differ from our previous state?
+			const bool bActorSelectionChanged = bShowingRootActorNodeSelected != bActorNodeSelected;
+
+			// If necessary, update the editor component selection
+			if (bActorSelectionChanged || ( bComponentSelectionChanged && !bActorNodeSelected ))
 			{
+				// Store whether we're now showing the actor root as selected
+				bShowingRootActorNodeSelected = bActorNodeSelected;
+
 				// Enable the selection guard to prevent OnEditorSelectionChanged() from altering the contents of the SCSTreeWidget
 				TGuardValue<bool> SelectionGuard(bSelectionGuard, true);
 
@@ -233,34 +238,44 @@ void SActorDetails::OnSCSEditorTreeViewSelectionChanged(const TArray<FSCSEditorT
 				SelectedComponents->BeginBatchSelectOperation();
 				SelectedComponents->DeselectAll();
 
-				for(auto& SelectedNode : SelectedNodes)
+				// If the actor root is selected, then the editor component selection should remain empty and we only show the Actor's details
+				if (bShowingRootActorNodeSelected)
 				{
-					if(SelectedNode.IsValid())
+					DetailsObjects.Add(Actor);
+				}
+				else
+				{
+					for (auto& SelectedNode : SelectedNodes)
 					{
-						UActorComponent* ComponentInstance = SelectedNode->FindComponentInstanceInActor(Actor);
-						if(ComponentInstance)
+						if (SelectedNode.IsValid())
 						{
-							DetailsObjects.Add(ComponentInstance);
-							SelectedComponents->Select(ComponentInstance);
-
-							// Ensure the selection override is bound for this component (including any attached editor-only children)
-							auto PrimComponent = Cast<UPrimitiveComponent>(ComponentInstance);
-							if(PrimComponent && !PrimComponent->SelectionOverrideDelegate.IsBound())
-							{
-								PrimComponent->SelectionOverrideDelegate = UPrimitiveComponent::FSelectionOverride::CreateUObject(GUnrealEd, &UUnrealEdEngine::IsComponentSelected);
-							}
-							else
+							UActorComponent* ComponentInstance = SelectedNode->FindComponentInstanceInActor(Actor);
+							if (ComponentInstance)
 							{
-								//@todo move the selection override binding check to FComponentEditorUtils
-								auto SceneComponent = Cast<USceneComponent>(ComponentInstance);
-								if(SceneComponent)
+								DetailsObjects.Add(ComponentInstance);
+								SelectedComponents->Select(ComponentInstance);
+
+								// Ensure the selection override is bound for this component (including any attached editor-only children)
+								auto PrimComponent = Cast<UPrimitiveComponent>(ComponentInstance);
+								if (PrimComponent && !PrimComponent->SelectionOverrideDelegate.IsBound())
 								{
-									for(auto Component : SceneComponent->AttachChildren)
+									PrimComponent->Modify();
+									PrimComponent->SelectionOverrideDelegate = UPrimitiveComponent::FSelectionOverride::CreateUObject(GUnrealEd, &UUnrealEdEngine::IsComponentSelected);
+								}
+								else
+								{
+									//@todo move the selection override binding check to FComponentEditorUtils
+									auto SceneComponent = Cast<USceneComponent>(ComponentInstance);
+									if (SceneComponent)
 									{
-										PrimComponent = Cast<UPrimitiveComponent>(Component);
-										if(PrimComponent && !PrimComponent->SelectionOverrideDelegate.IsBound())
+										for (auto Component : SceneComponent->AttachChildren)
 										{
-											PrimComponent->SelectionOverrideDelegate = UPrimitiveComponent::FSelectionOverride::CreateUObject(GUnrealEd, &UUnrealEdEngine::IsComponentSelected);
+											PrimComponent = Cast<UPrimitiveComponent>(Component);
+											if (PrimComponent && !PrimComponent->SelectionOverrideDelegate.IsBound())
+											{
+												PrimComponent->Modify();
+												PrimComponent->SelectionOverrideDelegate = UPrimitiveComponent::FSelectionOverride::CreateUObject(GUnrealEd, &UUnrealEdEngine::IsComponentSelected);
+											}
 										}
 									}
 								}
@@ -268,7 +283,7 @@ void SActorDetails::OnSCSEditorTreeViewSelectionChanged(const TArray<FSCSEditorT
 						}
 					}
 				}
-
+				
 				SelectedComponents->EndBatchSelectOperation();
 
 				DetailsView->SetObjects(DetailsObjects);
diff --git a/Engine/Source/Editor/LevelEditor/Private/SActorDetails.h b/Engine/Source/Editor/LevelEditor/Private/SActorDetails.h
--- a/Engine/Source/Editor/LevelEditor/Private/SActorDetails.h
+++ b/Engine/Source/Editor/LevelEditor/Private/SActorDetails.h
@@ -39,4 +39,7 @@ private:
 	TSharedPtr<class SSCSEditor> SCSEditor;
 
 	bool bSelectionGuard;
+
+	// True if the actor "root" node in the SCS editor is currently shown as selected
+	bool bShowingRootActorNodeSelected;
 };
\ No newline at end of file
diff --git a/Engine/Source/Editor/UnrealEd/Private/EditorSelectUtils.cpp b/Engine/Source/Editor/UnrealEd/Private/EditorSelectUtils.cpp
--- a/Engine/Source/Editor/UnrealEd/Private/EditorSelectUtils.cpp
+++ b/Engine/Source/Editor/UnrealEd/Private/EditorSelectUtils.cpp
@@ -556,6 +556,11 @@ void UUnrealEdEngine::SelectActor(AActor* Actor, bool bInSelected, bool bNotify,
 			GetSelectedActors()->Select( Actor, bInSelected );
 			if (!bInSelected)
 			{
+				if (GetSelectedComponentCount() > 0)
+				{
+					GetSelectedComponents()->Modify();
+				}
+
 				for (UActorComponent* Component : Actor->GetComponents())
 				{
 					GetSelectedComponents()->Deselect( Component );
@@ -564,6 +569,7 @@ void UUnrealEdEngine::SelectActor(AActor* Actor, bool bInSelected, bool bNotify,
 					auto PrimComponent = Cast<UPrimitiveComponent>(Component);
 					if (PrimComponent && PrimComponent->SelectionOverrideDelegate.IsBound())
 					{
+						PrimComponent->Modify();
 						PrimComponent->SelectionOverrideDelegate.Unbind();
 					}
 				}
@@ -576,6 +582,7 @@ void UUnrealEdEngine::SelectActor(AActor* Actor, bool bInSelected, bool bNotify,
 					auto PrimComponent = Cast<UPrimitiveComponent>(Component);
 					if (PrimComponent && !PrimComponent->SelectionOverrideDelegate.IsBound())
 					{
+						PrimComponent->Modify();
 						PrimComponent->SelectionOverrideDelegate = UPrimitiveComponent::FSelectionOverride::CreateUObject(this, &UUnrealEdEngine::IsComponentSelected);
 					}
 				}
diff --git a/Engine/Source/Editor/UnrealEd/Private/EditorServer.cpp b/Engine/Source/Editor/UnrealEd/Private/EditorServer.cpp
--- a/Engine/Source/Editor/UnrealEd/Private/EditorServer.cpp
+++ b/Engine/Source/Editor/UnrealEd/Private/EditorServer.cpp
@@ -1122,7 +1122,57 @@ UTransactor* UEditorEngine::CreateTrans()
 
 void UEditorEngine::PostUndo (bool bSuccess)
 {
-	//Make sure that the proper objects display as selected
+	//Update the actor selection followed by the component selection if needed (note: order is important)
+
+	//Get the list of all selected actors after the operation
+	TArray<AActor*> SelectedActors;
+	for (FSelectionIterator It(GEditor->GetSelectedActorIterator()); It; ++It)
+	{
+		AActor* Actor = CastChecked<AActor>(*It);
+		//if this actor is NOT in a hidden level add it to the list - otherwise de-select it
+		if (FLevelUtils::IsLevelLocked(Actor) == false)
+		{
+			SelectedActors.Add(Actor);
+		}
+		else
+		{
+			GetSelectedActors()->Select(Actor, false);
+		}
+	}
+
+	USelection* Selection = GetSelectedActors();
+	Selection->BeginBatchSelectOperation();
+
+	//Deselect all of the actors that were selected prior to the operation
+	for (int32 OldSelectedActorIndex = OldSelectedActors.Num() - 1; OldSelectedActorIndex >= 0; --OldSelectedActorIndex)
+	{
+		AActor* Actor = OldSelectedActors[OldSelectedActorIndex];
+
+		//To stop us from unselecting and then reselecting again (causing two force update components, we will remove (from both lists) any object that was selected and should continue to be selected
+		int32 FoundIndex;
+		if (SelectedActors.Find(Actor, FoundIndex))
+		{
+			OldSelectedActors.RemoveAt(OldSelectedActorIndex);
+			SelectedActors.RemoveAt(FoundIndex);
+		}
+		else
+		{
+			SelectActor(Actor, false, false);//First false is to deselect, 2nd is to notify
+			Actor->UpdateComponentTransforms();
+		}
+	}
+
+	//Select all of the actors in SelectedActors
+	for (int32 SelectedActorIndex = 0; SelectedActorIndex < SelectedActors.Num(); ++SelectedActorIndex)
+	{
+		AActor* Actor = SelectedActors[SelectedActorIndex];
+		SelectActor(Actor, true, false);	//false is to stop notify which is done below if bOpWasSuccessful
+		Actor->UpdateComponentTransforms();
+	}
+
+	OldSelectedActors.Empty();
+	Selection->EndBatchSelectOperation();
+	
 	if (GetSelectedComponentCount() > 0)
 	{
 		//@todo Check to see if component owner is in a hidden level
@@ -1171,57 +1221,6 @@ void UEditorEngine::PostUndo (bool bSuccess)
 		Selection->MarkBatchDirty();
 		Selection->EndBatchSelectOperation();
 	}
-	else
-	{
-		//Get the list of all selected actors after the operation
-		TArray<AActor*> SelectedActors;
-		for (FSelectionIterator It(GEditor->GetSelectedActorIterator()); It; ++It)
-		{
-			AActor* Actor = CastChecked<AActor>(*It);
-			//if this actor is NOT in a hidden level add it to the list - otherwise de-select it
-			if (FLevelUtils::IsLevelLocked(Actor) == false)
-			{
-				SelectedActors.Add(Actor);
-			}
-			else
-			{
-				GetSelectedActors()->Select(Actor, false);
-			}
-		}
-
-		USelection* Selection = GetSelectedActors();
-		Selection->BeginBatchSelectOperation();
-
-		//Deselect all of the actors that were selected prior to the operation
-		for (int32 OldSelectedActorIndex = OldSelectedActors.Num() - 1; OldSelectedActorIndex >= 0; --OldSelectedActorIndex)
-		{
-			AActor* Actor = OldSelectedActors[OldSelectedActorIndex];
-
-			//To stop us from unselecting and then reselecting again (causing two force update components, we will remove (from both lists) any object that was selected and should continue to be selected
-			int32 FoundIndex;
-			if (SelectedActors.Find(Actor, FoundIndex))
-			{
-				OldSelectedActors.RemoveAt(OldSelectedActorIndex);
-				SelectedActors.RemoveAt(FoundIndex);
-			}
-			else
-			{
-				SelectActor(Actor, false, false);//First false is to deselect, 2nd is to notify
-				Actor->UpdateComponentTransforms();
-			}
-		}
-
-		//Select all of the actors in SelectedActors
-		for (int32 SelectedActorIndex = 0; SelectedActorIndex < SelectedActors.Num(); ++SelectedActorIndex)
-		{
-			AActor* Actor = SelectedActors[SelectedActorIndex];
-			SelectActor(Actor, true, false);	//false is to stop notify which is done below if bOpWasSuccessful
-			Actor->UpdateComponentTransforms();
-		}
-
-		OldSelectedActors.Empty();
-		Selection->EndBatchSelectOperation();
-	}
 }
 
 bool UEditorEngine::UndoTransaction()
diff --git a/Engine/Source/Editor/UnrealEd/Private/LevelViewportClickHandlers.cpp b/Engine/Source/Editor/UnrealEd/Private/LevelViewportClickHandlers.cpp
--- a/Engine/Source/Editor/UnrealEd/Private/LevelViewportClickHandlers.cpp
+++ b/Engine/Source/Editor/UnrealEd/Private/LevelViewportClickHandlers.cpp
@@ -215,7 +215,10 @@ namespace ClickHandlers
 					const FScopedTransaction Transaction( NSLOCTEXT("UnrealEd", "ClickingOnActors", "Clicking on Actors") );
 					GEditor->GetSelectedActors()->Modify();
 
-					if( Click.IsControlDown() )
+					// Ctrl- or shift- clicking an actor is the same as regular clicking when components are selected
+					const bool bComponentSelected = GEditor->GetSelectedComponentCount() > 0;
+
+					if( Click.IsControlDown() && !bComponentSelected )
 					{
 						const bool bSelect = !Actor->IsSelected();
 						if( bSelect )
@@ -224,7 +227,7 @@ namespace ClickHandlers
 						}
 						GEditor->SelectActor( Actor, bSelect, true, true );
 					}
-					else if( Click.IsShiftDown() )
+					else if( Click.IsShiftDown() && !bComponentSelected )
 					{
 						if( !Actor->IsSelected() )
 						{
