[INTEGRATE] Change 2420726 by [REDACTED_PERSON][REDACTED_HANDLE] on 2015/01/27 13:34:50

	Improved the validation when drag-and-dropping a class into the level viewport

	UE-7184 - Show game c++ classes and engine c++ classes in the content browser

	It now prevents the placement of non-placeable classes that have come from an asset or class drag (Content Browser or Class Viewer).

diff --git a/Engine/Source/Editor/LevelEditor/Private/SLevelViewport.cpp b/Engine/Source/Editor/LevelEditor/Private/SLevelViewport.cpp
--- a/Engine/Source/Editor/LevelEditor/Private/SLevelViewport.cpp
+++ b/Engine/Source/Editor/LevelEditor/Private/SLevelViewport.cpp
@@ -446,7 +446,7 @@ void SLevelViewport::OnDragEnter( const FGeometry& MyGeometry, const FDragDropEv
 		{
 			if ( HandleDragObjects(MyGeometry, DragDropEvent) )
 			{
-				if ( HandlePlaceDraggedObjects(DragDropEvent, /*bCreateDropPreview=*/true) )
+				if ( HandlePlaceDraggedObjects(MyGeometry, DragDropEvent, /*bCreateDropPreview=*/true) )
 				{
 					DragDropEvent.GetOperation()->SetDecoratorVisibility(false);
 				}
@@ -586,7 +586,7 @@ FReply SLevelViewport::OnDragOver( const FGeometry& MyGeometry, const FDragDropE
 	return FReply::Unhandled();
 }
 
-bool SLevelViewport::HandlePlaceDraggedObjects(const FDragDropEvent& DragDropEvent, bool bCreateDropPreview)
+bool SLevelViewport::HandlePlaceDraggedObjects(const FGeometry& MyGeometry, const FDragDropEvent& DragDropEvent, bool bCreateDropPreview)
 {
 	bool bAllAssetWereLoaded = false;
 	bool bValidDrop = false;
@@ -598,6 +598,12 @@ bool SLevelViewport::HandlePlaceDraggedObjects(const FDragDropEvent& DragDropEve
 		return false;
 	}
 
+	// Don't handle the placement if we couldn't handle the drag
+	if (!HandleDragObjects(MyGeometry, DragDropEvent))
+	{
+		return false;
+	}
+
 	if (Operation->IsOfType<FClassDragDropOp>())
 	{
 		auto ClassOperation = StaticCastSharedPtr<FClassDragDropOp>( Operation );
@@ -801,7 +807,7 @@ FReply SLevelViewport::OnDrop( const FGeometry& MyGeometry, const FDragDropEvent
 
 	if (CurrentLevel && !FLevelUtils::IsLevelLocked(CurrentLevel))
 	{
-		return HandlePlaceDraggedObjects(DragDropEvent, /*bCreateDropPreview=*/false) ? FReply::Handled() : FReply::Unhandled();
+		return HandlePlaceDraggedObjects(MyGeometry, DragDropEvent, /*bCreateDropPreview=*/false) ? FReply::Handled() : FReply::Unhandled();
 	}
 	else
 	{
diff --git a/Engine/Source/Editor/LevelEditor/Public/SLevelViewport.h b/Engine/Source/Editor/LevelEditor/Public/SLevelViewport.h
--- a/Engine/Source/Editor/LevelEditor/Public/SLevelViewport.h
+++ b/Engine/Source/Editor/LevelEditor/Public/SLevelViewport.h
@@ -480,7 +480,7 @@ private:
 	 *	@param DragDropEvent		The drag event.
 	 *	@param bCreateDropPreview	If true, a drop preview actor will be spawned instead of a normal actor.
 	 */
-	bool HandlePlaceDraggedObjects(const FDragDropEvent& DragDropEvent, bool bCreateDropPreview);
+	bool HandlePlaceDraggedObjects(const FGeometry& MyGeometry, const FDragDropEvent& DragDropEvent, bool bCreateDropPreview);
 
 	/** SWidget Interface */
 	virtual FReply OnKeyDown( const FGeometry& MyGeometry, const FKeyEvent& InKeyEvent ) override;
diff --git a/Engine/Source/Editor/UnrealEd/Private/LevelEditorViewport.cpp b/Engine/Source/Editor/UnrealEd/Private/LevelEditorViewport.cpp
--- a/Engine/Source/Editor/UnrealEd/Private/LevelEditorViewport.cpp
+++ b/Engine/Source/Editor/UnrealEd/Private/LevelEditorViewport.cpp
@@ -1121,8 +1121,15 @@ FDropQuery FLevelEditorViewportClient::CanDropObjectsAtCoordinates(int32 MouseX,
 	UObject* AssetObj = AssetData.GetAsset();
 	UClass* ClassObj = Cast<UClass>( AssetObj );
 
-	if ( ClassObj != NULL )
+	if ( ClassObj )
 	{
+		if ( !ObjectTools::IsClassValidForPlacing(ClassObj) )
+		{
+			Result.bCanDrop = false;
+			Result.HintText = FText::Format(LOCTEXT("DragAndDrop_CannotDropAssetClassFmt", "The class '{0}' cannot be placed in a level"), FText::FromString(ClassObj->GetName()));
+			return Result;
+		}
+
 		AssetObj = ClassObj->GetDefaultObject();
 	}
 
diff --git a/Engine/Source/Editor/UnrealEd/Private/ObjectTools.cpp b/Engine/Source/Editor/UnrealEd/Private/ObjectTools.cpp
--- a/Engine/Source/Editor/UnrealEd/Private/ObjectTools.cpp
+++ b/Engine/Source/Editor/UnrealEd/Private/ObjectTools.cpp
@@ -3264,6 +3264,17 @@ namespace ObjectTools
 		return bResult;
 	}
 
+	bool IsClassValidForPlacing(const UClass* InClass)
+	{
+		check(InClass);
+
+		const bool bIsPlaceable = !InClass->HasAllClassFlags(CLASS_NotPlaceable) && (InClass->IsChildOf(AActor::StaticClass()) || InClass->IsChildOf(ABrush::StaticClass()) || InClass->IsChildOf(AVolume::StaticClass()));
+		const bool bIsAbstractOrDeprecated = InClass->HasAnyClassFlags(CLASS_Abstract | CLASS_Deprecated | CLASS_NewerVersionExists);
+		const bool bIsSkeletonClass = FKismetEditorUtilities::IsClassABlueprintSkeleton(InClass);
+		
+		return bIsPlaceable && !bIsAbstractOrDeprecated && !bIsSkeletonClass;
+	}
+
 	bool AreObjectsOfEquivalantType( const TArray<UObject*>& InProposedObjects )
 	{
 		if ( InProposedObjects.Num() > 0 )
diff --git a/Engine/Source/Editor/UnrealEd/Public/ObjectTools.h b/Engine/Source/Editor/UnrealEd/Public/ObjectTools.h
--- a/Engine/Source/Editor/UnrealEd/Public/ObjectTools.h
+++ b/Engine/Source/Editor/UnrealEd/Public/ObjectTools.h
@@ -486,6 +486,15 @@ namespace ObjectTools
 	 */
 	UNREALED_API bool IsAssetValidForPlacing(UWorld* InWorld, const FString& ObjectPath);
 
+	/**
+	 * Determines if the class is placeable in a world.
+	 *
+	 * @param InClass	Class to test.
+	 * 
+	 * @return true if the class can be placed in the world.
+	 */
+	UNREALED_API bool IsClassValidForPlacing(const UClass* InClass);
+
 	/**
 	 * Determines if a given class is a redirector.
 	 *
