[INTEGRATE] Change 2420755 by [REDACTED_PERSON][REDACTED_HANDLE] on 2015/01/27 13:47:11

	Prototype default layout

diff --git a/Engine/Source/Editor/LevelEditor/Private/SLevelEditor.cpp b/Engine/Source/Editor/LevelEditor/Private/SLevelEditor.cpp
--- a/Engine/Source/Editor/LevelEditor/Private/SLevelEditor.cpp
+++ b/Engine/Source/Editor/LevelEditor/Private/SLevelEditor.cpp
@@ -1075,75 +1075,74 @@ TSharedRef<SWidget> SLevelEditor::RestoreContentArea( const TSharedRef<SDockTab>
 		->AddArea
 		(
 			FTabManager::NewPrimaryArea()
-			->SetOrientation( Orient_Vertical )
+			->SetOrientation( Orient_Horizontal )
 			->Split
 			(
 				FTabManager::NewSplitter()
-				->SetOrientation( Orient_Horizontal )
+				->SetOrientation( Orient_Vertical )
+				->SetSizeCoefficient( 1 )
 				->Split
 				(
 					FTabManager::NewSplitter()
-					->SetSizeCoefficient( 0.2f )
-					->SetOrientation(Orient_Vertical)
+					->SetSizeCoefficient( .75f )
+					->SetOrientation(Orient_Horizontal)
 					->Split
 					(
 						FTabManager::NewStack()
-						->SetSizeCoefficient( 0.45f )
+						->SetSizeCoefficient( 0.3f )
 						->AddTab( "LevelEditorToolBox", ETabState::OpenedTab )
 					)
 					->Split
 					(
-						FTabManager::NewStack()->AddTab("ContentBrowserTab1", ETabState::OpenedTab)
+						FTabManager::NewSplitter()
+						->SetOrientation(Orient_Vertical)
+						->SetSizeCoefficient( 1.15f )
+						->Split
+						(
+							FTabManager::NewStack()
+							->SetHideTabWell(true)
+							->AddTab("LevelEditorToolBar", ETabState::OpenedTab)
+						)
+						->Split
+						(
+							FTabManager::NewStack()
+							->SetHideTabWell(true)
+							->SetSizeCoefficient( 1.0f )
+							->AddTab("LevelEditorViewport", ETabState::OpenedTab)
+						)
 					)
 				)
 				->Split
 				(
-					FTabManager::NewSplitter()
-					->SetSizeCoefficient( 0.60f )
-					->SetOrientation(Orient_Vertical)
-					->Split
-					(
-						FTabManager::NewStack()
-						->SetHideTabWell(true)
-						->AddTab( "LevelEditorToolBar", ETabState::OpenedTab )
-					)
-					->Split
-					(
-						FTabManager::NewStack()
-						->SetHideTabWell(true)
-						->SetSizeCoefficient(0.75f)
-						->AddTab( "LevelEditorViewport", ETabState::OpenedTab )
-					)
-					->Split
-					(
-						FTabManager::NewStack()
-						->SetSizeCoefficient(0.25f)
-						->AddTab( "OutputLog", ETabState::ClosedTab )
-					)
+					FTabManager::NewStack()
+					->SetSizeCoefficient(.5)
+					->AddTab("ContentBrowserTab1", ETabState::OpenedTab)
+					->AddTab("OutputLog", ETabState::ClosedTab)
 				)
+			)
+			->Split
+			(
+				FTabManager::NewSplitter()
+				->SetSizeCoefficient(0.25f)
+				->SetOrientation(Orient_Vertical)
 				->Split
 				(
-					FTabManager::NewSplitter()
-					->SetSizeCoefficient( 0.2f )
-					->SetOrientation(Orient_Vertical)
-					->Split
-					(
-						FTabManager::NewStack()
-						->SetSizeCoefficient(0.4f)
-						->AddTab("LevelEditorSceneOutliner", ETabState::OpenedTab)
-						->AddTab("LevelEditorLayerBrowser", ETabState::ClosedTab)	
-					)
-					->Split
-					(
-						FTabManager::NewStack()
-						->AddTab("LevelEditorSelectionDetails", ETabState::OpenedTab)
-						->AddTab("WorldSettingsTab", ETabState::ClosedTab)
-						->SetForegroundTab(FName("LevelEditorSelectionDetails"))
-					)
+					FTabManager::NewStack()
+					->SetSizeCoefficient(0.4f)
+					->AddTab("LevelEditorSceneOutliner", ETabState::OpenedTab)
+					->AddTab("LevelEditorLayerBrowser", ETabState::ClosedTab)
+
+				)
+				->Split
+				(
+					FTabManager::NewStack()
+					->AddTab("LevelEditorSelectionDetails", ETabState::OpenedTab)
+					->AddTab("WorldSettingsTab", ETabState::ClosedTab)
+					->SetForegroundTab(FName("LevelEditorSelectionDetails"))
 				)
 			)
-		)
-	);
+			
+		));
 	
 
 	return LevelEditorTabManager->RestoreFrom( Layout, OwnerWindow ).ToSharedRef();
