[INTEGRATE] Change 2420775 by [REDACTED_PERSON][REDACTED_HANDLE] on 2015/01/27 14:02:15

	Fix the build on Mac. !X==Y does not make sense.

diff --git a/Engine/Source/Editor/LevelEditor/Private/SActorDetails.cpp b/Engine/Source/Editor/LevelEditor/Private/SActorDetails.cpp
--- a/Engine/Source/Editor/LevelEditor/Private/SActorDetails.cpp
+++ b/Engine/Source/Editor/LevelEditor/Private/SActorDetails.cpp
@@ -207,7 +207,7 @@ void SActorDetails::OnSCSEditorTreeViewSelectionChanged(const TArray<FSCSEditorT
 				// Check to see if any of the selected nodes aren't already selected in the world
 				for (auto& SelectedNode : SelectedNodes)
 				{
-					if (SelectedNode.IsValid() && !SelectedNode->GetNodeType() == FSCSEditorTreeNode::RootActorNode)
+					if (SelectedNode.IsValid() && SelectedNode->GetNodeType() != FSCSEditorTreeNode::RootActorNode)
 					{
 						UActorComponent* ComponentInstance = SelectedNode->FindComponentInstanceInActor(Actor);
 						if (ComponentInstance && !SelectedComponents->IsSelected(ComponentInstance))
