[INTEGRATE] Change 2420790 by [REDACTED_PERSON][REDACTED_HANDLE] on 2015/01/27 14:09:29

	Tweak the size of the default content browser

diff --git a/Engine/Source/Editor/LevelEditor/Private/SLevelEditor.cpp b/Engine/Source/Editor/LevelEditor/Private/SLevelEditor.cpp
--- a/Engine/Source/Editor/LevelEditor/Private/SLevelEditor.cpp
+++ b/Engine/Source/Editor/LevelEditor/Private/SLevelEditor.cpp
@@ -1115,7 +1115,7 @@ TSharedRef<SWidget> SLevelEditor::RestoreContentArea( const TSharedRef<SDockTab>
 				->Split
 				(
 					FTabManager::NewStack()
-					->SetSizeCoefficient(.5)
+					->SetSizeCoefficient(.4)
 					->AddTab("ContentBrowserTab1", ETabState::OpenedTab)
 					->AddTab("OutputLog", ETabState::ClosedTab)
 				)
