[INTEGRATE] Change 2421519 by [REDACTED_PERSON][REDACTED_HANDLE] on 2015/01/27 20:33:36

	Fix not being able to redo component transactions in the scseditor.

diff --git a/Engine/Source/Editor/LevelEditor/Private/SActorDetails.cpp b/Engine/Source/Editor/LevelEditor/Private/SActorDetails.cpp
--- a/Engine/Source/Editor/LevelEditor/Private/SActorDetails.cpp
+++ b/Engine/Source/Editor/LevelEditor/Private/SActorDetails.cpp
@@ -232,7 +232,10 @@ void SActorDetails::OnSCSEditorTreeViewSelectionChanged(const TArray<FSCSEditorT
 				TGuardValue<bool> SelectionGuard(bSelectionGuard, true);
 
 				// Update the editor's component selection to match the node selection
-				const FScopedTransaction Transaction(NSLOCTEXT("UnrealEd", "ClickingOnComponentInTree", "Clicking on Component (tree view)"), !GIsTransacting);
+
+				// Note: this transaction should not take place if we are in the middle of executing an undo or redo because it would clear the top of the transaction stack. 
+				const bool bShouldActuallyTransact = !GIsTransacting;
+				const FScopedTransaction Transaction(NSLOCTEXT("UnrealEd", "ClickingOnComponentInTree", "Clicking on Component (tree view)"), bShouldActuallyTransact );
 
 				SelectedComponents->Modify();
 				SelectedComponents->BeginBatchSelectOperation();
