[INTEGRATE] Change 2421574 by [REDACTED_PERSON][REDACTED_HANDLE] on 2015/01/27 21:12:48

	Disable editing of  components created by a construction script on an instance

diff --git a/Engine/Source/Editor/LevelEditor/Private/SActorDetails.cpp b/Engine/Source/Editor/LevelEditor/Private/SActorDetails.cpp
--- a/Engine/Source/Editor/LevelEditor/Private/SActorDetails.cpp
+++ b/Engine/Source/Editor/LevelEditor/Private/SActorDetails.cpp
@@ -44,6 +44,27 @@ void SActorDetails::Construct(const FArguments& InArgs, const FName TabIdentifie
 
 	DetailsView->SetIsPropertyVisibleDelegate(FIsPropertyVisible::CreateStatic(&Local::IsPropertyVisible));
 
+	auto IsPropertyEditingEnabled = [&]()
+	{
+		bool bIsEditable = true;
+		const TArray<TWeakObjectPtr<UObject> >& Objects = DetailsView->GetSelectedObjects();
+		for (auto Object : Objects)
+		{
+			UActorComponent* ActorComp = Cast<UActorComponent>( Object.Get() );
+			if (ActorComp)
+			{
+				bIsEditable = ActorComp->CreationMethod != EComponentCreationMethod::ConstructionScript;
+				if( !bIsEditable )
+				{
+					break;
+				}
+			}
+		}
+		return bIsEditable;
+	};
+
+	DetailsView->SetIsPropertyEditingEnabledDelegate(FIsPropertyEditingEnabled::CreateLambda(IsPropertyEditingEnabled));
+
 	// Set up a delegate to call to add generic details to the view
 	DetailsView->SetGenericLayoutDetailsDelegate(FOnGetDetailCustomizationInstance::CreateStatic(&FLevelEditorGenericDetails::MakeInstance));
 
@@ -71,20 +92,20 @@ void SActorDetails::Construct(const FArguments& InArgs, const FName TabIdentifie
 			DetailsView->GetNameAreaWidget().ToSharedRef()
 		]
 		+ SVerticalBox::Slot()
-			.Padding(0.0f, 0.0f, 0.0f, 2.0f)
-			.AutoHeight()
-			[
-				DetailsView->GetFilterAreaWidget().ToSharedRef()
-			]
+		.Padding(0.0f, 0.0f, 0.0f, 2.0f)
+		.AutoHeight()
+		[
+			DetailsView->GetFilterAreaWidget().ToSharedRef()
+		]
 		+SVerticalBox::Slot()
+		[
+			SAssignNew(DetailsSplitter, SSplitter)
+			.Orientation(Orient_Vertical)
+			+ SSplitter::Slot()
 			[
-				SAssignNew(DetailsSplitter, SSplitter)
-				.Orientation(Orient_Vertical)
-				+ SSplitter::Slot()
-				[
-					DetailsView.ToSharedRef()
-				]
+				DetailsView.ToSharedRef()
 			]
+		]
 	];
 
 	DetailsSplitter->AddSlot(0)
diff --git a/Engine/Source/Editor/PropertyEditor/Private/SDetailSingleItemRow.cpp b/Engine/Source/Editor/PropertyEditor/Private/SDetailSingleItemRow.cpp
--- a/Engine/Source/Editor/PropertyEditor/Private/SDetailSingleItemRow.cpp
+++ b/Engine/Source/Editor/PropertyEditor/Private/SDetailSingleItemRow.cpp
@@ -105,7 +105,7 @@ void SDetailSingleItemRow::Construct( const FArguments& InArgs, FDetailLayoutCus
 		}
 
 		TSharedRef<SWidget> KeyFrameButton = CreateKeyframeButton( *Customization, InOwnerTreeNode );
-		auto IsPropertyEditingEnabled = InOwnerTreeNode->IsPropertyEditingEnabled();
+		TAttribute<bool> IsPropertyEditingEnabled = InOwnerTreeNode->IsPropertyEditingEnabled();
 
 		if( bHasMultipleColumns )
 		{
