[INTEGRATE] Change 2422849 by [REDACTED_PERSON][REDACTED_HANDLE] on 2015/01/28 16:56:45

	[UE-8147] - Details window locks as before, updated work with the component tree
	- Once locked, can change component selection to update the details, but the level selection is not affected
	- After unlocking, any component selection will force the level selection to match

diff --git a/Engine/Source/Editor/LevelEditor/Private/SActorDetails.cpp b/Engine/Source/Editor/LevelEditor/Private/SActorDetails.cpp
--- a/Engine/Source/Editor/LevelEditor/Private/SActorDetails.cpp
+++ b/Engine/Source/Editor/LevelEditor/Private/SActorDetails.cpp
@@ -77,9 +77,8 @@ void SActorDetails::Construct(const FArguments& InArgs, const FName TabIdentifie
 	(
 		SAssignNew(SCSEditor, SSCSEditor)
 		.EditorMode(EComponentEditorMode::ActorInstance)
-		.ActorContext(this, &SActorDetails::GetSelectedActor)												// Get the instance of the actor in the world
-		.OnSelectionUpdated(this, &SActorDetails::OnSCSEditorTreeViewSelectionChanged)						// A selection has been made in the tree view, so inform the level editor
-		//.OnHighlightPropertyInDetailsView(this, &SLevelEditor::OnSCSEditorHighlightPropertyInDetailsView)	// Also unsure and don't think it's needed
+		.ActorContext(this, &SActorDetails::GetActorContext)
+		.OnSelectionUpdated(this, &SActorDetails::OnSCSEditorTreeViewSelectionChanged)
 	);
 
 	ChildSlot
@@ -130,9 +129,10 @@ void SActorDetails::SetObjects(const TArray<UObject*>& InObjects)
 
 		if(InObjects.Num() == 1 && FKismetEditorUtilities::CanCreateBlueprintOfClass(InObjects[0]->GetClass()) && GetDefault<UEditorExperimentalSettings>()->bInWorldBPEditing)
 		{
-			auto Actor = GetSelectedActor();
+			auto Actor = GetSelectedActorInEditor();
 			if(Actor)
 			{
+				LockedActorSelection = Actor;
 				bShowingComponents = true;
 
 				// Update the tree if a new actor is selected
@@ -185,12 +185,28 @@ void SActorDetails::OnEditorSelectionChanged(UObject* Object)
 	}
 }
 
-AActor* SActorDetails::GetSelectedActor() const
+AActor* SActorDetails::GetSelectedActorInEditor() const
 {
 	//@todo this won't work w/ multi-select
 	return Cast<AActor>(*GEditor->GetSelectedActorIterator());
 }
 
+AActor* SActorDetails::GetActorContext() const
+{
+	AActor* SelectedActorInEditor = GetSelectedActorInEditor();
+	const bool bDetailsLocked = DetailsView->IsLocked();
+	
+	// If the details is locked or we have a valid locked selection that doesn't match the editor's selected actor, use the locked selection
+	if (bDetailsLocked || ( LockedActorSelection.IsValid() && LockedActorSelection.Get() != SelectedActorInEditor ))
+	{
+		return LockedActorSelection.Get();
+	}
+	else
+	{
+		return SelectedActorInEditor;
+	}
+}
+
 void SActorDetails::OnSCSEditorRootSelected(AActor* Actor)
 {
 	if(!bSelectionGuard)
@@ -204,10 +220,11 @@ void SActorDetails::OnSCSEditorTreeViewSelectionChanged(const TArray<FSCSEditorT
 {
 	if (!bSelectionGuard && SelectedNodes.Num() > 0)
 	{
-		auto Actor = GetSelectedActor();
+		const bool bDetailsLocked = DetailsView->IsLocked();
+		//AActor* Actor = bDetailsLocked ? LockedActorSelection.Get() : GetSelectedActorInEditor();
+		AActor* Actor = GetActorContext();
 		if (Actor)
 		{
-			USelection* SelectedComponents = GEditor->GetSelectedComponents();
 			TArray<UObject*> DetailsObjects;
 
 			// Determine if the root actor node is among the selected nodes
@@ -221,100 +238,118 @@ void SActorDetails::OnSCSEditorTreeViewSelectionChanged(const TArray<FSCSEditorT
 				}
 			}
 
-			// Determine if the selected non-root actor nodes differ from the editor component selection
-			bool bComponentSelectionChanged = GEditor->GetSelectedComponentCount() != SelectedNodes.Num() - ( bActorNodeSelected ? 1 : 0 );
-			if (!bComponentSelectionChanged)
+			if (bDetailsLocked)
 			{
-				// Check to see if any of the selected nodes aren't already selected in the world
-				for (auto& SelectedNode : SelectedNodes)
+				// When the details panel is locked, we don't want to touch the editor's component selection
+				// We do want to force the locked panel to update to match the selected components, though, since they are part of the actor selection we're locked on
+
+				if (bActorNodeSelected)
+				{
+					// If the actor root is selected, then the editor component selection should remain empty and we only show the Actor's details
+					DetailsObjects.Add(Actor);
+				}
+				else
 				{
-					if (SelectedNode.IsValid() && SelectedNode->GetNodeType() != FSCSEditorTreeNode::RootActorNode)
+					for (auto& SelectedNode : SelectedNodes)
 					{
 						UActorComponent* ComponentInstance = SelectedNode->FindComponentInstanceInActor(Actor);
-						if (ComponentInstance && !SelectedComponents->IsSelected(ComponentInstance))
+						if (ComponentInstance)
 						{
-							bComponentSelectionChanged = true;
-							break;
+							DetailsObjects.Add(ComponentInstance);
 						}
 					}
 				}
-			}
-
-			// Does the actor selection differ from our previous state?
-			const bool bActorSelectionChanged = bShowingRootActorNodeSelected != bActorNodeSelected;
 
-			// If necessary, update the editor component selection
-			if (bActorSelectionChanged || ( bComponentSelectionChanged && !bActorNodeSelected ))
+				const bool bOverrideDetailsLock = true;
+				DetailsView->SetObjects(DetailsObjects, false, bOverrideDetailsLock);
+			}
+			else
 			{
-				// Store whether we're now showing the actor root as selected
-				bShowingRootActorNodeSelected = bActorNodeSelected;
-
 				// Enable the selection guard to prevent OnEditorSelectionChanged() from altering the contents of the SCSTreeWidget
 				TGuardValue<bool> SelectionGuard(bSelectionGuard, true);
 
-				// Update the editor's component selection to match the node selection
-
-				// Note: this transaction should not take place if we are in the middle of executing an undo or redo because it would clear the top of the transaction stack. 
-				const bool bShouldActuallyTransact = !GIsTransacting;
-				const FScopedTransaction Transaction(NSLOCTEXT("UnrealEd", "ClickingOnComponentInTree", "Clicking on Component (tree view)"), bShouldActuallyTransact );
-
-				SelectedComponents->Modify();
-				SelectedComponents->BeginBatchSelectOperation();
-				SelectedComponents->DeselectAll();
-
-				// If the actor root is selected, then the editor component selection should remain empty and we only show the Actor's details
-				if (bShowingRootActorNodeSelected)
+				// Make sure the actor is selected in the editor (possible if the panel was just unlocked, but still assigned to the locked actor)
+				if (!GEditor->GetSelectedActors()->IsSelected(Actor))
 				{
-					DetailsObjects.Add(Actor);
+					GEditor->SelectNone(false, true, false);
+					GEditor->SelectActor(Actor, true, true, true);
 				}
-				else
+
+				USelection* SelectedComponents = GEditor->GetSelectedComponents();
+
+				// Determine if the selected non-root actor nodes differ from the editor component selection
+				bool bComponentSelectionChanged = GEditor->GetSelectedComponentCount() != SelectedNodes.Num() - ( bActorNodeSelected ? 1 : 0 );
+				if (!bComponentSelectionChanged)
 				{
+					// Check to see if any of the selected nodes aren't already selected in the world
 					for (auto& SelectedNode : SelectedNodes)
 					{
-						if (SelectedNode.IsValid())
+						if (SelectedNode.IsValid() && SelectedNode->GetNodeType() != FSCSEditorTreeNode::RootActorNode)
 						{
 							UActorComponent* ComponentInstance = SelectedNode->FindComponentInstanceInActor(Actor);
-							if (ComponentInstance)
+							if (ComponentInstance && !SelectedComponents->IsSelected(ComponentInstance))
 							{
-								DetailsObjects.Add(ComponentInstance);
-								SelectedComponents->Select(ComponentInstance);
+								bComponentSelectionChanged = true;
+								break;
+							}
+						}
+					}
+				}
 
-								// Ensure the selection override is bound for this component (including any attached editor-only children)
-								auto PrimComponent = Cast<UPrimitiveComponent>(ComponentInstance);
-								if (PrimComponent && !PrimComponent->SelectionOverrideDelegate.IsBound())
-								{
-									PrimComponent->Modify();
-									PrimComponent->SelectionOverrideDelegate = UPrimitiveComponent::FSelectionOverride::CreateUObject(GUnrealEd, &UUnrealEdEngine::IsComponentSelected);
-								}
-								else
+				// Does the actor selection differ from our previous state?
+				const bool bActorSelectionChanged = bShowingRootActorNodeSelected != bActorNodeSelected;
+
+				// If necessary, update the editor component selection
+				if (bActorSelectionChanged || ( bComponentSelectionChanged && !bActorNodeSelected ))
+				{
+					// Store whether we're now showing the actor root as selected
+					bShowingRootActorNodeSelected = bActorNodeSelected;
+
+					// Note: this transaction should not take place if we are in the middle of executing an undo or redo because it would clear the top of the transaction stack.
+					const bool bShouldActuallyTransact = !GIsTransacting;
+					const FScopedTransaction Transaction(NSLOCTEXT("UnrealEd", "ClickingOnComponentInTree", "Clicking on Component (tree view)"), bShouldActuallyTransact);
+
+					// Update the editor's component selection to match the node selection
+					SelectedComponents->Modify();
+					SelectedComponents->BeginBatchSelectOperation();
+					SelectedComponents->DeselectAll();
+
+					if (bShowingRootActorNodeSelected)
+					{
+						// If the actor root is selected, then the editor component selection should remain empty and we only show the Actor's details
+						DetailsObjects.Add(Actor);
+					}
+					else
+					{
+						for (auto& SelectedNode : SelectedNodes)
+						{
+							if (SelectedNode.IsValid())
+							{
+								UActorComponent* ComponentInstance = SelectedNode->FindComponentInstanceInActor(Actor);
+								if (ComponentInstance)
 								{
-									//@todo move the selection override binding check to FComponentEditorUtils
+									DetailsObjects.Add(ComponentInstance);
+									SelectedComponents->Select(ComponentInstance);
+
+									// Ensure the selection override is bound for this component (including any attached editor-only children)
 									auto SceneComponent = Cast<USceneComponent>(ComponentInstance);
 									if (SceneComponent)
 									{
-										for (auto Component : SceneComponent->AttachChildren)
-										{
-											PrimComponent = Cast<UPrimitiveComponent>(Component);
-											if (PrimComponent && !PrimComponent->SelectionOverrideDelegate.IsBound())
-											{
-												PrimComponent->Modify();
-												PrimComponent->SelectionOverrideDelegate = UPrimitiveComponent::FSelectionOverride::CreateUObject(GUnrealEd, &UUnrealEdEngine::IsComponentSelected);
-											}
-										}
+										FComponentEditorUtils::BindComponentSelectionOverride(SceneComponent);
 									}
 								}
 							}
 						}
 					}
-				}
-				
-				SelectedComponents->EndBatchSelectOperation();
 
-				DetailsView->SetObjects(DetailsObjects);
+					SelectedComponents->EndBatchSelectOperation();
+
+					DetailsView->SetObjects(DetailsObjects);
 
-				GUnrealEd->SetActorSelectionFlags(Actor);
-				GUnrealEd->UpdatePivotLocationForSelection(true);
-				GEditor->RedrawLevelEditingViewports();
+					GUnrealEd->SetActorSelectionFlags(Actor);
+					GUnrealEd->UpdatePivotLocationForSelection(true);
+					GEditor->RedrawLevelEditingViewports();
+				}
 			}
 		}
 	}
@@ -322,36 +357,39 @@ void SActorDetails::OnSCSEditorTreeViewSelectionChanged(const TArray<FSCSEditorT
 
 void SActorDetails::UpdateComponentTreeFromEditorSelection()
 {
-	// Enable the selection guard to prevent OnTreeSelectionChanged() from altering the editor's component selection
-	TGuardValue<bool> SelectionGuard(bSelectionGuard, true);
-
-	auto& SCSTreeWidget = SCSEditor->SCSTreeWidget;
-	TArray<UObject*> DetailsObjects;
-
-	// Update the tree selection to match the level editor component selection
-	SCSTreeWidget->ClearSelection();
-	for(FSelectionIterator It(GEditor->GetSelectedComponentIterator()); It; ++It)
+	if (!DetailsView->IsLocked())
 	{
-		UActorComponent* Component = CastChecked<UActorComponent>(*It);
+		// Enable the selection guard to prevent OnTreeSelectionChanged() from altering the editor's component selection
+		TGuardValue<bool> SelectionGuard(bSelectionGuard, true);
 
-		auto SCSTreeNode = SCSEditor->GetNodeFromActorComponent(Component, false);
-		if(SCSTreeNode.IsValid() && SCSTreeNode->GetComponentTemplate())
+		auto& SCSTreeWidget = SCSEditor->SCSTreeWidget;
+		TArray<UObject*> DetailsObjects;
+
+		// Update the tree selection to match the level editor component selection
+		SCSTreeWidget->ClearSelection();
+		for (FSelectionIterator It(GEditor->GetSelectedComponentIterator()); It; ++It)
 		{
-			SCSTreeWidget->RequestScrollIntoView(SCSTreeNode);
-			SCSTreeWidget->SetItemSelection(SCSTreeNode, true);
+			UActorComponent* Component = CastChecked<UActorComponent>(*It);
+
+			auto SCSTreeNode = SCSEditor->GetNodeFromActorComponent(Component, false);
+			if (SCSTreeNode.IsValid() && SCSTreeNode->GetComponentTemplate())
+			{
+				SCSTreeWidget->RequestScrollIntoView(SCSTreeNode);
+				SCSTreeWidget->SetItemSelection(SCSTreeNode, true);
 
-			auto ComponentTemplate = SCSTreeNode->GetComponentTemplate();
-			check(Component == ComponentTemplate);
-			DetailsObjects.Add(Component);
+				auto ComponentTemplate = SCSTreeNode->GetComponentTemplate();
+				check(Component == ComponentTemplate);
+				DetailsObjects.Add(Component);
+			}
 		}
-	}
 
-	if(DetailsObjects.Num() > 0)
-	{
-		DetailsView->SetObjects(DetailsObjects);
-	}
-	else
-	{
-		SCSEditor->SelectRoot();
+		if (DetailsObjects.Num() > 0)
+		{
+			DetailsView->SetObjects(DetailsObjects);
+		}
+		else
+		{
+			SCSEditor->SelectRoot();
+		}
 	}
 }
diff --git a/Engine/Source/Editor/LevelEditor/Private/SActorDetails.h b/Engine/Source/Editor/LevelEditor/Private/SActorDetails.h
--- a/Engine/Source/Editor/LevelEditor/Private/SActorDetails.h
+++ b/Engine/Source/Editor/LevelEditor/Private/SActorDetails.h
@@ -26,8 +26,10 @@ public:
 	virtual void PostRedo(bool bSuccess) override;
 
 private:
+	AActor* GetSelectedActorInEditor() const;
+
 	void OnEditorSelectionChanged(UObject* Object);
-	AActor* GetSelectedActor() const;
+	AActor* GetActorContext() const;
 	void OnSCSEditorRootSelected(AActor* Actor);
 	void OnSCSEditorTreeViewSelectionChanged(const TArray<TSharedPtr<class FSCSEditorTreeNode> >& SelectedNodes);
 	void UpdateComponentTreeFromEditorSelection();
@@ -38,6 +40,10 @@ private:
 	TSharedPtr<SBox> ComponentsBox;
 	TSharedPtr<class SSCSEditor> SCSEditor;
 
+	// The actor selected when the details panel was locked
+	TWeakObjectPtr<AActor> LockedActorSelection;
+
+	// Used to prevent reentrant changes
 	bool bSelectionGuard;
 
 	// True if the actor "root" node in the SCS editor is currently shown as selected
diff --git a/Engine/Source/Editor/PropertyEditor/Private/SDetailsView.cpp b/Engine/Source/Editor/PropertyEditor/Private/SDetailsView.cpp
--- a/Engine/Source/Editor/PropertyEditor/Private/SDetailsView.cpp
+++ b/Engine/Source/Editor/PropertyEditor/Private/SDetailsView.cpp
@@ -179,6 +179,9 @@ void SDetailsView::Construct(const FArguments& InArgs)
 	SAssignNew(NameArea, SDetailNameArea, &SelectedObjects)
 		// the name area is only for actors
 		.Visibility(this, &SDetailsView::GetActorNameAreaVisibility)
+		.OnLockButtonClicked(this, &SDetailsView::OnLockButtonClicked)
+		.IsLocked(this, &SDetailsView::IsLocked)
+		.ShowLockButton(DetailsViewArgs.bLockable)
 		.ShowActorLabel(DetailsViewArgs.bShowActorLabel)
 		// only show the selection tip if we're not selecting objects
 		.SelectionTip(!DetailsViewArgs.bHideSelectionTip);
@@ -285,9 +288,9 @@ void SDetailsView::ForceRefresh()
 }
 
 
-void SDetailsView::SetObjects( const TArray<UObject*>& InObjects, bool bForceRefresh/* = false*/ )
+void SDetailsView::SetObjects( const TArray<UObject*>& InObjects, bool bForceRefresh/* = false*/, bool bOverrideLock/* = false*/ )
 {
-	if( !IsLocked() )
+	if (!IsLocked() || bOverrideLock)
 	{
 		TArray< TWeakObjectPtr< UObject > > ObjectWeakPtrs;
 		
@@ -303,9 +306,9 @@ void SDetailsView::SetObjects( const TArray<UObject*>& InObjects, bool bForceRef
 	}
 }
 
-void SDetailsView::SetObjects( const TArray< TWeakObjectPtr< UObject > >& InObjects, bool bForceRefresh/* = false*/ )
+void SDetailsView::SetObjects( const TArray< TWeakObjectPtr< UObject > >& InObjects, bool bForceRefresh/* = false*/, bool bOverrideLock/* = false*/ )
 {
-	if( !IsLocked() )
+	if (!IsLocked() || bOverrideLock)
 	{
 		if( bForceRefresh || ShouldSetNewObjects( InObjects ) )
 		{
diff --git a/Engine/Source/Editor/PropertyEditor/Private/SDetailsView.h b/Engine/Source/Editor/PropertyEditor/Private/SDetailsView.h
--- a/Engine/Source/Editor/PropertyEditor/Private/SDetailsView.h
+++ b/Engine/Source/Editor/PropertyEditor/Private/SDetailsView.h
@@ -27,8 +27,8 @@ public:
 	void Construct(const FArguments& InArgs);
 
 	/** IDetailsView interface */
-	virtual void SetObjects( const TArray<UObject*>& InObjects, bool bForceRefresh = false ) override;
-	virtual void SetObjects( const TArray< TWeakObjectPtr< UObject > >& InObjects, bool bForceRefresh = false ) override;
+	virtual void SetObjects( const TArray<UObject*>& InObjects, bool bForceRefresh = false, bool bOverrideLock = false ) override;
+	virtual void SetObjects( const TArray< TWeakObjectPtr< UObject > >& InObjects, bool bForceRefresh = false, bool bOverrideLock = false ) override;
 	virtual void SetObject( UObject* InObject, bool bForceRefresh = false ) override;
 	virtual void RemoveInvalidObjects() override;
 
diff --git a/Engine/Source/Editor/PropertyEditor/Private/SStructureDetailsView.h b/Engine/Source/Editor/PropertyEditor/Private/SStructureDetailsView.h
--- a/Engine/Source/Editor/PropertyEditor/Private/SStructureDetailsView.h
+++ b/Engine/Source/Editor/PropertyEditor/Private/SStructureDetailsView.h
@@ -81,8 +81,8 @@ public:
 	virtual void SetOnObjectArrayChanged(FOnObjectArrayChanged OnObjectArrayChangedDelegate) override {}
 	virtual void RegisterInstancedCustomPropertyLayout(UClass* Class, FOnGetDetailCustomizationInstance DetailLayoutDelegate) override {}
 	virtual void UnregisterInstancedCustomPropertyLayout(UClass* Class) override {}
-	virtual void SetObjects(const TArray<UObject*>& InObjects, bool bForceRefresh = false) override {}
-	virtual void SetObjects(const TArray< TWeakObjectPtr< UObject > >& InObjects, bool bForceRefresh = false) override {}
+	virtual void SetObjects(const TArray<UObject*>& InObjects, bool bForceRefresh = false, bool bOverrideLock = false) override {}
+	virtual void SetObjects(const TArray< TWeakObjectPtr< UObject > >& InObjects, bool bForceRefresh = false, bool bOverrideLock = false) override {}
 	virtual void SetObject(UObject* InObject, bool bForceRefresh = false) override{}
 	virtual void RemoveInvalidObjects() override {}
 
diff --git a/Engine/Source/Editor/PropertyEditor/Public/IDetailsView.h b/Engine/Source/Editor/PropertyEditor/Public/IDetailsView.h
--- a/Engine/Source/Editor/PropertyEditor/Public/IDetailsView.h
+++ b/Engine/Source/Editor/PropertyEditor/Public/IDetailsView.h
@@ -123,9 +123,10 @@ public:
 	 *
 	 * @param InObjects		The list of objects to observe
 	 * @param bForceRefresh	If true, doesn't check if new objects are being set
+	 * @param bOverrideLock	If true, will set the objects even if the details view is locked
 	 */
-	virtual void SetObjects( const TArray<UObject*>& InObjects, bool bForceRefresh = false ) = 0;
-	virtual void SetObjects( const TArray< TWeakObjectPtr< UObject > >& InObjects, bool bForceRefresh = false ) = 0;
+	virtual void SetObjects( const TArray<UObject*>& InObjects, bool bForceRefresh = false, bool bOverrideLock = false ) = 0;
+	virtual void SetObjects( const TArray< TWeakObjectPtr< UObject > >& InObjects, bool bForceRefresh = false, bool bOverrideLock = false ) = 0;
 
 	/**
 	 * Sets a single objects that details view is viewing
diff --git a/Engine/Source/Editor/UnrealEd/Private/Kismet2/ComponentEditorUtils.cpp b/Engine/Source/Editor/UnrealEd/Private/Kismet2/ComponentEditorUtils.cpp
--- a/Engine/Source/Editor/UnrealEd/Private/Kismet2/ComponentEditorUtils.cpp
+++ b/Engine/Source/Editor/UnrealEd/Private/Kismet2/ComponentEditorUtils.cpp
@@ -203,6 +203,33 @@ void FComponentEditorUtils::AdjustComponentDelta(USceneComponent* Component, FVe
 	}
 }
 
+void FComponentEditorUtils::BindComponentSelectionOverride(USceneComponent* SceneComponent)
+{
+	if (SceneComponent)
+	{
+		// If the scene component is a primitive component, ensure the override is bound
+		auto PrimComponent = Cast<UPrimitiveComponent>(SceneComponent);
+		if (PrimComponent && !PrimComponent->SelectionOverrideDelegate.IsBound())
+		{
+			PrimComponent->Modify();
+			PrimComponent->SelectionOverrideDelegate = UPrimitiveComponent::FSelectionOverride::CreateUObject(GUnrealEd, &UUnrealEdEngine::IsComponentSelected);
+		}
+		else
+		{
+			// Otherwise, make sure the override is bound on any attached primitive components (to make sure we catch any editor-only billboards and the like)
+			for (auto Component : SceneComponent->AttachChildren)
+			{
+				PrimComponent = Cast<UPrimitiveComponent>(Component);
+				if (PrimComponent && !PrimComponent->SelectionOverrideDelegate.IsBound())
+				{
+					PrimComponent->Modify();
+					PrimComponent->SelectionOverrideDelegate = UPrimitiveComponent::FSelectionOverride::CreateUObject(GUnrealEd, &UUnrealEdEngine::IsComponentSelected);
+				}
+			}
+		}
+	}
+}
+
 void FComponentEditorUtils::PropagateTransformPropertyChange(
 	class USceneComponent* InSceneComponentTemplate,
 	const FTransformData& OldDefaultTransform,
diff --git a/Engine/Source/Editor/UnrealEd/Public/Kismet2/ComponentEditorUtils.h b/Engine/Source/Editor/UnrealEd/Public/Kismet2/ComponentEditorUtils.h
--- a/Engine/Source/Editor/UnrealEd/Public/Kismet2/ComponentEditorUtils.h
+++ b/Engine/Source/Editor/UnrealEd/Public/Kismet2/ComponentEditorUtils.h
@@ -31,6 +31,12 @@ public:
 	/** Generate a valid variable name string for the given component instance based on the name of the asset referenced by the component */
 	static FString GenerateValidVariableNameFromAsset(UObject* Asset, AActor* ComponentOwner);
 
+	/**
+	* Ensures that the selection override delegate is bound for the supplied component
+	* This includes any attached editor-only primitive components (such as billboard visualizers)
+	*/
+	static void BindComponentSelectionOverride(USceneComponent* SceneComponent);
+
 	struct FTransformData
 	{
 		FTransformData(const class USceneComponent& Component)
