[INTEGRATE] Change 2424614 by [REDACTED_PERSON][REDACTED_HANDLE] on 2015/01/29 17:55:29

	Editor: Got rid of some silly SAssignNew uses in the SSCSEditor creation

diff --git a/Engine/Source/Editor/LevelEditor/Private/SActorDetails.cpp b/Engine/Source/Editor/LevelEditor/Private/SActorDetails.cpp
--- a/Engine/Source/Editor/LevelEditor/Private/SActorDetails.cpp
+++ b/Engine/Source/Editor/LevelEditor/Private/SActorDetails.cpp
@@ -70,17 +70,16 @@ void SActorDetails::Construct(const FArguments& InArgs, const FName TabIdentifie
 
 	GEditor->RegisterForUndo(this);
 
-	SAssignNew(ComponentsBox, SBox)
+	ComponentsBox = SNew(SBox)
 		.Visibility(EVisibility::Collapsed);
 
-	ComponentsBox->SetContent
-	(
-		SAssignNew(SCSEditor, SSCSEditor)
+	SCSEditor = SNew(SSCSEditor)
 		.EditorMode(EComponentEditorMode::ActorInstance)
 		.ActorContext(this, &SActorDetails::GetActorContext)
 		.OnSelectionUpdated(this, &SActorDetails::OnSCSEditorTreeViewSelectionChanged)
-		.ActorMenuExtender(InArgs._ActorMenuExtender)
-	);
+		.ActorMenuExtender(InArgs._ActorMenuExtender);
+		
+	ComponentsBox->SetContent(SCSEditor.ToSharedRef());
 
 	ChildSlot
 	[
