Editor - Adjusting the mouse capture tip to only ever be collapsed or hit test invisible.  It was previously blocking clicks because the border was visible.

diff --git a/Engine/Source/Editor/LevelEditor/Private/SLevelViewport.cpp b/Engine/Source/Editor/LevelEditor/Private/SLevelViewport.cpp
--- a/Engine/Source/Editor/LevelEditor/Private/SLevelViewport.cpp
+++ b/Engine/Source/Editor/LevelEditor/Private/SLevelViewport.cpp
@@ -3317,19 +3317,17 @@ void SLevelViewport::StartPlayInEditorSession( UGameViewportClient* PlayClient,
 
 EVisibility SLevelViewport::GetMouseCaptureLabelVisibility() const
 {
-	EVisibility visibility = EVisibility::Collapsed;
-
 	if (GEditor->PlayWorld)
 	{
 		// Show the label if the local player's PC isn't set to show the cursor
 		auto const TargetPlayer = GEngine->GetLocalPlayerFromControllerId(GEditor->PlayWorld, 0);
 		if (TargetPlayer && TargetPlayer->PlayerController && !TargetPlayer->PlayerController->bShowMouseCursor)
 		{
-			visibility = EVisibility::Visible;
+			return EVisibility::HitTestInvisible;
 		}
 	}
 
-	return visibility;
+	return EVisibility::Collapsed;
 }
 
 FLinearColor SLevelViewport::GetMouseCaptureLabelColorAndOpacity() const
