[INTEGRATE] Change 2427067 by [REDACTED_PERSON][REDACTED_HANDLE] on 2015/01/31 11:47:42

	Renamed Scene Outliner to World Outliner in UI

diff --git a/Engine/Source/Editor/Layers/Private/SLayersView.h b/Engine/Source/Editor/Layers/Private/SLayersView.h
--- a/Engine/Source/Editor/Layers/Private/SLayersView.h
+++ b/Engine/Source/Editor/Layers/Private/SLayersView.h
@@ -101,7 +101,7 @@ public:
 					.OnItemScrolledIntoView( this, &SLayersView::OnItemScrolledIntoView)
 
 					// Help text 
-					.ToolTipText(LOCTEXT("HelpText", "Drag actors from scene outliner or right click to add a new layer."))
+					.ToolTipText(LOCTEXT("HelpText", "Drag actors from world outliner or right click to add a new layer."))
 				]
 			];
 
diff --git a/Engine/Source/Editor/LevelEditor/Private/SLevelEditor.cpp b/Engine/Source/Editor/LevelEditor/Private/SLevelEditor.cpp
--- a/Engine/Source/Editor/LevelEditor/Private/SLevelEditor.cpp
+++ b/Engine/Source/Editor/LevelEditor/Private/SLevelEditor.cpp
@@ -635,7 +635,7 @@ TSharedRef<SDockTab> SLevelEditor::SpawnLevelEditorTab( const FSpawnTabArgs& Arg
 		}
 
 
-		FText Label = NSLOCTEXT( "LevelEditor", "SceneOutlinerTabTitle", "Scene Outliner" );
+		FText Label = NSLOCTEXT( "LevelEditor", "SceneOutlinerTabTitle", "World Outliner" );
 
 		FSceneOutlinerModule& SceneOutlinerModule = FModuleManager::Get().LoadModuleChecked<FSceneOutlinerModule>( "SceneOutliner" );
 		return SNew( SDockTab )
@@ -1008,8 +1008,8 @@ TSharedRef<SWidget> SLevelEditor::RestoreContentArea( const TSharedRef<SDockTab>
 		{
 			const FSlateIcon OutlinerIcon(FEditorStyle::GetStyleSetName(), "LevelEditor.Tabs.Outliner");
 		    LevelEditorTabManager->RegisterTabSpawner( "LevelEditorSceneOutliner", FOnSpawnTab::CreateSP<SLevelEditor, FName, FString>(this, &SLevelEditor::SpawnLevelEditorTab, FName("LevelEditorSceneOutliner"), FString()) )
-				.SetDisplayName(NSLOCTEXT("LevelEditorTabs", "LevelEditorSceneOutliner", "Scene Outliner"))
-				.SetTooltipText(NSLOCTEXT("LevelEditorTabs", "LevelEditorSceneOutlinerTooltipText", "Open the Scene Outliner tab, which provides a searchable and filterable list of all actors in the scene."))
+				.SetDisplayName(NSLOCTEXT("LevelEditorTabs", "LevelEditorSceneOutliner", "World Outliner"))
+				.SetTooltipText(NSLOCTEXT("LevelEditorTabs", "LevelEditorSceneOutlinerTooltipText", "Open the World Outliner tab, which provides a searchable and filterable list of all actors in the world."))
 				.SetGroup( MenuStructure.GetLevelEditorCategory() )	
 				.SetIcon( OutlinerIcon );	
 		}
@@ -1018,7 +1018,7 @@ TSharedRef<SWidget> SLevelEditor::RestoreContentArea( const TSharedRef<SDockTab>
 			const FSlateIcon LayersIcon(FEditorStyle::GetStyleSetName(), "LevelEditor.Tabs.Layers");
 			LevelEditorTabManager->RegisterTabSpawner( "LevelEditorLayerBrowser", FOnSpawnTab::CreateSP<SLevelEditor, FName, FString>(this, &SLevelEditor::SpawnLevelEditorTab, FName("LevelEditorLayerBrowser"), FString()) )
 				.SetDisplayName(NSLOCTEXT("LevelEditorTabs", "LevelEditorLayerBrowser", "Layers"))
-				.SetTooltipText(NSLOCTEXT("LevelEditorTabs", "LevelEditorLayerBrowserTooltipText", "Open the Layers tab. Use this to manage which actors in the scene belong to which layers."))
+				.SetTooltipText(NSLOCTEXT("LevelEditorTabs", "LevelEditorLayerBrowserTooltipText", "Open the Layers tab. Use this to manage which actors in the world belong to which layers."))
 				.SetGroup( MenuStructure.GetLevelEditorCategory() )
 				.SetIcon( LayersIcon );
 		}
diff --git a/Engine/Source/Editor/SceneOutliner/Private/FolderTreeItem.cpp b/Engine/Source/Editor/SceneOutliner/Private/FolderTreeItem.cpp
--- a/Engine/Source/Editor/SceneOutliner/Private/FolderTreeItem.cpp
+++ b/Engine/Source/Editor/SceneOutliner/Private/FolderTreeItem.cpp
@@ -109,7 +109,7 @@ FDragValidationInfo FFolderDropTarget::ValidateDrop(FDragDropPayload& DraggedObj
 
 void FFolderDropTarget::OnDrop(FDragDropPayload& DraggedObjects, UWorld& World, const FDragValidationInfo& ValidationInfo, TSharedRef<SWidget> DroppedOnWidget)
 {
-	const FScopedTransaction Transaction( LOCTEXT("MoveOutlinerItems", "Move Scene Outliner Items") );
+	const FScopedTransaction Transaction( LOCTEXT("MoveOutlinerItems", "Move World Outliner Items") );
 
 	if (DraggedObjects.Folders)
 	{
diff --git a/Engine/Source/Editor/SceneOutliner/Private/SSceneOutliner.cpp b/Engine/Source/Editor/SceneOutliner/Private/SSceneOutliner.cpp
--- a/Engine/Source/Editor/SceneOutliner/Private/SSceneOutliner.cpp
+++ b/Engine/Source/Editor/SceneOutliner/Private/SSceneOutliner.cpp
@@ -1377,7 +1377,7 @@ namespace SceneOutliner
 			return;
 		}
 
-		const FScopedTransaction Transaction( LOCTEXT("MoveOutlinerItems", "Move Scene Outliner Items") );
+		const FScopedTransaction Transaction( LOCTEXT("MoveOutlinerItems", "Move World Outliner Items") );
 		DropTarget.OnDrop(DraggedObjects, *SharedData->RepresentingWorld, Validation, SNullWidget::NullWidget);
 	}
 
@@ -1631,7 +1631,7 @@ namespace SceneOutliner
 					const bool bSelectEvenIfHidden = true;	// @todo outliner: Is this actually OK?
 					for (auto* Actor : SelectedActors)
 					{
-						UE_LOG(LogSceneOutliner, Verbose,  TEXT("Clicking on Actor (scene outliner): %s (%s)"), *Actor->GetClass()->GetName(), *Actor->GetActorLabel());
+						UE_LOG(LogSceneOutliner, Verbose,  TEXT("Clicking on Actor (world outliner): %s (%s)"), *Actor->GetClass()->GetName(), *Actor->GetActorLabel());
 						GEditor->SelectActor( Actor, bShouldSelect, bNotifyAfterSelect, bSelectEvenIfHidden );
 					}
 
diff --git a/Engine/Source/Editor/SceneOutliner/Private/SSceneOutliner.h b/Engine/Source/Editor/SceneOutliner/Private/SSceneOutliner.h
--- a/Engine/Source/Editor/SceneOutliner/Private/SSceneOutliner.h
+++ b/Engine/Source/Editor/SceneOutliner/Private/SSceneOutliner.h
@@ -5,6 +5,10 @@
 #include "SceneOutlinerStandaloneTypes.h"
 #include "EditorUndoClient.h"
 
+/**
+ * Scene Outliner definition
+ * Note the Scene Outliner is also called the World Outliner
+ */
 namespace SceneOutliner
 {
 	typedef TTextFilter< const ITreeItem& > TreeItemTextFilter;
