[INTEGRATE] Change 2427088 by [REDACTED_PERSON][REDACTED_HANDLE] on 2015/01/31 12:58:20

	Show a warning when viewing blueprinted components. These components are read only and can only be modified in the blueprint editor

diff --git a/Engine/Source/Editor/EditorStyle/Private/SlateEditorStyle.cpp b/Engine/Source/Editor/EditorStyle/Private/SlateEditorStyle.cpp
--- a/Engine/Source/Editor/EditorStyle/Private/SlateEditorStyle.cpp
+++ b/Engine/Source/Editor/EditorStyle/Private/SlateEditorStyle.cpp
@@ -633,6 +633,7 @@ void FSlateEditorStyle::FStyle::SetupGeneralStyles()
 			.SetUnderlineStyle(EditBPHyperlinkButton)
 			.SetTextStyle(NormalText)
 			.SetPadding(FMargin(0.0f));
+
 		Set("EditBPHyperlink", EditBPHyperlink);
 
 		// Visible on hover hyper link
@@ -3045,6 +3046,24 @@ void FSlateEditorStyle::FStyle::SetupPropertyEditorStyles()
 		Set( "PropertyWindow.WindowBorder", new BOX_BRUSH( "Common/GroupBorder", FMargin(4.0f/16.0f) ) );
 		Set( "DetailsView.NameChangeCommitted", new BOX_BRUSH( "Common/EditableTextSelectionBackground", FMargin(4.f/16.f) ) );
 		Set( "DetailsView.HyperlinkStyle", FTextBlockStyle(NormalText) .SetFont( TTF_CORE_FONT( "Fonts/Roboto-Regular", 8 ) ) );
+
+		FTextBlockStyle BPWarningMessageTextStyle = FTextBlockStyle(NormalText) .SetFont(TTF_CORE_FONT("Fonts/Roboto-Regular", 8));
+		FTextBlockStyle BPWarningMessageHyperlinkTextStyle = FTextBlockStyle(BPWarningMessageTextStyle).SetColorAndOpacity(FLinearColor(0.25f, 0.5f, 1.0f));
+
+		FButtonStyle EditBPHyperlinkButton = FButtonStyle()
+			.SetNormal(BORDER_BRUSH("Old/HyperlinkDotted", FMargin(0, 0, 0, 3 / 16.0f), FLinearColor(0.25f, 0.5f, 1.0f)))
+			.SetPressed(FSlateNoResource())
+			.SetHovered(BORDER_BRUSH("Old/HyperlinkUnderline", FMargin(0, 0, 0, 3 / 16.0f), FLinearColor(0.25f, 0.5f, 1.0f)));
+
+		FHyperlinkStyle BPWarningMessageHyperlinkStyle = FHyperlinkStyle()
+			.SetUnderlineStyle(EditBPHyperlinkButton)
+			.SetTextStyle(BPWarningMessageHyperlinkTextStyle)
+			.SetPadding(FMargin(0.0f));
+
+		Set("DetailsView.BPMessageHyperlinkStyle", BPWarningMessageHyperlinkStyle);
+		Set("DetailsView.BPMessageTextStyle", BPWarningMessageTextStyle);
+
+
 		Set( "DetailsView.GroupSection", new BOX_BRUSH( "Common/RoundedSelection_16x", FMargin(4.0f/16.0f) ) );
 
 		Set( "DetailsView.PulldownArrow.Down", new IMAGE_BRUSH( "PropertyView/AdvancedButton_Down", FVector2D(10,8) ) );
diff --git a/Engine/Source/Editor/LevelEditor/Private/SActorDetails.cpp b/Engine/Source/Editor/LevelEditor/Private/SActorDetails.cpp
--- a/Engine/Source/Editor/LevelEditor/Private/SActorDetails.cpp
+++ b/Engine/Source/Editor/LevelEditor/Private/SActorDetails.cpp
@@ -29,26 +29,21 @@ void SActorDetails::Construct(const FArguments& InArgs, const FName TabIdentifie
 
 	DetailsView = PropPlugin.CreateDetailView(DetailsViewArgs);
 
-	auto IsPropertyEditingEnabled = [&]()
+	auto IsPropertyVisible = [](const FPropertyAndParent& PropertyAndParent)
 	{
-		bool bIsEditable = true;
-		const TArray<TWeakObjectPtr<UObject> >& Objects = DetailsView->GetSelectedObjects();
-		for (auto Object : Objects)
+		// For details views in the level editor all properties are the instanced versions
+		if(PropertyAndParent.Property.HasAllPropertyFlags(CPF_DisableEditOnInstance))
 		{
-			UActorComponent* ActorComp = Cast<UActorComponent>( Object.Get() );
-			if (ActorComp)
-			{
-				bIsEditable = ActorComp->CreationMethod != EComponentCreationMethod::ConstructionScript;
-				if( !bIsEditable )
-				{
-					break;
-				}
-			}
+			return false;
 		}
-		return bIsEditable;
+
+		return true;
 	};
 
-	DetailsView->SetIsPropertyEditingEnabledDelegate(FIsPropertyEditingEnabled::CreateLambda(IsPropertyEditingEnabled));
+	DetailsView->SetIsPropertyVisibleDelegate(FIsPropertyVisible::CreateLambda(IsPropertyVisible));
+
+	DetailsView->SetIsPropertyEditingEnabledDelegate(FIsPropertyEditingEnabled::CreateSP(this, &SActorDetails::IsPropertyEditingEnabled));
+
 
 	// Set up a delegate to call to add generic details to the view
 	DetailsView->SetGenericLayoutDetailsDelegate(FOnGetDetailCustomizationInstance::CreateStatic(&FLevelEditorGenericDetails::MakeInstance));
@@ -87,7 +82,43 @@ void SActorDetails::Construct(const FArguments& InArgs, const FName TabIdentifie
 			.Orientation(Orient_Vertical)
 			+ SSplitter::Slot()
 			[
-				DetailsView.ToSharedRef()
+				SNew( SVerticalBox )
+				+SVerticalBox::Slot()
+				.AutoHeight()
+				.Padding( FMargin( 0,0,0,1) )
+				[
+					SNew( SBorder )
+					.Visibility(this, &SActorDetails::GetBlueprintComponentWarningVisibility)
+					.BorderImage( FEditorStyle::Get().GetBrush( "ToolPanel.GroupBorder" ) )
+					[
+						SNew( SHorizontalBox )
+						+ SHorizontalBox::Slot()
+						.AutoWidth()
+						.HAlign( HAlign_Center )
+						.VAlign( VAlign_Center )
+						.Padding( 2 )
+						[
+							SNew( SImage )
+							.Image( FEditorStyle::Get().GetBrush("Icons.Warning") )
+						]
+						+ SHorizontalBox::Slot()
+						.VAlign( VAlign_Center )
+						.Padding( 2 )
+						[
+							SNew(SRichTextBlock)
+							.DecoratorStyleSet(&FEditorStyle::Get())
+							.Justification(ETextJustify::Left)
+							.TextStyle(FEditorStyle::Get(), "DetailsView.BPMessageTextStyle")
+							.Text(NSLOCTEXT("SActorDetails", "BlueprintedComponentWarning", "Blueprinted components must be edited in the <a id=\"HyperlinkDecorator\" style=\"DetailsView.BPMessageHyperlinkStyle\">Blueprint</>"))
+							.AutoWrapText(true)
+							+SRichTextBlock::HyperlinkDecorator(TEXT("HyperlinkDecorator"), this, &SActorDetails::OnBlueprintWarningHyperlinkClicked)
+						]
+					]
+				]
+				+ SVerticalBox::Slot()
+				[
+					DetailsView.ToSharedRef()
+				]
 			]
 		]
 	];
@@ -383,3 +414,37 @@ void SActorDetails::UpdateComponentTreeFromEditorSelection()
 		}
 	}
 }
+
+bool SActorDetails::IsPropertyEditingEnabled() const
+{
+	bool bIsEditable = true;
+	const TArray<TWeakObjectPtr<UObject> >& Objects = DetailsView->GetSelectedObjects();
+	for(auto Object : Objects)
+	{
+		UActorComponent* ActorComp = Cast<UActorComponent>(Object.Get());
+		if(ActorComp)
+		{
+			bIsEditable = ActorComp->CreationMethod != EComponentCreationMethod::ConstructionScript;
+			if(!bIsEditable)
+			{
+				break;
+			}
+		}
+	}
+	return bIsEditable;
+}
+
+void SActorDetails::OnBlueprintWarningHyperlinkClicked( const FSlateHyperlinkRun::FMetadata& Metadata )
+{
+	UBlueprint* Blueprint = SCSEditor->GetBlueprint();
+	if( Blueprint )
+	{
+		// Open the blueprint
+		GEditor->EditObject(Blueprint);
+	}
+}
+
+EVisibility SActorDetails::GetBlueprintComponentWarningVisibility() const
+{
+	return IsPropertyEditingEnabled() ? EVisibility::Collapsed : EVisibility::Visible;
+}
\ No newline at end of file
diff --git a/Engine/Source/Editor/LevelEditor/Private/SActorDetails.h b/Engine/Source/Editor/LevelEditor/Private/SActorDetails.h
--- a/Engine/Source/Editor/LevelEditor/Private/SActorDetails.h
+++ b/Engine/Source/Editor/LevelEditor/Private/SActorDetails.h
@@ -34,7 +34,9 @@ private:
 	void OnSCSEditorRootSelected(AActor* Actor);
 	void OnSCSEditorTreeViewSelectionChanged(const TArray<TSharedPtr<class FSCSEditorTreeNode> >& SelectedNodes);
 	void UpdateComponentTreeFromEditorSelection();
-
+	bool IsPropertyEditingEnabled() const;
+	void OnBlueprintWarningHyperlinkClicked( const FSlateHyperlinkRun::FMetadata& Metadata );
+	EVisibility GetBlueprintComponentWarningVisibility() const;
 private:
 	TSharedPtr<SSplitter> DetailsSplitter;
 	TSharedPtr<class IDetailsView> DetailsView;
