[INTEGRATE] Change 2428391 by [REDACTED_PERSON][REDACTED_HANDLE] on 2015/02/02 13:30:39

	Fixed crash in the light map res tool if the component is not part of the editor world.

diff --git a/Engine/Source/Editor/LevelEditor/Private/LightmapResRatioAdjust.cpp b/Engine/Source/Editor/LevelEditor/Private/LightmapResRatioAdjust.cpp
--- a/Engine/Source/Editor/LevelEditor/Private/LightmapResRatioAdjust.cpp
+++ b/Engine/Source/Editor/LevelEditor/Private/LightmapResRatioAdjust.cpp
@@ -72,7 +72,7 @@ bool FLightmapResRatioAdjustSettings::ApplyRatioAdjustment()
 				// Store the required levels relating to this object based on the level option settings
 				UWorld* MeshWorld = StaticMeshComp->GetWorld();
 				// Exclude preview worlds				
-				if( MeshWorld->WorldType != EWorldType::Preview && MeshWorld->WorldType != EWorldType::Inactive )
+				if( MeshWorld && (MeshWorld->WorldType != EWorldType::Preview && MeshWorld->WorldType != EWorldType::Inactive) )
 				{
 					AddRequiredLevels( Settings.LevelOptions, MeshWorld , MeshLevels );
 					Worlds.AddUnique( MeshWorld );
