Editor: Reverse-integrate to make sure the project name opacity change from CL# 2431751 makes it back intact

diff --git a/Engine/Source/Editor/LevelEditor/Private/LevelEditor.cpp b/Engine/Source/Editor/LevelEditor/Private/LevelEditor.cpp
--- a/Engine/Source/Editor/LevelEditor/Private/LevelEditor.cpp
+++ b/Engine/Source/Editor/LevelEditor/Private/LevelEditor.cpp
@@ -87,7 +87,7 @@ public:
 			.Text(RightContentText)
 			.Visibility(EVisibility::HitTestInvisible)
 			.Font(FSlateFontInfo(FPaths::EngineContentDir() / TEXT("Slate/Fonts/Roboto-Bold.ttf"), 14))
-			.ColorAndOpacity(FLinearColor(1.0f, 1.0f, 1.0f, 0.3f))
+			.ColorAndOpacity(FLinearColor(1.0f, 1.0f, 1.0f, 0.4f))
 		];
 	}
 	
