[INTEGRATE] Change 2430370 by [REDACTED_PERSON][REDACTED_HANDLE] on 2015/02/03 15:13:32

	Editor: Remove "Customize Scripting Behavior" context menu option that went down the anonymous blueprint route; create a named Blueprint via the details panel instead [UE-9013]

diff --git a/Engine/Source/Editor/LevelEditor/Private/LevelEditor.cpp b/Engine/Source/Editor/LevelEditor/Private/LevelEditor.cpp
--- a/Engine/Source/Editor/LevelEditor/Private/LevelEditor.cpp
+++ b/Engine/Source/Editor/LevelEditor/Private/LevelEditor.cpp
@@ -633,11 +633,6 @@ void FLevelEditorModule::BindGlobalLevelEditorCommands()
 		FExecuteAction::CreateStatic( &FLevelEditorActionCallbacks::GoToDocsForActor_Clicked )
 		);
 
-	ActionList.MapAction(
-		Commands.AddScriptBehavior,
-		FExecuteAction::CreateStatic( &FLevelEditorActionCallbacks::AddScriptBehavior_Clicked )
-		);
-
 	ActionList.MapAction( 
 		FGenericCommands::Get().Duplicate, 
 		FExecuteAction::CreateStatic( &FLevelEditorActionCallbacks::ExecuteExecCommand, FString( TEXT("DUPLICATE") ) ),
diff --git a/Engine/Source/Editor/LevelEditor/Private/LevelEditorActions.cpp b/Engine/Source/Editor/LevelEditor/Private/LevelEditorActions.cpp
--- a/Engine/Source/Editor/LevelEditor/Private/LevelEditorActions.cpp
+++ b/Engine/Source/Editor/LevelEditor/Private/LevelEditorActions.cpp
@@ -1012,17 +1012,6 @@ void FLevelEditorActionCallbacks::GoToDocsForActor_Clicked()
 	}
 }
 
-void FLevelEditorActionCallbacks::AddScriptBehavior_Clicked()
-{
-	AActor* SelectedActor = GEditor->GetSelectedActors()->GetTop<AActor>();
-
-	if (SelectedActor)
-	{
-		const FName Name = *FString::Printf(TEXT("%s_BPClass"), *SelectedActor->GetName());
-		FKismetEditorUtilities::CreateBlueprintFromActor(Name, SelectedActor->GetLevel(), SelectedActor, true);
-	}
-}
-
 void FLevelEditorActionCallbacks::FindInContentBrowser_Clicked()
 {
 	GEditor->SyncToContentBrowser();
@@ -2751,8 +2740,6 @@ void FLevelEditorCommands::RegisterCommands()
 	UI_COMMAND( GoToCodeForActor, "Go to C++ Code for Actor", "Opens a code editing IDE and navigates to the source file associated with the seleced actor", EUserInterfaceActionType::Button, FInputGesture() );
 	UI_COMMAND( GoToDocsForActor, "Go to Documentation for Actor", "Opens documentation for the Actor in the default web browser", EUserInterfaceActionType::Button, FInputGesture() );
 
-	UI_COMMAND( AddScriptBehavior, "Customize Scripting Behavior", "Click to customize scripting behavior of this Actor", EUserInterfaceActionType::Button, FInputGesture() );
-
 	UI_COMMAND( PasteHere, "Paste Here", "Pastes the actor at the click location", EUserInterfaceActionType::Button, FInputGesture() );
 
 	UI_COMMAND( SnapOriginToGrid, "Snap Origin to Grid", "Snaps the actor to the nearest grid location at its origin", EUserInterfaceActionType::Button, FInputGesture( EModifierKey::Control, EKeys::End ) );
diff --git a/Engine/Source/Editor/LevelEditor/Private/LevelEditorContextMenu.cpp b/Engine/Source/Editor/LevelEditor/Private/LevelEditorContextMenu.cpp
--- a/Engine/Source/Editor/LevelEditor/Private/LevelEditorContextMenu.cpp
+++ b/Engine/Source/Editor/LevelEditor/Private/LevelEditorContextMenu.cpp
@@ -300,22 +300,6 @@ void FLevelEditorContextMenu::FillMenu( FMenuBuilder& MenuBuilder, TWeakPtr<SLev
 				}
 				MenuBuilder.EndSection();
 			}
-
-			// If it is in the map then it is already anonymous
-			if (GetDefault<UEditorExperimentalSettings>()->bInWorldBPEditing
-				&& SelectedActors.Num() == 1
-				&& !FBlueprintEditorUtils::IsAnonymousBlueprintClass(SelectionInfo.SelectionClass)
-				&& FKismetEditorUtilities::CanCreateBlueprintOfClass(SelectionInfo.SelectionClass))
-			{
-				MenuBuilder.BeginSection("Customize", LOCTEXT("CustomizeHeading", "Customize"));
-				{
-					MenuBuilder.AddMenuEntry(FLevelEditorCommands::Get().AddScriptBehavior,
-						NAME_None,
-						LOCTEXT("AddScriptBehavior", "Customize Scripting Behavior"),
-						LOCTEXT("AddScriptBehavior_ToolTip", "Click to customize scripting behavior of this Actor"));
-				}
-				MenuBuilder.EndSection();
-			}
 		}
 
 		MenuBuilder.BeginSection("ActorSelectVisibilityLevels");
diff --git a/Engine/Source/Editor/LevelEditor/Public/LevelEditorActions.h b/Engine/Source/Editor/LevelEditor/Public/LevelEditorActions.h
--- a/Engine/Source/Editor/LevelEditor/Public/LevelEditorActions.h
+++ b/Engine/Source/Editor/LevelEditor/Public/LevelEditorActions.h
@@ -815,9 +815,6 @@ public:
 	/** Called when 'Go to Documentation for Actor' is clicked */
 	static void GoToDocsForActor_Clicked();
 
-	/** Called when 'Add Script Behavior' is clicked */
-	static void AddScriptBehavior_Clicked();
-
 	/**
 	 * Called when the LockActorMovement command is executed
 	 */
