BSP Surface Properties: Scale now initialized to selected surface scale.

diff --git a/Engine/Source/Editor/LevelEditor/Private/SSurfaceProperties.cpp b/Engine/Source/Editor/LevelEditor/Private/SSurfaceProperties.cpp
--- a/Engine/Source/Editor/LevelEditor/Private/SSurfaceProperties.cpp
+++ b/Engine/Source/Editor/LevelEditor/Private/SSurfaceProperties.cpp
@@ -23,6 +23,31 @@ void SSurfaceProperties::Construct( const FArguments& InArgs )
 	CachedScalingValueU = 1.0f;
 	CachedScalingValueV = 1.0f;
 
+	// Initialize scale fields according to the scale of the first selected surface
+	for (TSelectedSurfaceIterator<> It(GWorld); It; ++It)
+	{
+		FBspSurf* Surf = *It;
+		UModel* Model = It.GetModel();
+
+		const FVector TextureU(Model->Vectors[Surf->vTextureU]);
+		const FVector TextureV(Model->Vectors[Surf->vTextureV]);
+
+		const float TextureUSize = TextureU.Size();
+		const float TextureVSize = TextureV.Size();
+
+		if (!FMath::IsNearlyZero(TextureUSize))
+		{
+			CachedScalingValueU = 1.0f / TextureUSize;
+		}
+
+		if (!FMath::IsNearlyZero(TextureVSize))
+		{
+			CachedScalingValueV = 1.0f / TextureVSize;
+		}
+
+		break;
+	}
+
 	bPreserveScaleRatio = false;
 	bUseRelativeScaling = false;
 
