[INTEGRATE] Change 2431592 by [REDACTED_PERSON][REDACTED_HANDLE] on 2015/02/04 08:39:36

	Fix crash if an actor details panel is closed and then an undo occurs

diff --git a/Engine/Source/Editor/LevelEditor/Private/SActorDetails.cpp b/Engine/Source/Editor/LevelEditor/Private/SActorDetails.cpp
--- a/Engine/Source/Editor/LevelEditor/Private/SActorDetails.cpp
+++ b/Engine/Source/Editor/LevelEditor/Private/SActorDetails.cpp
@@ -132,6 +132,7 @@ void SActorDetails::Construct(const FArguments& InArgs, const FName TabIdentifie
 
 SActorDetails::~SActorDetails()
 {
+	GEditor->UnregisterForUndo(this);
 	USelection::SelectionChangedEvent.RemoveAll(this);
 }
 
