[INTEGRATE] Change 2431751 by [REDACTED_PERSON][REDACTED_HANDLE] on 2015/02/04 10:47:36

	Fixed missing padding after project name in editor title bar
	- Also made the text ever-so-slightly more opaque

diff --git a/Engine/Source/Editor/LevelEditor/Private/LevelEditor.cpp b/Engine/Source/Editor/LevelEditor/Private/LevelEditor.cpp
--- a/Engine/Source/Editor/LevelEditor/Private/LevelEditor.cpp
+++ b/Engine/Source/Editor/LevelEditor/Private/LevelEditor.cpp
@@ -137,6 +137,7 @@ TSharedRef<SDockTab> FLevelEditorModule::SpawnLevelEditor( const FSpawnTabArgs&
 				SNew( SHorizontalBox )
 
 				+SHorizontalBox::Slot()
+				.Padding(0.0f, 0.0f, 14.0f, 0.0f)
 				.AutoWidth()
 				[
 					SNew(SLevelEditorWatermark)
