Additional change needed for CL 2438299.

diff --git a/Engine/Source/Editor/LevelEditor/Private/LevelEditorActions.cpp b/Engine/Source/Editor/LevelEditor/Private/LevelEditorActions.cpp
--- a/Engine/Source/Editor/LevelEditor/Private/LevelEditorActions.cpp
+++ b/Engine/Source/Editor/LevelEditor/Private/LevelEditorActions.cpp
@@ -1054,7 +1054,8 @@ void FLevelEditorActionCallbacks::EditAsset_Clicked( const EToolkitMode::Type To
 	if( GEditor->GetSelectedActorCount() > 0 )
 	{
 		TArray< UObject* > ReferencedAssets;
-		GEditor->GetReferencedAssetsForEditorSelection( ReferencedAssets );
+		const bool bIgnoreOtherAssetsIfBPReferenced = true;
+		GEditor->GetReferencedAssetsForEditorSelection( ReferencedAssets, bIgnoreOtherAssetsIfBPReferenced );
 
 		bool bShouldOpenEditors = (ReferencedAssets.Num() == 1);
 
