Renamed Actor Locking editor functionality to Actor Pilot

Actor pilots are now initiated from the actor's context menu, rather than an unrelated menu on the viewport. Similarly, pilots can be stopped via the context menu for the actor, or from the viewport toolbar that is visible when a pilot is active.

This addresses UE-1091 - Allow actors to be piloted with the camera in UE4.

diff --git a/Engine/Source/Editor/EditorStyle/Private/SlateEditorStyle.cpp b/Engine/Source/Editor/EditorStyle/Private/SlateEditorStyle.cpp
--- a/Engine/Source/Editor/EditorStyle/Private/SlateEditorStyle.cpp
+++ b/Engine/Source/Editor/EditorStyle/Private/SlateEditorStyle.cpp
@@ -4205,6 +4205,20 @@ void FSlateEditorStyle::FStyle::SetupLevelEditorStyle()
 		Set( "LevelViewport.ViewportConfig_FourPanesBottom.Small", new IMAGE_BRUSH( "Icons/ViewportLayout_FourPanesBottom", IconLayoutSizeSmall ) );
 		Set( "LevelViewport.ViewportConfig_FourPanes2x2", new IMAGE_BRUSH( "Icons/ViewportLayout_FourPanes2x2", IconLayoutSize ) );
 		Set( "LevelViewport.ViewportConfig_FourPanes2x2.Small", new IMAGE_BRUSH( "Icons/ViewportLayout_FourPanes2x2", IconLayoutSizeSmall ) );
+
+		Set( "LevelViewport.EjectActorPilot", new IMAGE_BRUSH( "Icons/icon_EjectActorPilot_16x", Icon16x16 ) );
+		Set( "LevelViewport.EjectActorPilot.Small", new IMAGE_BRUSH( "Icons/icon_EjectActorPilot_16x", Icon16x16 ) );
+		Set( "LevelViewport.PilotSelectedActor", new IMAGE_BRUSH( "Icons/icon_PilotSelectedActor_16x", Icon16x16 ) );
+		Set( "LevelViewport.PilotSelectedActor.Small", new IMAGE_BRUSH( "Icons/icon_PilotSelectedActor_16x", Icon16x16 ) );
+		Set( "LevelViewport.ToggleActorPilotCameraView", new IMAGE_BRUSH( "Icons/icon_ToggleActorPilotCameraView_16x", Icon16x16 ) );
+		Set( "LevelViewport.ToggleActorPilotCameraView.Small", new IMAGE_BRUSH( "Icons/icon_ToggleActorPilotCameraView_16x", Icon16x16 ) );
+
+		Set( "LevelViewport.ActorPilotText", FTextBlockStyle()
+			.SetFont( TTF_FONT( "Fonts/Roboto-BoldCondensed", 12 ) )
+			.SetColorAndOpacity( FLinearColor(0.9f, 0.9f, 0.9f, 1.f) )
+			.SetShadowColorAndOpacity( FLinearColor(0.f, 0.f, 0.f, 0.4f) )
+			.SetShadowOffset( FVector2D(1.f, 1.f) )
+		);
 	}
 
 	// Level editor status bar
diff --git a/Engine/Source/Editor/LevelEditor/Private/LevelEditorContextMenu.cpp b/Engine/Source/Editor/LevelEditor/Private/LevelEditorContextMenu.cpp
--- a/Engine/Source/Editor/LevelEditor/Private/LevelEditorContextMenu.cpp
+++ b/Engine/Source/Editor/LevelEditor/Private/LevelEditorContextMenu.cpp
@@ -199,12 +199,6 @@ void FLevelEditorContextMenu::FillMenu( FMenuBuilder& MenuBuilder, TWeakPtr<SLev
 		}
 		MenuBuilder.PushExtender(FExtender::Combine(Extenders).ToSharedRef());
 
-		TArray<TWeakObjectPtr<UObject>> LabelObjects;
-		for (FSelectionIterator SelItor(*GEditor->GetSelectedActors()); SelItor; ++SelItor)
-		{
-			LabelObjects.Add(*SelItor);
-		}
-
 		// Check if current selection has any assets that can be browsed to
 		TArray< UObject* > ReferencedAssets;
 		GEditor->GetReferencedAssetsForEditorSelection(ReferencedAssets);
@@ -263,6 +257,34 @@ void FLevelEditorContextMenu::FillMenu( FMenuBuilder& MenuBuilder, TWeakPtr<SLev
 
 			MenuBuilder.AddMenuEntry(FLevelEditorCommands::Get().GoHere);
 			MenuBuilder.AddMenuEntry(FLevelEditorCommands::Get().SnapCameraToActor);
+
+			if (SelectedActors.Num() == 1)
+			{
+				const FLevelViewportCommands& Actions = FLevelViewportCommands::Get();
+
+				auto Viewport = LevelEditor.Pin()->GetActiveViewport();
+				auto& ViewportClient = Viewport->GetLevelViewportClient();
+
+				if (ViewportClient.IsPerspective() && !ViewportClient.IsLockedToMatinee())
+				{
+				    if (Viewport->IsSelectedActorLocked())
+				    {
+					    MenuBuilder.AddMenuEntry(
+						    Actions.EjectActorPilot,
+						    NAME_None,
+						    FText::Format(LOCTEXT("PilotActor", "Stop piloting '{0}'"), FText::FromString(SelectedActors[0]->GetActorLabel()))
+						    );
+					}
+					else
+					{
+					    MenuBuilder.AddMenuEntry(
+						    Actions.PilotSelectedActor,
+						    NAME_None,
+						    FText::Format(LOCTEXT("PilotActor", "Pilot '{0}'"), FText::FromString(SelectedActors[0]->GetActorLabel()))
+						    );
+					}
+				}
+			}
 		}
 		MenuBuilder.EndSection();
 
diff --git a/Engine/Source/Editor/LevelEditor/Private/LevelViewportActions.cpp b/Engine/Source/Editor/LevelEditor/Private/LevelViewportActions.cpp
--- a/Engine/Source/Editor/LevelEditor/Private/LevelViewportActions.cpp
+++ b/Engine/Source/Editor/LevelEditor/Private/LevelViewportActions.cpp
@@ -27,16 +27,15 @@ void FLevelViewportCommands::RegisterCommands()
 	UI_COMMAND( ToggleMaximize, "Maximize Viewport", "Toggles the Maximize state of the current viewport", EUserInterfaceActionType::ToggleButton, FInputGesture() );
 	UI_COMMAND( ToggleGameView, "Game View", "Toggles game view.  Game view shows the scene as it appears in game", EUserInterfaceActionType::ToggleButton, FInputGesture( EKeys::G ) );
 	UI_COMMAND( ToggleImmersive, "Immersive Mode", "Switches this viewport between immersive mode and regular mode", EUserInterfaceActionType::ToggleButton, FInputGesture( EKeys::F11 ) );
-	UI_COMMAND( ToggleLockedCameraView, "Camera View", "Toggles showing the exact camera view when locking this viewport to a camera", EUserInterfaceActionType::ToggleButton, FInputGesture( EModifierKey::Control | EModifierKey::Shift, EKeys::C ) );
 
 	UI_COMMAND( CreateCamera, "Create Camera Here", "Creates a new camera actor at the current location of this viewport's camera", EUserInterfaceActionType::Button, FInputGesture() );
 	UI_COMMAND( HighResScreenshot, "High Resolution Screenshot...", "Opens the control panel for high resolution screenshots", EUserInterfaceActionType::Button, FInputGesture() );
 	
 	UI_COMMAND( UseDefaultShowFlags, "Use Defaults", "Resets all show flags to default", EUserInterfaceActionType::Button, FInputGesture() );
 
-	UI_COMMAND( ActorUnlock, "Unlock from Actor", "Unlock the viewport's position and orientation from the locked actor.", EUserInterfaceActionType::Button, FInputGesture() );
-	UI_COMMAND( ActorLockSelected, "Lock Selected Actor", "Lock the viewport's position and orientation to the selected actor.", EUserInterfaceActionType::Button, FInputGesture() );
-	UI_COMMAND( ActorUnlockSelected, "Unlock Selected Actor", "Unlock the viewport's position and orientation from the selected actor.", EUserInterfaceActionType::Check, FInputGesture() );
+	UI_COMMAND( PilotSelectedActor, "Pilot Selected Actor", "Move the selected actor around using the viewport controls, and bind the viewport to the actor's location and orientation.", EUserInterfaceActionType::Button, FInputGesture( EModifierKey::Control | EModifierKey::Shift, EKeys::P ) );
+	UI_COMMAND( EjectActorPilot, "Eject from Actor Pilot", "Stop piloting an actor with the current viewport. Unlocks the viewport's position and orientation from the actor the viewport is currently piloting.", EUserInterfaceActionType::Button, FInputGesture() );
+	UI_COMMAND( ToggleActorPilotCameraView, "Actor Pilot Camera View", "Toggles showing the exact camera view when using the viewport to pilot a camera", EUserInterfaceActionType::ToggleButton, FInputGesture( EModifierKey::Control | EModifierKey::Shift, EKeys::C ) );
 
 	UI_COMMAND( ViewportConfig_OnePane, "Layout One Pane", "Changes the viewport arrangement to one pane", EUserInterfaceActionType::ToggleButton, FInputGesture() );
 	UI_COMMAND( ViewportConfig_TwoPanesH, "Layout Two Panes (horizontal)", "Changes the viewport arrangement to two panes, side-by-side", EUserInterfaceActionType::ToggleButton, FInputGesture() );
diff --git a/Engine/Source/Editor/LevelEditor/Private/SActorPilotViewportToolbar.h b/Engine/Source/Editor/LevelEditor/Private/SActorPilotViewportToolbar.h
new file mode 100644
--- /dev/null
+++ b/Engine/Source/Editor/LevelEditor/Private/SActorPilotViewportToolbar.h
@@ -0,0 +1,97 @@
+
+#include "Editor/UnrealEd/Public/SViewportToolBar.h"
+
+#include "SLevelEditor.h"
+#include "Editor/UnrealEd/Public/SEditorViewportToolBarMenu.h"
+#include "Editor/UnrealEd/Public/SEditorViewportToolBarButton.h"
+
+class SLevelViewport;
+
+#define LOCTEXT_NAMESPACE "SActorPilotViewportToolbar"
+
+class SActorPilotViewportToolbar : public SViewportToolBar
+{
+public:
+	SLATE_BEGIN_ARGS(SActorPilotViewportToolbar){}
+		SLATE_ARGUMENT(TSharedPtr<SLevelViewport>, Viewport)
+	SLATE_END_ARGS()
+
+	FText GetActiveText() const
+	{
+		const AActor* Pilot = nullptr;
+		auto ViewportPtr = Viewport.Pin();
+		if (ViewportPtr.IsValid())
+		{
+			Pilot = ViewportPtr->GetLevelViewportClient().GetActiveActorLock().Get();
+		}
+		
+		return Pilot ? FText::Format(LOCTEXT("ActiveText", "[ Pilot Active - {0} ]"), FText::FromString(Pilot->GetActorLabel())) : FText();
+	}
+
+	void Construct(const FArguments& InArgs)
+	{
+		SViewportToolBar::Construct(SViewportToolBar::FArguments());
+
+		Viewport = InArgs._Viewport;
+
+		auto& ViewportCommands = FLevelViewportCommands::Get();
+
+		FToolBarBuilder ToolbarBuilder( InArgs._Viewport->GetCommandList(), FMultiBoxCustomization::None, nullptr/*InExtenders*/ );
+
+		// Use a custom style
+		FName ToolBarStyle = "ViewportMenu";
+		ToolbarBuilder.SetStyle(&FEditorStyle::Get(), ToolBarStyle);
+		ToolbarBuilder.SetStyle(&FEditorStyle::Get(), ToolBarStyle);
+		ToolbarBuilder.SetLabelVisibility(EVisibility::Collapsed);
+
+		ToolbarBuilder.BeginSection("ActorPilot");
+		ToolbarBuilder.BeginBlockGroup();
+		{
+			static FName EjectActorPilotName = FName(TEXT("EjectActorPilot"));
+			ToolbarBuilder.AddToolBarButton( ViewportCommands.EjectActorPilot, NAME_None, TAttribute<FText>(), TAttribute<FText>(), TAttribute<FSlateIcon>(), EjectActorPilotName );
+
+			static FName ToggleActorPilotCameraViewName = FName(TEXT("ToggleActorPilotCameraView"));
+			ToolbarBuilder.AddToolBarButton( ViewportCommands.ToggleActorPilotCameraView, NAME_None, TAttribute<FText>(), TAttribute<FText>(), TAttribute<FSlateIcon>(), ToggleActorPilotCameraViewName );
+		}
+		ToolbarBuilder.EndBlockGroup();
+		ToolbarBuilder.EndSection();
+
+		ToolbarBuilder.BeginSection("ActorPilot_Label");
+			ToolbarBuilder.AddWidget(
+				SNew(SBox)
+				// Nasty hack to make this VAlign_center properly. The parent Box is set to VAlign_Bottom, so we can't fill properly.
+				.HeightOverride(24.f)
+				[
+					SNew(SVerticalBox)
+
+					+ SVerticalBox::Slot()
+					.VAlign(VAlign_Center)
+					[
+						SNew(STextBlock)
+						.TextStyle(FEditorStyle::Get(), "LevelViewport.ActorPilotText")
+						.Text(this, &SActorPilotViewportToolbar::GetActiveText)
+					]
+				]
+			);
+		ToolbarBuilder.EndSection();
+
+		ChildSlot
+		[
+			SNew( SBorder )
+			.BorderImage( FEditorStyle::GetBrush("NoBorder") )
+			.Padding(FMargin(4.f, 0.f))
+			// Color and opacity is changed based on whether or not the mouse cursor is hovering over the toolbar area
+			.ColorAndOpacity( this, &SViewportToolBar::OnGetColorAndOpacity )
+			[
+				ToolbarBuilder.MakeWidget()
+			]
+		];
+	}
+
+private:
+
+	/** The viewport that we are in */
+	TWeakPtr<SLevelViewport> Viewport;
+};
+
+#undef LOCTEXT_NAMESPACE
\ No newline at end of file
diff --git a/Engine/Source/Editor/LevelEditor/Private/SLevelViewport.cpp b/Engine/Source/Editor/LevelEditor/Private/SLevelViewport.cpp
--- a/Engine/Source/Editor/LevelEditor/Private/SLevelViewport.cpp
+++ b/Engine/Source/Editor/LevelEditor/Private/SLevelViewport.cpp
@@ -40,6 +40,7 @@
 #include "Engine/Selection.h"
 #include "GameFramework/PlayerInput.h"
 #include "GameFramework/PlayerController.h"
+#include "SActorPilotViewportToolbar.h"
 
 static const FName LevelEditorName("LevelEditor");
 
@@ -1156,8 +1157,8 @@ void SLevelViewport::BindOptionCommands( FUICommandList& CommandList )
 		);
 
 	CommandList.MapAction(
-		ViewportActions.ToggleLockedCameraView,
-		FExecuteAction::CreateSP(this, &SLevelViewport::ToggleLockedCameraView),
+		ViewportActions.ToggleActorPilotCameraView,
+		FExecuteAction::CreateSP(this, &SLevelViewport::ToggleActorPilotCameraView),
 		FCanExecuteAction(),
 		FIsActionChecked::CreateSP( this, &SLevelViewport::IsLockedCameraViewEnabled )
 		);
@@ -1198,24 +1199,17 @@ void SLevelViewport::BindViewCommands( FUICommandList& CommandList )
 		);
 
 	CommandList.MapAction(
-		ViewportActions.ActorUnlock,
+		ViewportActions.EjectActorPilot,
 		FExecuteAction::CreateSP( this, &SLevelViewport::OnActorUnlock ),
 		FCanExecuteAction::CreateSP( this, &SLevelViewport::CanExecuteActorUnlock )
 		);
 
 	CommandList.MapAction(
-		ViewportActions.ActorLockSelected,
+		ViewportActions.PilotSelectedActor,
 		FExecuteAction::CreateSP( this, &SLevelViewport::OnActorLockSelected ),
 		FCanExecuteAction::CreateSP( this, &SLevelViewport::CanExecuteActorLockSelected )
 		);
 
-	CommandList.MapAction(
-		ViewportActions.ActorUnlockSelected,
-		FExecuteAction::CreateSP( this, &SLevelViewport::OnActorUnlock ),
-		FCanExecuteAction::CreateSP( this, &SLevelViewport::CanExecuteActorUnlock ),
-		FIsActionChecked::CreateSP( this, &SLevelViewport::IsSelectedActorLocked )
-		);
-
 	CommandList.MapAction(
 		ViewportActions.ViewportConfig_OnePane,
 		FExecuteAction::CreateSP(this, &SLevelViewport::OnSetViewportConfiguration, LevelViewportConfigurationNames::OnePane),
@@ -1574,15 +1568,9 @@ TSharedPtr<SWidget> SLevelViewport::MakeViewportToolbar()
 		.VAlign(VAlign_Top)
 		.HAlign(HAlign_Left)
 		[
-			SNew(SHorizontalBox)
-			.Visibility( EVisibility::SelfHitTestInvisible )
-			+SHorizontalBox::Slot()
-			[
-				SNew(SImage)
-				.Image(FEditorStyle::GetBrush("LevelViewport.ActorLockIcon"))
-				.Visibility(this, &SLevelViewport::GetLockedIconVisibility)
-				.ToolTipText(this, &SLevelViewport::GetLockedIconToolTip)
-			]
+			SNew(SActorPilotViewportToolbar)
+			.Viewport( SharedThis( this ) )
+			.Visibility(this, &SLevelViewport::GetLockedIconVisibility)
 		];
 }
 
@@ -2118,7 +2106,7 @@ bool SLevelViewport::IsAnyActorLocked() const
 	return LevelViewportClient->IsAnyActorLocked();
 }
 
-void SLevelViewport::ToggleLockedCameraView()
+void SLevelViewport::ToggleActorPilotCameraView()
 {
 	LevelViewportClient->bLockedCameraView = !LevelViewportClient->bLockedCameraView;
 }
diff --git a/Engine/Source/Editor/LevelEditor/Private/SLevelViewportToolBar.cpp b/Engine/Source/Editor/LevelEditor/Private/SLevelViewportToolBar.cpp
--- a/Engine/Source/Editor/LevelEditor/Private/SLevelViewportToolBar.cpp
+++ b/Engine/Source/Editor/LevelEditor/Private/SLevelViewportToolBar.cpp
@@ -391,142 +391,6 @@ static void OnGenerateClearBookmarkMenu( FMenuBuilder& MenuBuilder , TWeakPtr<cl
 	}
 }
 
-template<typename ActorClass>
-static void OnGenerateActorLockingMenuSection( TWeakPtr<SLevelViewport> Viewport, FMenuBuilder& MenuBuilder )
-{
-	check(ActorClass::StaticClass()->IsChildOf(AActor::StaticClass()));
-
-	// Creates a scene outliner to fill the menu. The outliner lists only actors of class T and its children
-
-	using namespace SceneOutliner;
-
-	SceneOutliner::FInitializationOptions InitOptions;
-	InitOptions.Mode = ESceneOutlinerMode::ActorPicker;			
-	InitOptions.bShowHeaderRow = false;
-
-	InitOptions.ColumnMap.Add(FBuiltInColumnTypes::Label(), FColumnInfo(EColumnVisibility::Visible, 0) );
-	InitOptions.ColumnMap.Add("LockedToViewport", FColumnInfo(EColumnVisibility::Visible, 10,
-		FCreateSceneOutlinerColumn::CreateSP( Viewport.Pin().ToSharedRef(), &SLevelViewport::CreateActorLockSceneOutlinerColumn )) );
-
-	/** Only show actors that are of the templated type */
-	InitOptions.Filters->AddFilterPredicate( FActorFilterPredicate::CreateStatic([](const AActor* InActor){
-		return InActor && InActor->IsA(ActorClass::StaticClass()) && !InActor->IsPendingKill();
-	}) );
-
-	// Create an outliner with the options from above. It sits in a box with a max height limit to stop it getting too tall when lots of actors are added.
-	FSceneOutlinerModule& SceneOutlinerModule = FModuleManager::LoadModuleChecked<FSceneOutlinerModule>("SceneOutliner");
-	TSharedRef< SWidget > MiniSceneOutliner =
-		SNew( SVerticalBox )
-		+SVerticalBox::Slot()
-		.MaxHeight(400.0f)
-		[
-			SceneOutlinerModule.CreateSceneOutliner(
-				InitOptions,
-				FOnActorPicked::CreateSP( Viewport.Pin().ToSharedRef(), &SLevelViewport::OnActorLockToggleFromMenu ) )
-		];
-
-	// Add the outliner to the menu
-	MenuBuilder.AddWidget(MiniSceneOutliner, FText(), true);
-}
-
-static void OnGenerateCameraActorLockingMenu( FMenuBuilder& MenuBuilder, TWeakPtr<SLevelViewport> Viewport )
-{
-	MenuBuilder.BeginSection("LevelViewportCameraActors", LOCTEXT("ActorLockingMenu_CameraActorsHeader", "Camera Actors"));
-	{
-		OnGenerateActorLockingMenuSection<ACameraActor>(Viewport, MenuBuilder);
-	}
-	MenuBuilder.EndSection();
-}
-
-static void OnGenerateLightActorLockingMenu( FMenuBuilder& MenuBuilder, TWeakPtr<SLevelViewport> Viewport )
-{
-	MenuBuilder.BeginSection("LevelViewportLightActors", LOCTEXT("ActorLockingMenu_LightActorsHeader", "Light Actors"));
-	{
-		OnGenerateActorLockingMenuSection<ALight>(Viewport, MenuBuilder);
-	}
-	MenuBuilder.EndSection();
-}
-
-static void OnGenerateActorLockingMenu( FMenuBuilder& MenuBuilder, TWeakPtr<SLevelViewport> Viewport )
-{
-	const FLevelViewportCommands& Actions = FLevelViewportCommands::Get();
-	FLevelEditorViewportClient& ViewportClient = Viewport.Pin()->GetLevelViewportClient();
-
-	MenuBuilder.AddMenuEntry(Actions.ToggleLockedCameraView);
-
-	bool IsLocked = false;
-	if (ViewportClient.GetActiveActorLock().IsValid())
-	{
-		// Viewport is locked - show the unlock item
-		AActor* Actor = ViewportClient.GetActiveActorLock().Get();
-
-		if (!Actor->IsPendingKill())
-		{
-			MenuBuilder.BeginSection("LevelViewportLocked", LOCTEXT("LockingMenuLocked", "Locked"));
-			{
-				MenuBuilder.AddMenuEntry(
-					Actions.ActorUnlock,
-					NAME_None,
-					FText::Format( LOCTEXT("UnlockMenuItem", "Unlock from {0}"), FText::FromString( Actor->GetActorLabel() ) ) );
-			}
-			MenuBuilder.EndSection();
-
-			IsLocked = true;
-		}
-	}
-
-	if (!IsLocked)
-	{
-		MenuBuilder.BeginSection("LevelViewportNotLocked", LOCTEXT("LockingMenuNotLocked", "Not Locked"));
-		MenuBuilder.EndSection();
-	}
-
-	// If a single actor is selected, show an item to lock the viewport to it
-	USelection* ActorSelection = GEditor->GetSelectedActors();
-	if (1 == ActorSelection->Num() && ActorSelection->GetSelectedObject(0) != NULL)
-	{
-		MenuBuilder.BeginSection("LevelViewportSelectedActor", LOCTEXT("LockingMenuSelectionHeader", "Selected Actor"));
-		{
-
-		    AActor* Actor = CastChecked<AActor>(ActorSelection->GetSelectedObject(0));
-    
-		    if (Viewport.Pin()->IsSelectedActorLocked())
-		    {
-			    MenuBuilder.AddMenuEntry(
-				    Actions.ActorUnlockSelected,
-				    NAME_None,
-				    FText::FromString( Actor->GetActorLabel() ) );
-			}
-			else
-			{
-				MenuBuilder.AddMenuEntry(
-					Actions.ActorLockSelected,
-					NAME_None,
-					FText::FromString( Actor->GetActorLabel() ) );
-			}
-		}
-		MenuBuilder.EndSection();
-	}
-
-	MenuBuilder.BeginSection("LevelViewportCamerasLights");
-	{
-		// Add sub-menu listing camera actors that can be locked
-		MenuBuilder.AddSubMenu( 
-			LOCTEXT("ActorLockingCamerasSubMenu", "Cameras"), 
-			FText::GetEmpty(), 
-			FNewMenuDelegate::CreateStatic( &OnGenerateCameraActorLockingMenu, Viewport )
-			);
-
-		// Add sub-menu listing camera actors that can be locked
-		MenuBuilder.AddSubMenu( 
-			LOCTEXT("ActorLockingLightsSubMenu", "Lights"), 
-			FText::GetEmpty(), 
-			FNewMenuDelegate::CreateStatic( &OnGenerateLightActorLockingMenu, Viewport )
-			);
-	}
-	MenuBuilder.EndSection();
-}
-
 /**
  * Called to generate the bookmark submenu
  */
@@ -598,7 +462,6 @@ TSharedRef<SWidget> SLevelViewportToolBar::GenerateOptionsMenu() const
 	TSharedPtr<FExtender> MenuExtender = FExtender::Combine(Extenders);
 
 	const bool bIsPerspective = Viewport.Pin()->GetLevelViewportClient().IsPerspective();
-	const bool bIsLocked = Viewport.Pin()->GetLevelViewportClient().IsAnyActorLocked();
 	const bool bInShouldCloseWindowAfterMenuSelection = true;
 	FMenuBuilder OptionsMenuBuilder( bInShouldCloseWindowAfterMenuSelection, Viewport.Pin()->GetCommandList(), MenuExtender );
 	{
@@ -640,19 +503,6 @@ TSharedRef<SWidget> SLevelViewportToolBar::GenerateOptionsMenu() const
 		}
 		OptionsMenuBuilder.EndSection();
 
-		if( (bIsPerspective || bIsLocked) && !Viewport.Pin()->GetLevelViewportClient().IsLockedToMatinee())
-		{
-			OptionsMenuBuilder.BeginSection("LevelViewportActorLocking");
-			{
-				OptionsMenuBuilder.AddSubMenu( 
-					LOCTEXT("ActorLockingSubMenu", "Lock Viewport to Actor"), 
-					LOCTEXT("ActorLockingSubMenu_ToolTip", "Lock Viewport position and orientation to Cameras, Lights or other scene actors"), 
-					FNewMenuDelegate::CreateStatic( &OnGenerateActorLockingMenu, Viewport )
-					);
-			}
-			OptionsMenuBuilder.EndSection();
-		}
-
 		if( bIsPerspective )
 		{
 			// Bookmarks only work in perspective viewports so only show the menu option if this toolbar is in one
diff --git a/Engine/Source/Editor/LevelEditor/Public/LevelViewportActions.h b/Engine/Source/Editor/LevelEditor/Public/LevelViewportActions.h
--- a/Engine/Source/Editor/LevelEditor/Public/LevelViewportActions.h
+++ b/Engine/Source/Editor/LevelEditor/Public/LevelViewportActions.h
@@ -68,9 +68,6 @@ public:
 	/** Toggles immersive mode in the viewport */
 	TSharedPtr< FUICommandInfo > ToggleImmersive;
 
-	/** Toggles showing the exact camera view when locking a viewport to a camera */
-	TSharedPtr< FUICommandInfo > ToggleLockedCameraView;
-
 	/** Toggles maximize mode in the viewport */
 	TSharedPtr< FUICommandInfo > ToggleMaximize;
 
@@ -144,10 +141,12 @@ public:
 	TArray< TSharedPtr< FUICommandInfo > > ClearBookmarkCommands;
 	TSharedPtr< FUICommandInfo > ClearAllBookMarks;
 
-	/** Actor locking commands */
-	TSharedPtr< FUICommandInfo > ActorUnlock;
-	TSharedPtr< FUICommandInfo > ActorLockSelected;
-	TSharedPtr< FUICommandInfo > ActorUnlockSelected;
+	/** Actor pilot commands */
+	TSharedPtr< FUICommandInfo > EjectActorPilot;
+	TSharedPtr< FUICommandInfo > PilotSelectedActor;
+
+	/** Toggles showing the exact camera view when locking a viewport to a camera */
+	TSharedPtr< FUICommandInfo > ToggleActorPilotCameraView;
 
 	/** Viewport pane configurations */
 	TSharedPtr< FUICommandInfo > ViewportConfig_OnePane;
diff --git a/Engine/Source/Editor/LevelEditor/Public/SLevelViewport.h b/Engine/Source/Editor/LevelEditor/Public/SLevelViewport.h
--- a/Engine/Source/Editor/LevelEditor/Public/SLevelViewport.h
+++ b/Engine/Source/Editor/LevelEditor/Public/SLevelViewport.h
@@ -195,7 +195,7 @@ public:
 	/**
 	 * Toggles enabling the exact camera view when locking a viewport to a camera
 	 */
-	void ToggleLockedCameraView();
+	void ToggleActorPilotCameraView();
 
 	/**
 	 * Check whether locked camera view is enabled
diff --git a/Engine/Source/Editor/UnrealEd/Private/LevelEditorViewport.cpp b/Engine/Source/Editor/UnrealEd/Private/LevelEditorViewport.cpp
--- a/Engine/Source/Editor/UnrealEd/Private/LevelEditorViewport.cpp
+++ b/Engine/Source/Editor/UnrealEd/Private/LevelEditorViewport.cpp
@@ -1635,6 +1635,17 @@ ELevelViewportType FLevelEditorViewportClient::GetViewportType() const
 	}
 }
 
+void FLevelEditorViewportClient::SetViewportType( ELevelViewportType InViewportType )
+{
+	if (InViewportType != LVT_Perspective)
+	{
+		SetActorLock(nullptr);
+		UpdateViewForLockedActor();
+	}
+
+	FEditorViewportClient::SetViewportType(InViewportType);
+}
+
 void FLevelEditorViewportClient::OverridePostProcessSettings( FSceneView& View )
 {
 	const UCameraComponent* CameraComponent = GetCameraComponentForView();
diff --git a/Engine/Source/Editor/UnrealEd/Public/LevelEditorViewport.h b/Engine/Source/Editor/UnrealEd/Public/LevelEditorViewport.h
--- a/Engine/Source/Editor/UnrealEd/Public/LevelEditorViewport.h
+++ b/Engine/Source/Editor/UnrealEd/Public/LevelEditorViewport.h
@@ -506,6 +506,7 @@ protected:
 	/** FEditorViewportClient Interface*/
 	virtual void UpdateLinkedOrthoViewports( bool bInvalidate = false ) override;
 	virtual ELevelViewportType GetViewportType() const override;
+	virtual void SetViewportType( ELevelViewportType InViewportType ) override;
 	virtual void OverridePostProcessSettings( FSceneView& View ) override;
 	virtual void PerspectiveCameraMoved() override;
 	virtual bool ShouldLockPitch() const override;
