[INTEGRATE] Change 2436670 by [REDACTED_PERSON][REDACTED_HANDLE] on 2015/02/07 02:42:19

	Editor: Add/clarify make Blueprint actions to the Blueprints toolbar button
	- Expose 'Blueprint from selected actor' (single actor; identical to invoking it from the details panel)
	- Expose 'Blueprint from selected components' (harvests components from mutliple actors, ignoring the actor classes themselves)
	- Rename 'New Blueprint Class' to 'New Empty Blueprint Class'
	- Move the make blueprint options to the top of the menu and push level script to be the second category

diff --git a/Engine/Source/Editor/LevelEditor/LevelEditor.Build.cs b/Engine/Source/Editor/LevelEditor/LevelEditor.Build.cs
--- a/Engine/Source/Editor/LevelEditor/LevelEditor.Build.cs
+++ b/Engine/Source/Editor/LevelEditor/LevelEditor.Build.cs
@@ -56,6 +56,7 @@ public class LevelEditor : ModuleRules
 				"WebBrowser",
                 "Persona",
                 "Kismet",
+				"KismetWidgets"
 			}
 		);
 
diff --git a/Engine/Source/Editor/LevelEditor/Private/LevelEditorActions.cpp b/Engine/Source/Editor/LevelEditor/Private/LevelEditorActions.cpp
--- a/Engine/Source/Editor/LevelEditor/Private/LevelEditorActions.cpp
+++ b/Engine/Source/Editor/LevelEditor/Private/LevelEditorActions.cpp
@@ -62,6 +62,7 @@
 #include "Engine/Light.h"
 #include "Animation/SkeletalMeshActor.h"
 #include "Editor/Persona/Public/AnimationRecorder.h"
+#include "Editor/KismetWidgets/Public/CreateBlueprintFromActorDialog.h"
 
 DEFINE_LOG_CATEGORY_STATIC(LevelEditorActions, Log, All);
 
@@ -1982,7 +1983,7 @@ void FLevelEditorActionCallbacks::OpenLevelBlueprint( TWeakPtr< SLevelEditor > L
 	}
 }
 
-void FLevelEditorActionCallbacks::CreateBlueprintClass()
+void FLevelEditorActionCallbacks::CreateBlankBlueprintClass()
 {
 	// Use the BlueprintFactory to allow the user to pick a parent class for the new Blueprint class
 	UBlueprintFactory* NewFactory = Cast<UBlueprintFactory>(NewObject<UFactory>(GetTransientPackage(), UBlueprintFactory::StaticClass()));
@@ -1992,7 +1993,7 @@ void FLevelEditorActionCallbacks::CreateBlueprintClass()
 		UClass* SelectedClass = NewFactory->ParentClass;
 
 		// Now help the user pick a path and name for the new Blueprint
-		UBlueprint* Blueprint = FKismetEditorUtilities::CreateBlueprintFromClass(NSLOCTEXT("LevelEditorCommands", "CreateBlueprintClass_Title", "Create Blueprint Class"), SelectedClass);
+		UBlueprint* Blueprint = FKismetEditorUtilities::CreateBlueprintFromClass(NSLOCTEXT("LevelEditorCommands", "CreateBlankBlueprintClass_Title", "Create Blank Blueprint Class"), SelectedClass);
 
 		if( Blueprint )
 		{
@@ -2005,6 +2006,28 @@ void FLevelEditorActionCallbacks::CreateBlueprintClass()
 	}
 }
 
+bool FLevelEditorActionCallbacks::CanHarvestSelectedActorsIntoBlueprintClass()
+{
+	return GEditor->GetSelectedActorCount() > 0;
+}
+
+void FLevelEditorActionCallbacks::HarvestSelectedActorsIntoBlueprintClass()
+{
+	const bool bHarvest = true;
+	FCreateBlueprintFromActorDialog::OpenDialog(bHarvest);
+}
+
+bool FLevelEditorActionCallbacks::CanSubclassSelectedActorIntoBlueprintClass()
+{
+	return GEditor->GetSelectedActorCount() == 1;
+}
+
+void FLevelEditorActionCallbacks::SubclassSelectedActorIntoBlueprintClass()
+{
+	const bool bHarvest = false;
+	FCreateBlueprintFromActorDialog::OpenDialog(bHarvest);
+}
+
 void FLevelEditorActionCallbacks::CheckOutProjectSettingsConfig( )
 {
 	FString ConfigFilePath = FPaths::ConvertRelativePathToFull(FString::Printf(TEXT("%sDefaultEngine.ini"), *FPaths::SourceConfigDir()));
@@ -2902,7 +2925,9 @@ void FLevelEditorCommands::RegisterCommands()
 
 	UI_COMMAND( OpenLevelBlueprint, "Open Level Blueprint", "Edit the Level Blueprint for the current level", EUserInterfaceActionType::Button, FInputGesture() );
 	UI_COMMAND( CheckOutProjectSettingsConfig, "Check Out", "Checks out the project settings config file so the game mode can be set.", EUserInterfaceActionType::Button, FInputGesture() );
-	UI_COMMAND( CreateBlueprintClass, "New Blueprint Class...", "Create a new Blueprint Class", EUserInterfaceActionType::Button, FInputGesture());
+	UI_COMMAND( CreateBlankBlueprintClass, "New Empty Blueprint Class...", "Create a new Blueprint Class", EUserInterfaceActionType::Button, FInputGesture() );
+	UI_COMMAND( ConvertSelectionToBlueprintViaHarvest, "Convert Selected Components to Blueprint Class...", "Replace all of the selected actors with a new Blueprint Class based on Actor that contains the components", EUserInterfaceActionType::Button, FInputGesture() );
+	UI_COMMAND( ConvertSelectionToBlueprintViaSubclass, "Convert Selected Actor to Blueprint Class...", "Replace the selected actor with a new Blueprint subclass based on the class of the selected Actor", EUserInterfaceActionType::Button, FInputGesture() );
 
 	UI_COMMAND( ShowTransformWidget, "Show Transform Widget", "Toggles the visibility of the transform widgets", EUserInterfaceActionType::ToggleButton, FInputGesture() );
 	UI_COMMAND( AllowTranslucentSelection, "Allow Translucent Selection", "Allows translucent objects to be selected", EUserInterfaceActionType::ToggleButton, FInputGesture(EKeys::T) );
diff --git a/Engine/Source/Editor/LevelEditor/Private/LevelEditorToolBar.cpp b/Engine/Source/Editor/LevelEditor/Private/LevelEditorToolBar.cpp
--- a/Engine/Source/Editor/LevelEditor/Private/LevelEditorToolBar.cpp
+++ b/Engine/Source/Editor/LevelEditor/Private/LevelEditorToolBar.cpp
@@ -1897,6 +1897,26 @@ TSharedRef< SWidget > FLevelEditorToolBar::GenerateOpenBlueprintMenuContent( TSh
 	const bool bShouldCloseWindowAfterMenuSelection = true;
 	FMenuBuilder MenuBuilder( bShouldCloseWindowAfterMenuSelection, InCommandList );
 
+	MenuBuilder.BeginSection(NAME_None, LOCTEXT("BlueprintClass", "Blueprint Class"));
+	{
+		// Create a blank BP
+		MenuBuilder.AddMenuEntry(FLevelEditorCommands::Get().CreateBlankBlueprintClass);
+
+		// Convert selection to BP
+		MenuBuilder.AddMenuEntry(FLevelEditorCommands::Get().ConvertSelectionToBlueprintViaHarvest);
+		MenuBuilder.AddMenuEntry(FLevelEditorCommands::Get().ConvertSelectionToBlueprintViaSubclass);
+
+		// Open an existing Blueprint Class...
+		FSlateIcon OpenBPIcon(FEditorStyle::GetStyleSetName(), "LevelEditor.OpenClassBlueprint");
+		MenuBuilder.AddSubMenu(
+			LOCTEXT("OpenBlueprintClassSubMenu", "Open Blueprint Class..."),
+			LOCTEXT("OpenBlueprintClassSubMenu_ToolTip", "Open an existing Blueprint Class in this project"),
+			FNewMenuDelegate::CreateStatic(&FBlueprintMenus::MakeOpenBPClassMenu),
+			false,
+			OpenBPIcon);
+	}
+	MenuBuilder.EndSection();
+
 	MenuBuilder.BeginSection(NAME_None, LOCTEXT("LevelScriptBlueprints", "Level Blueprints"));
 	{
 		MenuBuilder.AddMenuEntry( FLevelEditorCommands::Get().OpenLevelBlueprint );
@@ -1934,22 +1954,6 @@ TSharedRef< SWidget > FLevelEditorToolBar::GenerateOpenBlueprintMenuContent( TSh
 	}
 	MenuBuilder.EndSection();
 
-	MenuBuilder.BeginSection(NAME_None, LOCTEXT("Blueprints Class", "Blueprints Class"));
-	{
-		// New Blueprint Class...
-		MenuBuilder.AddMenuEntry(FLevelEditorCommands::Get().CreateBlueprintClass, NAME_None, LOCTEXT("NewBlueprintClass", "New Blueprint Class..."));
-
-		// Open Blueprint Class...
-		FSlateIcon OpenBPIcon(FEditorStyle::GetStyleSetName(), "LevelEditor.OpenClassBlueprint");
-		MenuBuilder.AddSubMenu(
-			LOCTEXT("OpenBlueprintClassSubMenu", "Open Blueprint Class..."),
-			LOCTEXT("OpenBlueprintClassSubMenu_ToolTip", "Open an existing Blueprint Class in this project"),
-			FNewMenuDelegate::CreateStatic(&FBlueprintMenus::MakeOpenBPClassMenu), 
-			false, 
-			OpenBPIcon );
-	}
-	MenuBuilder.EndSection();
-
 #undef LOCTEXT_NAMESPACE
 
 	return MenuBuilder.MakeWidget();
diff --git a/Engine/Source/Editor/LevelEditor/Private/SLevelEditor.cpp b/Engine/Source/Editor/LevelEditor/Private/SLevelEditor.cpp
--- a/Engine/Source/Editor/LevelEditor/Private/SLevelEditor.cpp
+++ b/Engine/Source/Editor/LevelEditor/Private/SLevelEditor.cpp
@@ -94,8 +94,18 @@ void SLevelEditor::BindCommands()
 		FExecuteAction::CreateStatic< TWeakPtr< SLevelEditor > >( &FLevelEditorActionCallbacks::OpenLevelBlueprint, SharedThis( this ) ) );
 	
 	LevelEditorCommands->MapAction(
-		Actions.CreateBlueprintClass,
-		FExecuteAction::CreateStatic( &FLevelEditorActionCallbacks::CreateBlueprintClass ) );
+		Actions.CreateBlankBlueprintClass,
+		FExecuteAction::CreateStatic( &FLevelEditorActionCallbacks::CreateBlankBlueprintClass ) );
+
+	LevelEditorCommands->MapAction(
+		Actions.ConvertSelectionToBlueprintViaHarvest,
+		FExecuteAction::CreateStatic( &FLevelEditorActionCallbacks::HarvestSelectedActorsIntoBlueprintClass ),
+		FCanExecuteAction::CreateStatic( &FLevelEditorActionCallbacks::CanHarvestSelectedActorsIntoBlueprintClass ) );
+
+	LevelEditorCommands->MapAction(
+		Actions.ConvertSelectionToBlueprintViaSubclass,
+		FExecuteAction::CreateStatic( &FLevelEditorActionCallbacks::SubclassSelectedActorIntoBlueprintClass ),
+		FCanExecuteAction::CreateStatic( &FLevelEditorActionCallbacks::CanSubclassSelectedActorIntoBlueprintClass ) );
 
 	LevelEditorCommands->MapAction(
 		Actions.OpenContentBrowser,
diff --git a/Engine/Source/Editor/LevelEditor/Public/LevelEditorActions.h b/Engine/Source/Editor/LevelEditor/Public/LevelEditorActions.h
--- a/Engine/Source/Editor/LevelEditor/Public/LevelEditorActions.h
+++ b/Engine/Source/Editor/LevelEditor/Public/LevelEditorActions.h
@@ -503,7 +503,9 @@ public:
 	 */
 	TSharedPtr< FUICommandInfo > OpenLevelBlueprint;
 	TSharedPtr< FUICommandInfo > CheckOutProjectSettingsConfig;
-	TSharedPtr< FUICommandInfo > CreateBlueprintClass;
+	TSharedPtr< FUICommandInfo > CreateBlankBlueprintClass;
+	TSharedPtr< FUICommandInfo > ConvertSelectionToBlueprintViaHarvest;
+	TSharedPtr< FUICommandInfo > ConvertSelectionToBlueprintViaSubclass;
 
 	/** Editor mode commands */
 	TArray< TSharedPtr< FUICommandInfo > > EditorModeCommands;
@@ -1068,7 +1070,19 @@ public:
 	static bool CanSelectGameModeBlueprint();
 
 	/** Helps the user create a Blueprint class */
-	static void CreateBlueprintClass();
+	static void CreateBlankBlueprintClass();
+
+	/** Can call HarvestSelectedActorsIntoBlueprintClass right now?  (is anything selected) */
+	static bool CanHarvestSelectedActorsIntoBlueprintClass();
+
+	/** Harvest all of the components in the selected actors into a Blueprint, and replace the instances with one new actor */
+	static void HarvestSelectedActorsIntoBlueprintClass();
+
+	/** Can call SubclassSelectedActorIntoBlueprintClass right now?  (is exactly one thing selected) */
+	static bool CanSubclassSelectedActorIntoBlueprintClass();
+
+	/** Convert the selected actor into a Blueprint (via subclassing) */
+	static void SubclassSelectedActorIntoBlueprintClass();
 
 	/** Shows only selected actors, hiding any unselected actors and unhiding any selected hidden actors. */
 	static void OnShowOnlySelectedActors();
