[INTEGRATE] Change 2438669 by [REDACTED_PERSON][REDACTED_HANDLE] on 2015/02/09 15:46:09

	[UE-9456] - The cached actor selection is now updated appropriately on undo/redo

diff --git a/Engine/Source/Editor/LevelEditor/Private/SActorDetails.cpp b/Engine/Source/Editor/LevelEditor/Private/SActorDetails.cpp
--- a/Engine/Source/Editor/LevelEditor/Private/SActorDetails.cpp
+++ b/Engine/Source/Editor/LevelEditor/Private/SActorDetails.cpp
@@ -177,6 +177,17 @@ void SActorDetails::PostUndo(bool bSuccess)
 	// Enable the selection guard to prevent OnTreeSelectionChanged() from altering the editor's component selection
 	TGuardValue<bool> SelectionGuard(bSelectionGuard, true);
 
+	if (!DetailsView->IsLocked())
+	{
+		// Make sure the locked actor selection matches the editor selection
+		AActor* SelectedActor = GetSelectedActorInEditor();
+		if (SelectedActor && SelectedActor != LockedActorSelection.Get())
+		{
+			LockedActorSelection = SelectedActor;
+		}
+	}
+	
+
 	// Refresh the tree and update the selection to match the world
 	SCSEditor->UpdateTree();
 	UpdateComponentTreeFromEditorSelection();
@@ -230,8 +241,8 @@ void SActorDetails::OnEditorSelectionChanged(UObject* Object)
 
 AActor* SActorDetails::GetSelectedActorInEditor() const
 {
-	//@todo this won't work w/ multi-select
-	return Cast<AActor>(*GEditor->GetSelectedActorIterator());
+	//@todo this doesn't work w/ multi-select
+	return GEditor->GetSelectedActors()->GetTop<AActor>();
 }
 
 AActor* SActorDetails::GetActorContext() const
