Fix non-unity compile

diff --git a/Engine/Source/Editor/LevelEditor/Private/LevelEditorContextMenu.cpp b/Engine/Source/Editor/LevelEditor/Private/LevelEditorContextMenu.cpp
--- a/Engine/Source/Editor/LevelEditor/Private/LevelEditorContextMenu.cpp
+++ b/Engine/Source/Editor/LevelEditor/Private/LevelEditorContextMenu.cpp
@@ -32,6 +32,7 @@
 #include "Engine/Selection.h"
 #include "Kismet2/BlueprintEditorUtils.h"
 #include "ComponentEditorUtils.h"
+#include "LevelViewportActions.h"
 
 #define LOCTEXT_NAMESPACE "LevelViewportContextMenu"
 
