[INTEGRATE] Change 2447921 by [REDACTED_PERSON][REDACTED_HANDLE] on 2015/02/16 18:06:45

	Editor: Locked out editing in the level editor version of the components tree when in PIE [UE-9374]

diff --git a/Engine/Source/Editor/LevelEditor/Private/SActorDetails.cpp b/Engine/Source/Editor/LevelEditor/Private/SActorDetails.cpp
--- a/Engine/Source/Editor/LevelEditor/Private/SActorDetails.cpp
+++ b/Engine/Source/Editor/LevelEditor/Private/SActorDetails.cpp
@@ -57,6 +57,7 @@ void SActorDetails::Construct(const FArguments& InArgs, const FName TabIdentifie
 
 	SCSEditor = SNew(SSCSEditor)
 		.EditorMode(EComponentEditorMode::ActorInstance)
+		.AllowEditing(this, &SActorDetails::GetAllowComponentTreeEditing)
 		.ActorContext(this, &SActorDetails::GetActorContext)
 		.OnSelectionUpdated(this, &SActorDetails::OnSCSEditorTreeViewSelectionChanged)
 		.ActorMenuExtender(InArgs._ActorMenuExtender);
@@ -247,6 +248,11 @@ AActor* SActorDetails::GetSelectedActorInEditor() const
 	return GEditor->GetSelectedActors()->GetTop<AActor>();
 }
 
+bool SActorDetails::GetAllowComponentTreeEditing() const
+{
+	return GEditor->PlayWorld == nullptr;
+}
+
 AActor* SActorDetails::GetActorContext() const
 {
 	AActor* SelectedActorInEditor = GetSelectedActorInEditor();
diff --git a/Engine/Source/Editor/LevelEditor/Private/SActorDetails.h b/Engine/Source/Editor/LevelEditor/Private/SActorDetails.h
--- a/Engine/Source/Editor/LevelEditor/Private/SActorDetails.h
+++ b/Engine/Source/Editor/LevelEditor/Private/SActorDetails.h
@@ -29,6 +29,7 @@ public:
 private:
 	AActor* GetSelectedActorInEditor() const;
 	AActor* GetActorContext() const;
+	bool GetAllowComponentTreeEditing() const;
 
 	void OnComponentsEditedInWorld();
 	void OnEditorSelectionChanged(UObject* Object);
