[INTEGRATE] Change 2462411 by [REDACTED_PERSON][REDACTED_HANDLE] on 2015/02/26 12:58:06

	Remove menu button for API reference if the CHM file doesn't exist.

diff --git a/Engine/Source/Editor/Documentation/Private/Documentation.cpp b/Engine/Source/Editor/Documentation/Private/Documentation.cpp
--- a/Engine/Source/Editor/Documentation/Private/Documentation.cpp
+++ b/Engine/Source/Editor/Documentation/Private/Documentation.cpp
@@ -53,6 +53,11 @@ bool FDocumentation::OpenAPIHome() const
 	}
 }
 
+bool FDocumentation::CanOpenAPIHome() const
+{
+	return FPaths::FileExists(FPaths::Combine(*FPaths::EngineDir(), TEXT("Documentation/CHM/API.chm")));
+}
+
 bool FDocumentation::Open(const FString& Link, FDocumentationSourceInfo Source) const
 {
 	FString DocumentationUrl;
diff --git a/Engine/Source/Editor/Documentation/Private/Documentation.h b/Engine/Source/Editor/Documentation/Private/Documentation.h
--- a/Engine/Source/Editor/Documentation/Private/Documentation.h
+++ b/Engine/Source/Editor/Documentation/Private/Documentation.h
@@ -19,6 +19,8 @@ public:
 
 	virtual bool OpenAPIHome() const override;
 
+	virtual bool CanOpenAPIHome() const override;
+
 	virtual bool Open( const FString& Link, FDocumentationSourceInfo Source = FDocumentationSourceInfo() ) const override;
 
 	virtual bool Open( const FString& Link, const FCultureRef& Culture, FDocumentationSourceInfo Source = FDocumentationSourceInfo() ) const override;
diff --git a/Engine/Source/Editor/Documentation/Public/IDocumentation.h b/Engine/Source/Editor/Documentation/Public/IDocumentation.h
--- a/Engine/Source/Editor/Documentation/Public/IDocumentation.h
+++ b/Engine/Source/Editor/Documentation/Public/IDocumentation.h
@@ -178,6 +178,8 @@ public:
 
 	virtual bool OpenAPIHome() const = 0;
 
+	virtual bool CanOpenAPIHome() const = 0;
+
 	virtual bool Open(const FString& Link, FDocumentationSourceInfo Source = FDocumentationSourceInfo()) const = 0;
 
 	virtual bool Open(const FString& Link, const FCultureRef& Culture, FDocumentationSourceInfo Source = FDocumentationSourceInfo()) const = 0;
diff --git a/Engine/Source/Editor/LevelEditor/Private/LevelEditorMenu.cpp b/Engine/Source/Editor/LevelEditor/Private/LevelEditorMenu.cpp
--- a/Engine/Source/Editor/LevelEditor/Private/LevelEditorMenu.cpp
+++ b/Engine/Source/Editor/LevelEditor/Private/LevelEditorMenu.cpp
@@ -322,7 +322,11 @@ TSharedRef< SWidget > FLevelEditorMenu::MakeLevelEditorMenu( const TSharedPtr<FU
 			MenuBuilder.BeginSection("HelpBrowse", NSLOCTEXT("MainHelpMenu", "Browse", "Browse"));
 			{
 				MenuBuilder.AddMenuEntry( FLevelEditorCommands::Get().BrowseDocumentation );
-				MenuBuilder.AddMenuEntry( FLevelEditorCommands::Get().BrowseAPIReference );
+
+				if(IDocumentation::Get()->CanOpenAPIHome())
+				{
+					MenuBuilder.AddMenuEntry( FLevelEditorCommands::Get().BrowseAPIReference );
+				}
 
 				MenuBuilder.AddMenuSeparator();
 
