non unity fix

diff --git a/Engine/Source/Editor/LevelEditor/Private/LevelEditorMenu.cpp b/Engine/Source/Editor/LevelEditor/Private/LevelEditorMenu.cpp
--- a/Engine/Source/Editor/LevelEditor/Private/LevelEditorMenu.cpp
+++ b/Engine/Source/Editor/LevelEditor/Private/LevelEditorMenu.cpp
@@ -10,6 +10,7 @@
 #include "GlobalEditorCommonCommands.h"
 #include "Editor/ContentBrowser/Public/ContentBrowserModule.h"
 #include "GenericCommands.h"
+#include "IDocumentation.h"
 
 #define LOCTEXT_NAMESPACE "LevelEditorMenu"
 
