PR #874: Copy current camera FOV when using "Create Camera Here" command (contributed by [REDACTED_PERSON])

diff --git a/Engine/Source/Editor/LevelEditor/Private/SLevelViewport.cpp b/Engine/Source/Editor/LevelEditor/Private/SLevelViewport.cpp
--- a/Engine/Source/Editor/LevelEditor/Private/SLevelViewport.cpp
+++ b/Engine/Source/Editor/LevelEditor/Private/SLevelViewport.cpp
@@ -1651,6 +1651,7 @@ void SLevelViewport::OnCreateCameraActor()
 	ACameraActor* pNewCamera = ViewportClient->GetWorld()->SpawnActor<ACameraActor>();
 	pNewCamera->SetActorLocation( ViewportClient->GetViewLocation(), false );
 	pNewCamera->SetActorRotation( ViewportClient->GetViewRotation() );
+	pNewCamera->CameraComponent->FieldOfView = ViewportClient->ViewFOV;
 
 	// Deselect any currently selected actors
 	GUnrealEd->SelectNone( true, true );
