UE-11204 Crash when returning to viewport from immersive mode with an editor window tabbed beneath

diff --git a/Engine/Source/Editor/LevelEditor/Private/LevelViewportLayout.cpp b/Engine/Source/Editor/LevelEditor/Private/LevelViewportLayout.cpp
--- a/Engine/Source/Editor/LevelEditor/Private/LevelViewportLayout.cpp
+++ b/Engine/Source/Editor/LevelEditor/Private/LevelViewportLayout.cpp
@@ -334,13 +334,19 @@ void FLevelViewportLayout::MaximizeViewport( TSharedRef<SLevelViewport> Viewport
 		}
 		bIsQueryingLayoutMetrics = false;
 
-		
-		// Keep track of the window we're contained in
-		// @todo immersive: Caching this after the transition is risky -- the widget could be moved to a new window!
-		//		We really need a safe way to query a widget's window that doesn't require a full layout pass.  Then,
-		//	    instead of caching the window we can look it up whenever it's needed
-		CachedOwnerWindow = OwnerWindow;
-
+		// If the widget can't be found in the layout pass, attempt to use the cached owner window
+		if(!OwnerWindow.IsValid() && CachedOwnerWindow.IsValid())
+		{
+			OwnerWindow = CachedOwnerWindow.Pin();
+		}
+		else
+		{
+			// Keep track of the window we're contained in
+			// @todo immersive: Caching this after the transition is risky -- the widget could be moved to a new window!
+			//		We really need a safe way to query a widget's window that doesn't require a full layout pass.  Then,
+			//	    instead of caching the window we can look it up whenever it's needed
+			CachedOwnerWindow = OwnerWindow;
+		}
 
 		if( !bIsImmersive && bWantImmersive )
 		{
@@ -449,7 +455,7 @@ void FLevelViewportLayout::MaximizeViewport( TSharedRef<SLevelViewport> Viewport
 		// a viewport based on saved layout, while that viewport is hosted in a background tab.  For this case, we'll
 		// never animate (checked here), so we don't need to store "before" metrics.
 		check( OwnerWindow.IsValid() || !bAllowAnimation );
-		if( OwnerWindow.IsValid() )	
+		if( OwnerWindow.IsValid() && ViewportWidgetPath.IsValid() )
 		{
 			// Setup transition metrics
 			if( bIsImmersive || bWasImmersive )
