[UE-9294] - Native components that  are not visible uproperties are now read-only for in-world component editing

diff --git a/Engine/Source/Editor/Kismet/Private/SSCSEditor.cpp b/Engine/Source/Editor/Kismet/Private/SSCSEditor.cpp
--- a/Engine/Source/Editor/Kismet/Private/SSCSEditor.cpp
+++ b/Engine/Source/Editor/Kismet/Private/SSCSEditor.cpp
@@ -841,6 +841,17 @@ bool FSCSEditorTreeNodeInstancedInheritedComponent::IsDefaultSceneRoot() const
 	return false;
 }
 
+bool FSCSEditorTreeNodeInstancedInheritedComponent::CanEditDefaults() const
+{
+	UActorComponent* ComponentTemplate = GetComponentTemplate();
+	if (IsNative() && ComponentTemplate)
+	{
+		return FComponentEditorUtils::CanEditNativeComponent(ComponentTemplate);
+	}
+
+	return false;
+}
+
 FText FSCSEditorTreeNodeInstancedInheritedComponent::GetDisplayName() const
 {
 	FName VariableName = GetVariableName();
@@ -1075,47 +1086,11 @@ bool FSCSEditorTreeNodeComponent::CanEditDefaults() const
 	if (!IsNative())
 	{
 		USCS_Node* SCS_Node = GetSCSNode();
-		bCanEdit = (SCS_Node != NULL);
+		bCanEdit = (SCS_Node != nullptr);
 	}
 	else if (UActorComponent* ComponentTemplate = GetComponentTemplate())
 	{
-		// Evaluate to TRUE for native nodes if it is bound to a member variable and that variable has either EditDefaultsOnly or EditAnywhere flags set
-		check(ComponentTemplate->GetOwner());
-		UClass* OwnerClass = ComponentTemplate->GetOwner()->GetActorClass();
-		if (OwnerClass != NULL)
-		{
-			UBlueprint* Blueprint = UBlueprint::GetBlueprintFromClass(OwnerClass);
-			if (Blueprint != NULL && Blueprint->ParentClass != NULL)
-			{
-				for (TFieldIterator<UProperty> It(Blueprint->ParentClass); It; ++It)
-				{
-					UProperty* Property = *It;
-					if (UObjectProperty* ObjectProp = Cast<UObjectProperty>(Property))
-					{
-						//must be editable
-						if ((Property->PropertyFlags & (CPF_Edit)) == 0)
-						{
-							continue;
-						}
-
-						UObject* ParentCDO = Blueprint->ParentClass->GetDefaultObject();
-
-						if (!ComponentTemplate->GetClass()->IsChildOf(ObjectProp->PropertyClass))
-						{
-							continue;
-						}
-
-						UObject* Object = ObjectProp->GetObjectPropertyValue(ObjectProp->ContainerPtrToValuePtr<void>(ParentCDO));
-						bCanEdit = Object != NULL && Object->GetFName() == ComponentTemplate->GetFName();
-
-						if (bCanEdit)
-						{
-							break;
-						}
-					}
-				}
-			}
-		}
+		bCanEdit = FComponentEditorUtils::CanEditNativeComponent(ComponentTemplate);
 	}
 
 	return bCanEdit;
diff --git a/Engine/Source/Editor/Kismet/Public/SSCSEditor.h b/Engine/Source/Editor/Kismet/Public/SSCSEditor.h
--- a/Engine/Source/Editor/Kismet/Public/SSCSEditor.h
+++ b/Engine/Source/Editor/Kismet/Public/SSCSEditor.h
@@ -324,7 +324,7 @@ public:
 	virtual bool IsInherited() const override { return true; }
 	virtual bool IsUserInstanced() const override { return false; }
 	virtual bool IsDefaultSceneRoot() const override;
-	virtual bool CanEditDefaults() const { return IsNative(); }
+	virtual bool CanEditDefaults() const;
 	//virtual FName GetVariableName() const override;
 	//virtual FString GetDisplayString() const override;
 	virtual FText GetDisplayName() const override;
diff --git a/Engine/Source/Editor/LevelEditor/Private/LevelEditorActions.cpp b/Engine/Source/Editor/LevelEditor/Private/LevelEditorActions.cpp
--- a/Engine/Source/Editor/LevelEditor/Private/LevelEditorActions.cpp
+++ b/Engine/Source/Editor/LevelEditor/Private/LevelEditorActions.cpp
@@ -1371,7 +1371,7 @@ bool FLevelEditorActionCallbacks::Duplicate_CanExecute()
 	if (GEditor->GetSelectedComponentCount() > 0)
 	{
 		TArray<UActorComponent*> SelectedComponents;
-		for (FSelectedEditableComponentIterator It(GEditor->GetSelectedEditableComponentIterator()); It; ++It)
+		for (FSelectionIterator It(GEditor->GetSelectedComponentIterator()); It; ++It)
 		{
 			SelectedComponents.Add(CastChecked<UActorComponent>(*It));
 		}
@@ -1407,7 +1407,7 @@ bool FLevelEditorActionCallbacks::Delete_CanExecute()
 	if (GEditor->GetSelectedComponentCount() > 0)
 	{
 		TArray<UActorComponent*> SelectedComponents;
-		for (FSelectedEditableComponentIterator It(GEditor->GetSelectedEditableComponentIterator()); It; ++It)
+		for (FSelectionIterator It(GEditor->GetSelectedComponentIterator()); It; ++It)
 		{
 			SelectedComponents.Add(CastChecked<UActorComponent>(*It));
 		}
@@ -1480,7 +1480,7 @@ bool FLevelEditorActionCallbacks::Cut_CanExecute()
 	{
 		// Make sure the components can be copied and deleted
 		TArray<UActorComponent*> SelectedComponents;
-		for (FSelectedEditableComponentIterator It(GEditor->GetSelectedEditableComponentIterator()); It; ++It)
+		for (FSelectionIterator It(GEditor->GetSelectedComponentIterator()); It; ++It)
 		{
 			SelectedComponents.Add(CastChecked<UActorComponent>(*It));
 		}
@@ -1517,7 +1517,7 @@ bool FLevelEditorActionCallbacks::Copy_CanExecute()
 	if (GEditor->GetSelectedComponentCount() > 0)
 	{
 		TArray<UActorComponent*> SelectedComponents;
-		for (FSelectedEditableComponentIterator It(GEditor->GetSelectedEditableComponentIterator()); It; ++It)
+		for (FSelectionIterator It(GEditor->GetSelectedComponentIterator()); It; ++It)
 		{
 			SelectedComponents.Add(CastChecked<UActorComponent>(*It));
 		}
diff --git a/Engine/Source/Editor/LevelEditor/Private/SActorDetails.cpp b/Engine/Source/Editor/LevelEditor/Private/SActorDetails.cpp
--- a/Engine/Source/Editor/LevelEditor/Private/SActorDetails.cpp
+++ b/Engine/Source/Editor/LevelEditor/Private/SActorDetails.cpp
@@ -7,6 +7,58 @@
 #include "Editor/PropertyEditor/Public/PropertyEditing.h"
 #include "LevelEditorGenericDetails.h"
 #include "ScopedTransaction.h"
+#include "SourceCodeNavigation.h"
+
+class SUneditableComponentWarning : public SCompoundWidget
+{
+public:
+	SLATE_BEGIN_ARGS(SUneditableComponentWarning)
+		: _WarningText()
+		, _OnHyperlinkClicked()
+	{}
+		
+		/** The rich text to show in the warning */
+		SLATE_ATTRIBUTE(FText, WarningText)
+
+		/** Called when the hyperlink in the rich text is clicked */
+		SLATE_EVENT(FSlateHyperlinkRun::FOnClick, OnHyperlinkClicked)
+
+	SLATE_END_ARGS()
+
+	/** Constructs the widget */
+	void Construct(const FArguments& InArgs)
+	{
+		ChildSlot
+		[
+			SNew(SBorder)
+			.BorderImage(FEditorStyle::Get().GetBrush("ToolPanel.GroupBorder"))
+			[
+				SNew(SHorizontalBox)
+				+ SHorizontalBox::Slot()
+				.AutoWidth()
+				.HAlign(HAlign_Center)
+				.VAlign(VAlign_Center)
+				.Padding(2)
+				[
+					SNew(SImage)
+					.Image(FEditorStyle::Get().GetBrush("Icons.Warning"))
+				]
+				+ SHorizontalBox::Slot()
+					.VAlign(VAlign_Center)
+					.Padding(2)
+					[
+						SNew(SRichTextBlock)
+						.DecoratorStyleSet(&FEditorStyle::Get())
+						.Justification(ETextJustify::Left)
+						.TextStyle(FEditorStyle::Get(), "DetailsView.BPMessageTextStyle")
+						.Text(InArgs._WarningText)
+						.AutoWrapText(true)
+						+ SRichTextBlock::HyperlinkDecorator(TEXT("HyperlinkDecorator"), InArgs._OnHyperlinkClicked)
+					]
+			]
+		];
+	}
+};
 
 void SActorDetails::Construct(const FArguments& InArgs, const FName TabIdentifier)
 {
@@ -89,33 +141,19 @@ void SActorDetails::Construct(const FArguments& InArgs, const FName TabIdentifie
 				.AutoHeight()
 				.Padding( FMargin( 0,0,0,1) )
 				[
-					SNew( SBorder )
-					.Visibility(this, &SActorDetails::GetBlueprintComponentWarningVisibility)
-					.BorderImage( FEditorStyle::Get().GetBrush( "ToolPanel.GroupBorder" ) )
-					[
-						SNew( SHorizontalBox )
-						+ SHorizontalBox::Slot()
-						.AutoWidth()
-						.HAlign( HAlign_Center )
-						.VAlign( VAlign_Center )
-						.Padding( 2 )
-						[
-							SNew( SImage )
-							.Image( FEditorStyle::Get().GetBrush("Icons.Warning") )
-						]
-						+ SHorizontalBox::Slot()
-						.VAlign( VAlign_Center )
-						.Padding( 2 )
-						[
-							SNew(SRichTextBlock)
-							.DecoratorStyleSet(&FEditorStyle::Get())
-							.Justification(ETextJustify::Left)
-							.TextStyle(FEditorStyle::Get(), "DetailsView.BPMessageTextStyle")
-							.Text(NSLOCTEXT("SActorDetails", "BlueprintedComponentWarning", "Blueprinted components must be edited in the <a id=\"HyperlinkDecorator\" style=\"DetailsView.BPMessageHyperlinkStyle\">Blueprint</>"))
-							.AutoWrapText(true)
-							+SRichTextBlock::HyperlinkDecorator(TEXT("HyperlinkDecorator"), this, &SActorDetails::OnBlueprintWarningHyperlinkClicked)
-						]
-					]
+					SNew(SUneditableComponentWarning)
+					.Visibility(this, &SActorDetails::GetBlueprintedComponentWarningVisibility)
+					.WarningText(NSLOCTEXT("SActorDetails", "BlueprintedComponentWarning", "Blueprinted components must be edited in the <a id=\"HyperlinkDecorator\" style=\"DetailsView.BPMessageHyperlinkStyle\">Blueprint</>"))
+					.OnHyperlinkClicked(this, &SActorDetails::OnBlueprintedComponentWarningHyperlinkClicked)
+				]
+				+SVerticalBox::Slot()
+				.AutoHeight()
+				.Padding( FMargin( 0,0,0,1) )
+				[
+					SNew(SUneditableComponentWarning)
+					.Visibility(this, &SActorDetails::GetNativeComponentWarningVisibility)
+					.WarningText(NSLOCTEXT("SActorDetails", "UneditableNativeComponentWarning", "Native components are editable when declared as a UProperty in <a id=\"HyperlinkDecorator\" style=\"DetailsView.BPMessageHyperlinkStyle\">C++</>"))
+					.OnHyperlinkClicked(this, &SActorDetails::OnNativeComponentWarningHyperlinkClicked)
 				]
 				+ SVerticalBox::Slot()
 				[
@@ -463,33 +501,83 @@ void SActorDetails::UpdateComponentTreeFromEditorSelection()
 bool SActorDetails::IsPropertyEditingEnabled() const
 {
 	bool bIsEditable = true;
-	const TArray<TWeakObjectPtr<UObject> >& Objects = DetailsView->GetSelectedObjects();
-	for(auto Object : Objects)
+	const TArray<FSCSEditorTreeNodePtrType> SelectedNodes = SCSEditor->GetSelectedNodes();
+	for (const auto& Node : SelectedNodes)
 	{
-		UActorComponent* ActorComp = Cast<UActorComponent>(Object.Get());
-		if(ActorComp)
+		bIsEditable = Node->CanEditDefaults() || Node->GetNodeType() == FSCSEditorTreeNode::ENodeType::RootActorNode;
+		if (!bIsEditable)
 		{
-			bIsEditable = !ActorComp->IsCreatedByConstructionScript();
-			if(!bIsEditable)
-			{
-				break;
-			}
+			break;
 		}
 	}
 	return bIsEditable;
 }
 
-void SActorDetails::OnBlueprintWarningHyperlinkClicked( const FSlateHyperlinkRun::FMetadata& Metadata )
+void SActorDetails::OnBlueprintedComponentWarningHyperlinkClicked(const FSlateHyperlinkRun::FMetadata& Metadata)
 {
 	UBlueprint* Blueprint = SCSEditor->GetBlueprint();
-	if( Blueprint )
+	if (Blueprint)
 	{
 		// Open the blueprint
 		GEditor->EditObject(Blueprint);
 	}
 }
 
-EVisibility SActorDetails::GetBlueprintComponentWarningVisibility() const
+void SActorDetails::OnNativeComponentWarningHyperlinkClicked(const FSlateHyperlinkRun::FMetadata& Metadata)
+{
+	// Find the closest native parent
+	UBlueprint* Blueprint = SCSEditor->GetBlueprint();
+	UClass* ParentClass = Blueprint ? Blueprint->ParentClass : GetActorContext()->GetClass();
+	while (ParentClass && !ParentClass->HasAllClassFlags(CLASS_Native))
+	{
+		ParentClass = ParentClass->GetSuperClass();
+	}
+
+	if (ParentClass)
+	{
+		FString NativeParentClassHeaderPath;
+		const bool bFileFound = FSourceCodeNavigation::FindClassHeaderPath(ParentClass, NativeParentClassHeaderPath)
+			&& ( IFileManager::Get().FileSize(*NativeParentClassHeaderPath) != INDEX_NONE );
+		if (bFileFound)
+		{
+			const FString AbsoluteHeaderPath = IFileManager::Get().ConvertToAbsolutePathForExternalAppForRead(*NativeParentClassHeaderPath);
+			FSourceCodeNavigation::OpenSourceFile(AbsoluteHeaderPath);
+		}
+	}
+}
+
+EVisibility SActorDetails::GetBlueprintedComponentWarningVisibility() const
 {
-	return IsPropertyEditingEnabled() ? EVisibility::Collapsed : EVisibility::Visible;
+	//return IsPropertyEditingEnabled() ? EVisibility::Collapsed : EVisibility::Visible;
+
+	bool bIsBlueprinted = false;
+
+	// Check to see if any selected components are inherited from blueprint
+	for (const auto& Node : SCSEditor->GetSelectedNodes())
+	{
+		UActorComponent* Component = Node->GetComponentTemplate();
+		bIsBlueprinted = Component ? Component->IsCreatedByConstructionScript() : false;
+		if (bIsBlueprinted)
+		{
+			break;
+		}
+	}
+
+	return bIsBlueprinted ? EVisibility::Visible : EVisibility::Collapsed;
+}
+
+EVisibility SActorDetails::GetNativeComponentWarningVisibility() const
+{
+	bool bIsUneditableNative = false;
+	for (const auto& Node : SCSEditor->GetSelectedNodes())
+	{
+		// Check to see if the component is native and not editable
+		bIsUneditableNative = Node->IsNative() && !Node->CanEditDefaults();
+		if (bIsUneditableNative)
+		{
+			break;
+		}
+	}
+	
+	return bIsUneditableNative ? EVisibility::Visible : EVisibility::Collapsed;
 }
\ No newline at end of file
diff --git a/Engine/Source/Editor/LevelEditor/Private/SActorDetails.h b/Engine/Source/Editor/LevelEditor/Private/SActorDetails.h
--- a/Engine/Source/Editor/LevelEditor/Private/SActorDetails.h
+++ b/Engine/Source/Editor/LevelEditor/Private/SActorDetails.h
@@ -35,9 +35,12 @@ private:
 	void OnSCSEditorRootSelected(AActor* Actor);
 	void OnSCSEditorTreeViewSelectionChanged(const TArray<TSharedPtr<class FSCSEditorTreeNode> >& SelectedNodes);
 	void UpdateComponentTreeFromEditorSelection();
+
 	bool IsPropertyEditingEnabled() const;
-	void OnBlueprintWarningHyperlinkClicked( const FSlateHyperlinkRun::FMetadata& Metadata );
-	EVisibility GetBlueprintComponentWarningVisibility() const;
+	EVisibility GetBlueprintedComponentWarningVisibility() const;
+	EVisibility GetNativeComponentWarningVisibility() const;
+	void OnBlueprintedComponentWarningHyperlinkClicked(const FSlateHyperlinkRun::FMetadata& Metadata);
+	void OnNativeComponentWarningHyperlinkClicked(const FSlateHyperlinkRun::FMetadata& Metadata);
 
 private:
 	TSharedPtr<SSplitter> DetailsSplitter;
diff --git a/Engine/Source/Editor/UnrealEd/Private/Kismet2/ComponentEditorUtils.cpp b/Engine/Source/Editor/UnrealEd/Private/Kismet2/ComponentEditorUtils.cpp
--- a/Engine/Source/Editor/UnrealEd/Private/Kismet2/ComponentEditorUtils.cpp
+++ b/Engine/Source/Editor/UnrealEd/Private/Kismet2/ComponentEditorUtils.cpp
@@ -111,6 +111,55 @@ protected:
 	// FCustomizableTextObjectFactory (end)
 };
 
+bool FComponentEditorUtils::CanEditNativeComponent(const UActorComponent* NativeComponent)
+{
+	// A native component can be edited if it is bound to a member variable and that variable is marked as visible in the editor
+	// Note: We aren't concerned with whether the component is marked editable - the component itself is responsible for determining which of its properties are editable
+
+	bool bCanEdit = false;
+	
+	UClass* OwnerClass = (NativeComponent && NativeComponent->GetOwner()) ? NativeComponent->GetOwner()->GetActorClass() : nullptr;
+	if (OwnerClass != nullptr)
+	{
+		// If the owner is a blueprint generated class, use the BP parent class
+		UBlueprint* Blueprint = UBlueprint::GetBlueprintFromClass(OwnerClass);
+		if (Blueprint != nullptr && Blueprint->ParentClass != nullptr)
+		{
+			OwnerClass = Blueprint->ParentClass;
+		}
+
+		for (TFieldIterator<UProperty> It(OwnerClass); It; ++It)
+		{
+			UProperty* Property = *It;
+			if (UObjectProperty* ObjectProp = Cast<UObjectProperty>(Property))
+			{
+				// Must be visible - note CPF_Edit is set for all properties that should be visible, not just those that are editable
+				if (( Property->PropertyFlags & ( CPF_Edit ) ) == 0)
+				{
+					continue;
+				}
+
+				UObject* ParentCDO = OwnerClass->GetDefaultObject();
+
+				if (!NativeComponent->GetClass()->IsChildOf(ObjectProp->PropertyClass))
+				{
+					continue;
+				}
+
+				UObject* Object = ObjectProp->GetObjectPropertyValue(ObjectProp->ContainerPtrToValuePtr<void>(ParentCDO));
+				bCanEdit = Object != nullptr && Object->GetFName() == NativeComponent->GetFName();
+
+				if (bCanEdit)
+				{
+					break;
+				}
+			}
+		}
+	}
+
+	return bCanEdit;
+}
+
 bool FComponentEditorUtils::IsValidVariableNameString(const UActorComponent* InComponent, const FString& InString)
 {
 	// First test to make sure the string is not empty and does not equate to the DefaultSceneRoot node name
diff --git a/Engine/Source/Editor/UnrealEd/Public/Kismet2/ComponentEditorUtils.h b/Engine/Source/Editor/UnrealEd/Public/Kismet2/ComponentEditorUtils.h
--- a/Engine/Source/Editor/UnrealEd/Public/Kismet2/ComponentEditorUtils.h
+++ b/Engine/Source/Editor/UnrealEd/Public/Kismet2/ComponentEditorUtils.h
@@ -8,6 +8,9 @@
 class UNREALED_API FComponentEditorUtils
 {
 public:
+	/** Tests whether the native component is editable */
+	static bool CanEditNativeComponent(const UActorComponent* NativeComponent);
+
 	/** Test whether or not the given string is a valid variable name string for the given component instance */
 	static bool IsValidVariableNameString(const UActorComponent* InComponent, const FString& InString);
 
diff --git a/Engine/Source/Runtime/Engine/Classes/Engine/Selection.h b/Engine/Source/Runtime/Engine/Classes/Engine/Selection.h
--- a/Engine/Source/Runtime/Engine/Classes/Engine/Selection.h
+++ b/Engine/Source/Runtime/Engine/Classes/Engine/Selection.h
@@ -6,6 +6,7 @@
  */
 
 #pragma once
+#include "ComponentEditorUtils.h"
 #include "Selection.generated.h"
 
 UCLASS(customConstructor, transient)
@@ -467,7 +468,13 @@ public:
 	bool IsObjectValid(const UObject* Object) const
 	{
 		const UActorComponent* Comp = Cast<UActorComponent>( Object );
-		return Comp && !Comp->IsCreatedByConstructionScript();
+		if (Comp)
+		{
+			bool bCanEditIfNative = ( Comp->CreationMethod == EComponentCreationMethod::Native ) ? FComponentEditorUtils::CanEditNativeComponent(Comp) : true;
+			return !Comp->IsCreatedByConstructionScript() && bCanEditIfNative;
+		}
+		
+		return false;
 	}
 };
 
diff --git a/Engine/Source/Runtime/Engine/Private/ActorEditorUtils.cpp b/Engine/Source/Runtime/Engine/Private/ActorEditorUtils.cpp
--- a/Engine/Source/Runtime/Engine/Private/ActorEditorUtils.cpp
+++ b/Engine/Source/Runtime/Engine/Private/ActorEditorUtils.cpp
@@ -2,6 +2,7 @@
 
 #include "EnginePrivate.h"
 #include "ActorEditorUtils.h"
+#include "ComponentEditorUtils.h"
 
 namespace FActorEditorUtils
 {
@@ -35,13 +36,21 @@ namespace FActorEditorUtils
 
 	void GetEditableComponents( const AActor* InActor, TArray<UActorComponent*>& OutEditableComponents )
 	{
-		// If a component wasn't created by a construction script (i.e. it's native or an instance), it's editable
 		TInlineComponentArray<UActorComponent*> InstanceComponents;
 		InActor->GetComponents(InstanceComponents);
 		for (auto Component : InstanceComponents)
 		{
-			if (!Component->IsCreatedByConstructionScript())
+			if (Component->CreationMethod == EComponentCreationMethod::Native)
 			{
+				// Make sure it's an exposed native component
+				if (FComponentEditorUtils::CanEditNativeComponent(Component))
+				{
+					OutEditableComponents.Add(Component);
+				}
+			}
+			else if (Component->CreationMethod == EComponentCreationMethod::Instance)
+			{
+				// Instance components are always editable
 				OutEditableComponents.Add(Component);
 			}
 		}
