Moved details panel search bar for actor details below the components tree

diff --git a/Engine/Source/Editor/LevelEditor/Private/SActorDetails.cpp b/Engine/Source/Editor/LevelEditor/Private/SActorDetails.cpp
--- a/Engine/Source/Editor/LevelEditor/Private/SActorDetails.cpp
+++ b/Engine/Source/Editor/LevelEditor/Private/SActorDetails.cpp
@@ -124,12 +124,6 @@ void SActorDetails::Construct(const FArguments& InArgs, const FName TabIdentifie
 		[
 			DetailsView->GetNameAreaWidget().ToSharedRef()
 		]
-		+ SVerticalBox::Slot()
-		.Padding(0.0f, 0.0f, 0.0f, 2.0f)
-		.AutoHeight()
-		[
-			DetailsView->GetFilterAreaWidget().ToSharedRef()
-		]
 		+SVerticalBox::Slot()
 		[
 			SAssignNew(DetailsSplitter, SSplitter)
@@ -156,6 +150,11 @@ void SActorDetails::Construct(const FArguments& InArgs, const FName TabIdentifie
 					.OnHyperlinkClicked(this, &SActorDetails::OnNativeComponentWarningHyperlinkClicked)
 				]
 				+ SVerticalBox::Slot()
+				.AutoHeight()
+				[
+					DetailsView->GetFilterAreaWidget().ToSharedRef()
+				]
+				+ SVerticalBox::Slot()
 				[
 					DetailsView.ToSharedRef()
 				]
