Enable BP component editing on a per-actor instance basis
Add ability to prevent prevent inherited versions of components from editing the properties set at a higher level

diff --git a/Engine/Source/Editor/Kismet/Private/SKismetInspector.cpp b/Engine/Source/Editor/Kismet/Private/SKismetInspector.cpp
--- a/Engine/Source/Editor/Kismet/Private/SKismetInspector.cpp
+++ b/Engine/Source/Editor/Kismet/Private/SKismetInspector.cpp
@@ -23,6 +23,57 @@
 
 #define LOCTEXT_NAMESPACE "KismetInspector"
 
+class SKismetInspectorUneditableComponentWarning : public SCompoundWidget
+{
+public:
+	SLATE_BEGIN_ARGS(SKismetInspectorUneditableComponentWarning)
+		: _WarningText()
+		, _OnHyperlinkClicked()
+	{}
+		
+		/** The rich text to show in the warning */
+		SLATE_ATTRIBUTE(FText, WarningText)
+
+		/** Called when the hyperlink in the rich text is clicked */
+		SLATE_EVENT(FSlateHyperlinkRun::FOnClick, OnHyperlinkClicked)
+
+	SLATE_END_ARGS()
+
+	/** Constructs the widget */
+	void Construct(const FArguments& InArgs)
+	{
+		ChildSlot
+		[
+			SNew(SBorder)
+			.BorderImage(FEditorStyle::Get().GetBrush("ToolPanel.GroupBorder"))
+			[
+				SNew(SHorizontalBox)
+				+ SHorizontalBox::Slot()
+				.AutoWidth()
+				.HAlign(HAlign_Center)
+				.VAlign(VAlign_Center)
+				.Padding(2)
+				[
+					SNew(SImage)
+					.Image(FEditorStyle::Get().GetBrush("Icons.Warning"))
+				]
+				+ SHorizontalBox::Slot()
+					.VAlign(VAlign_Center)
+					.Padding(2)
+					[
+						SNew(SRichTextBlock)
+						.DecoratorStyleSet(&FEditorStyle::Get())
+						.Justification(ETextJustify::Left)
+						.TextStyle(FEditorStyle::Get(), "DetailsView.BPMessageTextStyle")
+						.Text(InArgs._WarningText)
+						.AutoWrapText(true)
+						+ SRichTextBlock::HyperlinkDecorator(TEXT("HyperlinkDecorator"), InArgs._OnHyperlinkClicked)
+					]
+			]
+		];
+	}
+};
+
 //////////////////////////////////////////////////////////////////////////
 // FKismetSelectionInfo
 
@@ -82,6 +133,10 @@ TSharedRef<SWidget> SKismetInspector::MakeContextualEditingWidget(struct FKismet
 	// Show the property editor
 	PropertyView->HideFilterArea(Options.bHideFilterArea);
 	PropertyView->SetObjects(SelectionInfo.ObjectsForPropertyEditing, Options.bForceRefresh);
+
+	PropertyView->SetIsPropertyEditingEnabledDelegate(FIsPropertyEditingEnabled::CreateSP(this, &SKismetInspector::IsPropertyEditingEnabled));
+
+
 	if (SelectionInfo.ObjectsForPropertyEditing.Num())
 	{
 		ContextualEditingWidget->AddSlot()
@@ -91,7 +146,20 @@ TSharedRef<SWidget> SKismetInspector::MakeContextualEditingWidget(struct FKismet
 			SNew( SBox )
 			.Visibility(this, &SKismetInspector::GetPropertyViewVisibility)
 			[
-				PropertyView.ToSharedRef()
+				SNew( SVerticalBox )
+				+SVerticalBox::Slot()
+				.AutoHeight()
+				.Padding( FMargin( 0,0,0,1) )
+				[
+					SNew(SKismetInspectorUneditableComponentWarning)
+					.Visibility(this, &SKismetInspector::GetInheritedBlueprintComponentWarningVisibility)
+					.WarningText(NSLOCTEXT("SKismetInspector", "BlueprintUneditableInheritedComponentWarning", "Components flagged as not editable when inherited must be edited in the <a id=\"HyperlinkDecorator\" style=\"DetailsView.BPMessageHyperlinkStyle\">Parent Blueprint</>"))
+					.OnHyperlinkClicked(this, &SKismetInspector::OnInheritedBlueprintComponentWarningHyperlinkClicked)
+				]
+				+SVerticalBox::Slot()
+				[
+					PropertyView.ToSharedRef()
+				]
 			]
 		];
 
@@ -641,6 +709,53 @@ EVisibility SKismetInspector::GetPropertyViewVisibility() const
 	return bShowInspectorPropertyView? EVisibility::Visible : EVisibility::Collapsed;
 }
 
+bool SKismetInspector::IsPropertyEditingEnabled() const
+{
+	bool bIsEditable = true;
+	for (const TWeakObjectPtr<UObject>& SelectedObject : SelectedObjects)
+	{
+		UActorComponent* Component = Cast<UActorComponent>(SelectedObject.Get());
+		if (Component && !CastChecked<UActorComponent>(Component->GetArchetype())->IsEditableWhenInherited() )
+		{
+			bIsEditable = false;
+			break;
+		}
+	}
+	return bIsEditable;
+}
+
+EVisibility SKismetInspector::GetInheritedBlueprintComponentWarningVisibility() const
+{
+	bool bIsUneditableBlueprintComponent = false;
+
+	// Check to see if any selected components are inherited from blueprint
+	for (const TWeakObjectPtr<UObject>& SelectedObject : SelectedObjects)
+	{
+		UActorComponent* Component = Cast<UActorComponent>(SelectedObject.Get());
+		bIsUneditableBlueprintComponent = Component ? !CastChecked<UActorComponent>(Component->GetArchetype())->IsEditableWhenInherited() : false;
+		if (bIsUneditableBlueprintComponent)
+		{
+			break;
+		}
+	}
+
+	return bIsUneditableBlueprintComponent ? EVisibility::Visible : EVisibility::Collapsed;
+}
+
+void SKismetInspector::OnInheritedBlueprintComponentWarningHyperlinkClicked(const FSlateHyperlinkRun::FMetadata& Metadata)
+{
+	if (BlueprintEditorPtr.IsValid())
+	{
+		UBlueprint* Blueprint = BlueprintEditorPtr.Pin()->GetBlueprintObj();
+		if (Blueprint && Blueprint->ParentClass->HasAllClassFlags(CLASS_CompiledFromBlueprint))
+		{
+			// Open the blueprint
+			GEditor->EditObject(CastChecked<UBlueprint>(Blueprint->ParentClass->ClassGeneratedBy));
+		}
+	}
+}
+
+
 ECheckBoxState SKismetInspector::GetPublicViewCheckboxState() const
 {
 	return PublicViewState;
diff --git a/Engine/Source/Editor/Kismet/Private/SSCSEditor.cpp b/Engine/Source/Editor/Kismet/Private/SSCSEditor.cpp
--- a/Engine/Source/Editor/Kismet/Private/SSCSEditor.cpp
+++ b/Engine/Source/Editor/Kismet/Private/SSCSEditor.cpp
@@ -844,12 +844,7 @@ bool FSCSEditorTreeNodeInstancedInheritedComponent::IsDefaultSceneRoot() const
 bool FSCSEditorTreeNodeInstancedInheritedComponent::CanEditDefaults() const
 {
 	UActorComponent* ComponentTemplate = GetComponentTemplate();
-	if (IsNative() && ComponentTemplate)
-	{
-		return FComponentEditorUtils::CanEditNativeComponent(ComponentTemplate);
-	}
-
-	return false;
+	return (ComponentTemplate ? ComponentTemplate->IsEditableWhenInherited() : false);
 }
 
 FText FSCSEditorTreeNodeInstancedInheritedComponent::GetDisplayName() const
diff --git a/Engine/Source/Editor/Kismet/Public/SKismetInspector.h b/Engine/Source/Editor/Kismet/Public/SKismetInspector.h
--- a/Engine/Source/Editor/Kismet/Public/SKismetInspector.h
+++ b/Engine/Source/Editor/Kismet/Public/SKismetInspector.h
@@ -143,6 +143,15 @@ protected:
 	/** Returns whether the property view be visible */
 	EVisibility GetPropertyViewVisibility() const;
 
+	/** Returns whether the properties in the view should be editable */
+	bool IsPropertyEditingEnabled() const;
+
+	/** Returns whether the warning about an inherited component not being editable should be visible */
+	EVisibility GetInheritedBlueprintComponentWarningVisibility() const;
+
+	/** Opens the parent blueprint when the hyperlink in the warning is clicked */
+	void OnInheritedBlueprintComponentWarningHyperlinkClicked(const FSlateHyperlinkRun::FMetadata& Metadata);
+
 	/**
 	 * Generates a widget that is used to edit the specified object array contextually.  This widget
 	 * will be displayed along with a property view in the level editor
diff --git a/Engine/Source/Editor/LevelEditor/Private/SActorDetails.cpp b/Engine/Source/Editor/LevelEditor/Private/SActorDetails.cpp
--- a/Engine/Source/Editor/LevelEditor/Private/SActorDetails.cpp
+++ b/Engine/Source/Editor/LevelEditor/Private/SActorDetails.cpp
@@ -9,10 +9,10 @@
 #include "ScopedTransaction.h"
 #include "SourceCodeNavigation.h"
 
-class SUneditableComponentWarning : public SCompoundWidget
+class SActorDetailsUneditableComponentWarning : public SCompoundWidget
 {
 public:
-	SLATE_BEGIN_ARGS(SUneditableComponentWarning)
+	SLATE_BEGIN_ARGS(SActorDetailsUneditableComponentWarning)
 		: _WarningText()
 		, _OnHyperlinkClicked()
 	{}
@@ -135,16 +135,25 @@ void SActorDetails::Construct(const FArguments& InArgs, const FName TabIdentifie
 				.AutoHeight()
 				.Padding( FMargin( 0,0,0,1) )
 				[
-					SNew(SUneditableComponentWarning)
-					.Visibility(this, &SActorDetails::GetBlueprintedComponentWarningVisibility)
-					.WarningText(NSLOCTEXT("SActorDetails", "BlueprintedComponentWarning", "Blueprinted components must be edited in the <a id=\"HyperlinkDecorator\" style=\"DetailsView.BPMessageHyperlinkStyle\">Blueprint</>"))
+					SNew(SActorDetailsUneditableComponentWarning)
+					.Visibility(this, &SActorDetails::GetUCSComponentWarningVisibility)
+					.WarningText(NSLOCTEXT("SActorDetails", "BlueprintUCSComponentWarning", "Components created by the User Construction Script can only be edited in the <a id=\"HyperlinkDecorator\" style=\"DetailsView.BPMessageHyperlinkStyle\">Blueprint</>"))
 					.OnHyperlinkClicked(this, &SActorDetails::OnBlueprintedComponentWarningHyperlinkClicked)
 				]
 				+SVerticalBox::Slot()
 				.AutoHeight()
 				.Padding( FMargin( 0,0,0,1) )
 				[
-					SNew(SUneditableComponentWarning)
+					SNew(SActorDetailsUneditableComponentWarning)
+					.Visibility(this, &SActorDetails::GetInheritedBlueprintComponentWarningVisibility)
+					.WarningText(NSLOCTEXT("SActorDetails", "BlueprintUneditableInheritedComponentWarning", "Components flagged as not editable when inherited must be edited in the <a id=\"HyperlinkDecorator\" style=\"DetailsView.BPMessageHyperlinkStyle\">Blueprint</>"))
+					.OnHyperlinkClicked(this, &SActorDetails::OnBlueprintedComponentWarningHyperlinkClicked)
+				]
+				+SVerticalBox::Slot()
+				.AutoHeight()
+				.Padding( FMargin( 0,0,0,1) )
+				[
+					SNew(SActorDetailsUneditableComponentWarning)
 					.Visibility(this, &SActorDetails::GetNativeComponentWarningVisibility)
 					.WarningText(NSLOCTEXT("SActorDetails", "UneditableNativeComponentWarning", "Native components are editable when declared as a UProperty in <a id=\"HyperlinkDecorator\" style=\"DetailsView.BPMessageHyperlinkStyle\">C++</>"))
 					.OnHyperlinkClicked(this, &SActorDetails::OnNativeComponentWarningHyperlinkClicked)
@@ -545,24 +554,77 @@ void SActorDetails::OnNativeComponentWarningHyperlinkClicked(const FSlateHyperli
 	}
 }
 
-EVisibility SActorDetails::GetBlueprintedComponentWarningVisibility() const
+EVisibility SActorDetails::GetUCSComponentWarningVisibility() const
+{
+	bool bIsUneditableBlueprintComponent = false;
+
+	// Check to see if any selected components are inherited from blueprint
+	for (const auto& Node : SCSEditor->GetSelectedNodes())
+	{
+		if (!Node->IsNative())
+		{
+			UActorComponent* Component = Node->GetComponentTemplate();
+			bIsUneditableBlueprintComponent = Component ? Component->CreationMethod == EComponentCreationMethod::UserConstructionScript : false;
+			if (bIsUneditableBlueprintComponent)
+			{
+				break;
+			}
+		}
+	}
+
+	return bIsUneditableBlueprintComponent ? EVisibility::Visible : EVisibility::Collapsed;
+}
+
+bool NotEditableSetByBlueprint(UActorComponent* Component)
 {
-	//return IsPropertyEditingEnabled() ? EVisibility::Collapsed : EVisibility::Visible;
+	// Determine if it is locked out from a blueprint or from the native
+	UActorComponent* Archetype = CastChecked<UActorComponent>(Component->GetArchetype());
+	while (Archetype)
+	{
+		if (Archetype->GetOuter()->IsA<UBlueprintGeneratedClass>() || Archetype->GetOuter()->GetClass()->HasAllClassFlags(CLASS_CompiledFromBlueprint))
+		{
+			if (!Archetype->bEditableWhenInherited)
+			{
+				return true;
+			}
 
-	bool bIsBlueprinted = false;
+			Archetype = CastChecked<UActorComponent>(Archetype->GetArchetype());
+		}
+		else
+		{
+			Archetype = nullptr;
+		}
+	}
+
+	return false;
+}
+
+EVisibility SActorDetails::GetInheritedBlueprintComponentWarningVisibility() const
+{
+	bool bIsUneditableBlueprintComponent = false;
 
 	// Check to see if any selected components are inherited from blueprint
 	for (const auto& Node : SCSEditor->GetSelectedNodes())
 	{
-		UActorComponent* Component = Node->GetComponentTemplate();
-		bIsBlueprinted = Component ? Component->IsCreatedByConstructionScript() : false;
-		if (bIsBlueprinted)
+		if (!Node->IsNative())
+		{
+			if (UActorComponent* Component = Node->GetComponentTemplate())
+			{
+				if (!Component->IsEditableWhenInherited() && Component->CreationMethod == EComponentCreationMethod::SimpleConstructionScript)
+				{
+					bIsUneditableBlueprintComponent = true;
+					break;
+				}
+			}
+		}
+		else if (!Node->CanEditDefaults() && NotEditableSetByBlueprint(Node->GetComponentTemplate()))
 		{
+			bIsUneditableBlueprintComponent = true;
 			break;
 		}
 	}
 
-	return bIsBlueprinted ? EVisibility::Visible : EVisibility::Collapsed;
+	return bIsUneditableBlueprintComponent ? EVisibility::Visible : EVisibility::Collapsed;
 }
 
 EVisibility SActorDetails::GetNativeComponentWarningVisibility() const
@@ -571,9 +633,9 @@ EVisibility SActorDetails::GetNativeComponentWarningVisibility() const
 	for (const auto& Node : SCSEditor->GetSelectedNodes())
 	{
 		// Check to see if the component is native and not editable
-		bIsUneditableNative = Node->IsNative() && !Node->CanEditDefaults();
-		if (bIsUneditableNative)
+		if (Node->IsNative() && !Node->CanEditDefaults() && !NotEditableSetByBlueprint(Node->GetComponentTemplate()))
 		{
+			bIsUneditableNative = true;
 			break;
 		}
 	}
diff --git a/Engine/Source/Editor/LevelEditor/Private/SActorDetails.h b/Engine/Source/Editor/LevelEditor/Private/SActorDetails.h
--- a/Engine/Source/Editor/LevelEditor/Private/SActorDetails.h
+++ b/Engine/Source/Editor/LevelEditor/Private/SActorDetails.h
@@ -37,7 +37,8 @@ private:
 	void UpdateComponentTreeFromEditorSelection();
 
 	bool IsPropertyEditingEnabled() const;
-	EVisibility GetBlueprintedComponentWarningVisibility() const;
+	EVisibility GetUCSComponentWarningVisibility() const;
+	EVisibility GetInheritedBlueprintComponentWarningVisibility() const;
 	EVisibility GetNativeComponentWarningVisibility() const;
 	void OnBlueprintedComponentWarningHyperlinkClicked(const FSlateHyperlinkRun::FMetadata& Metadata);
 	void OnNativeComponentWarningHyperlinkClicked(const FSlateHyperlinkRun::FMetadata& Metadata);
diff --git a/Engine/Source/Editor/UnrealEd/Private/Editor.cpp b/Engine/Source/Editor/UnrealEd/Private/Editor.cpp
--- a/Engine/Source/Editor/UnrealEd/Private/Editor.cpp
+++ b/Engine/Source/Editor/UnrealEd/Private/Editor.cpp
@@ -7378,6 +7378,9 @@ namespace EditorUtilities
 					}
 				}
 
+				TSet<const UProperty*> SourceUCSModifiedProperties;
+				SourceComponent->GetUCSModifiedProperties(SourceUCSModifiedProperties);
+
 				// Copy component properties
 				bool bIsFirstModification = true;
 				for( UProperty* Property = ComponentClass->PropertyLink; Property != NULL; Property = Property->PropertyLinkNext )
@@ -7390,7 +7393,7 @@ namespace EditorUtilities
 						Property->GetFName() == FName( TEXT( "RelativeLocation" ) ) ||
 						Property->GetFName() == FName( TEXT( "RelativeRotation" ) );
 
-					if( !bIsTransient && !bIsIdentical && !bIsComponent
+					if( !bIsTransient && !bIsIdentical && !bIsComponent && !SourceUCSModifiedProperties.Contains(Property)
 						&& ( !bIsTransform || SourceComponent != SourceActor->GetRootComponent() || ( !SourceActor->HasAnyFlags( RF_ClassDefaultObject | RF_ArchetypeObject ) && !TargetActor->HasAnyFlags( RF_ClassDefaultObject | RF_ArchetypeObject ) ) ) )
 					{
 						const bool bIsSafeToCopy = !( Options & ECopyOptions::OnlyCopyEditOrInterpProperties ) || ( Property->HasAnyPropertyFlags( CPF_Edit | CPF_Interp ) );
diff --git a/Engine/Source/Runtime/Core/Public/UObject/ObjectVersion.h b/Engine/Source/Runtime/Core/Public/UObject/ObjectVersion.h
--- a/Engine/Source/Runtime/Core/Public/UObject/ObjectVersion.h
+++ b/Engine/Source/Runtime/Core/Public/UObject/ObjectVersion.h
@@ -484,7 +484,8 @@ enum EUnrealEngineObjectUE4Version
 	VER_UE4_REBUILD_HIERARCHICAL_INSTANCE_TREES,
 	// Package summary includes an CompatibleWithEngineVersion field, separately to the version it's saved with
 	VER_UE4_PACKAGE_SUMMARY_HAS_COMPATIBLE_ENGINE_VERSION,
-
+	// Track UCS modified properties on Actor Components
+	VER_UE4_TRACK_UCS_MODIFIED_PROPERTIES,
 
 	// -----<new versions can be added before this line>-------------------------------------------------
 	// - this needs to be the last line (see note below)
diff --git a/Engine/Source/Runtime/Engine/Classes/Components/ActorComponent.h b/Engine/Source/Runtime/Engine/Classes/Components/ActorComponent.h
--- a/Engine/Source/Runtime/Engine/Classes/Components/ActorComponent.h
+++ b/Engine/Source/Runtime/Engine/Classes/Components/ActorComponent.h
@@ -3,6 +3,7 @@
 #pragma once
 #include "Engine/EngineBaseTypes.h"
 #include "Engine/EngineTypes.h"
+#include "Engine/MemberReference.h"
 #include "Interfaces/Interface_AssetUserData.h"
 #include "ActorComponent.generated.h"
 
@@ -115,6 +116,9 @@ public:
 	UPROPERTY(transient, ReplicatedUsing=OnRep_IsActive)
 	uint32 bIsActive:1;
 
+	UPROPERTY(EditDefaultsOnly, Category="Variable")
+	uint32 bEditableWhenInherited:1;
+
 	/** If true, we call the virtual InitializeComponent */
 	uint32 bWantsInitializeComponent:1;
 
@@ -132,8 +136,20 @@ private:
 	/** Indicates that BeginPlay has been called, but EndPlay has not yet */
 	uint32 bHasBegunPlay:1;
 
+	friend class FActorComponentInstanceData;
+
+	UPROPERTY()
+	TArray<FSimpleMemberReference> UCSModifiedProperties;
+
 public:
 
+	void DetermineUCSModifiedProperties();
+	void GetUCSModifiedProperties(TSet<const UProperty*>& ModifiedProperties) const;
+
+#if WITH_EDITOR
+	bool IsEditableWhenInherited() const;
+#endif
+
 	bool HasBeenCreated() const { return bHasBeenCreated; }
 	bool HasBeenInitialized() const { return bHasBeenInitialized; }
 	bool HasBegunPlay() const { return bHasBegunPlay; }
diff --git a/Engine/Source/Runtime/Engine/Classes/Engine/Selection.h b/Engine/Source/Runtime/Engine/Classes/Engine/Selection.h
--- a/Engine/Source/Runtime/Engine/Classes/Engine/Selection.h
+++ b/Engine/Source/Runtime/Engine/Classes/Engine/Selection.h
@@ -471,18 +471,13 @@ class FSelectedEditableComponentFilter
 public:
 	bool IsObjectValid(const UObject* Object) const
 	{
+#if WITH_EDITOR
 		const UActorComponent* Comp = Cast<UActorComponent>( Object );
 		if (Comp)
 		{
-			bool bCanEditIfNative = true;
-#if WITH_EDITOR
-			if (Comp->CreationMethod == EComponentCreationMethod::Native)
-			{
-				bCanEditIfNative = FComponentEditorUtils::CanEditNativeComponent(Comp);
-			}
-#endif
-			return !Comp->IsCreatedByConstructionScript() && bCanEditIfNative;
+			return Comp->IsEditableWhenInherited();
 		}
+#endif
 		
 		return false;
 	}
diff --git a/Engine/Source/Runtime/Engine/Private/ActorComponent.cpp b/Engine/Source/Runtime/Engine/Private/ActorComponent.cpp
--- a/Engine/Source/Runtime/Engine/Private/ActorComponent.cpp
+++ b/Engine/Source/Runtime/Engine/Private/ActorComponent.cpp
@@ -128,6 +128,7 @@ UActorComponent::UActorComponent(const FObjectInitializer& ObjectInitializer /*=
 
 	bAutoRegister = true;
 	bNetAddressable = false;
+	bEditableWhenInherited = true;
 }
 
 void UActorComponent::PostInitProperties()
@@ -147,14 +148,14 @@ void UActorComponent::PostLoad()
 
 	if (GetLinkerUE4Version() < VER_UE4_ACTOR_COMPONENT_CREATION_METHOD)
 	{
-	if (bCreatedByConstructionScript_DEPRECATED)
-	{
-			CreationMethod = EComponentCreationMethod::SimpleConstructionScript;
-	}
-	else if (bInstanceComponent_DEPRECATED)
-	{
-		CreationMethod = EComponentCreationMethod::Instance;
-	}
+		if (bCreatedByConstructionScript_DEPRECATED)
+		{
+				CreationMethod = EComponentCreationMethod::SimpleConstructionScript;
+		}
+		else if (bInstanceComponent_DEPRECATED)
+		{
+			CreationMethod = EComponentCreationMethod::Instance;
+		}
 
 		if (CreationMethod == EComponentCreationMethod::SimpleConstructionScript)
 		{
@@ -177,6 +178,11 @@ void UActorComponent::PostLoad()
 			}
 		}
 	}
+
+	if (GetLinkerUE4Version() < VER_UE4_TRACK_UCS_MODIFIED_PROPERTIES)
+	{
+		DetermineUCSModifiedProperties();
+	}
 }
 
 bool UActorComponent::Rename( const TCHAR* InName, UObject* NewOuter, ERenameFlags Flags )
@@ -1424,5 +1430,75 @@ void UActorComponent::OnRep_IsActive()
 	SetComponentTickEnabled(bIsActive);
 }
 
+#if WITH_EDITOR
+bool UActorComponent::IsEditableWhenInherited() const
+{
+	bool bCanEdit = bEditableWhenInherited;
+	if (bCanEdit)
+	{
+		if (CreationMethod == EComponentCreationMethod::Native && !IsTemplate())
+		{
+			bCanEdit = FComponentEditorUtils::CanEditNativeComponent(this);
+		}
+		else if (CreationMethod == EComponentCreationMethod::UserConstructionScript)
+		{
+			bCanEdit = false;
+		}
+	}
+	return bCanEdit;
+}
+#endif
+
+void UActorComponent::DetermineUCSModifiedProperties()
+{
+	class FComponentPropertySkipper : public FArchive
+	{
+	public:
+		FComponentPropertySkipper()
+			: FArchive()
+		{
+			ArIsSaving = true;
+		}
+
+		virtual bool ShouldSkipProperty(const UProperty* InProperty) const override
+		{
+			return (    InProperty->HasAnyPropertyFlags(CPF_Transient | CPF_ContainsInstancedReference | CPF_InstancedReference)
+					|| !InProperty->HasAnyPropertyFlags(CPF_Edit | CPF_Interp));
+		}
+	} PropertySkipper;
+
+	UClass* ComponentClass = GetClass();
+	UObject* ComponentArchetype = GetArchetype();
+
+	UCSModifiedProperties.Empty();
+
+	for (TFieldIterator<UProperty> It(ComponentClass); It; ++It)
+	{
+		UProperty* Property = *It;
+		if( Property->ShouldSerializeValue(PropertySkipper) )
+		{
+			for( int32 Idx=0; Idx<Property->ArrayDim; Idx++ )
+			{
+				uint8* DataPtr      = Property->ContainerPtrToValuePtr           <uint8>((uint8*)this, Idx);
+				uint8* DefaultValue = Property->ContainerPtrToValuePtrForDefaults<uint8>(ComponentClass, (uint8*)ComponentArchetype, Idx);
+				if (!Property->Identical( DataPtr, DefaultValue))
+				{
+					UCSModifiedProperties.Add(FSimpleMemberReference());
+					UCSModifiedProperties.Last().FillSimpleMemberReference<UProperty>(Property);
+					break;
+				}
+			}
+		}
+	}
+}
+
+void UActorComponent::GetUCSModifiedProperties(TSet<const UProperty*>& ModifiedProperties) const
+{
+	for (const FSimpleMemberReference& MemberReference : UCSModifiedProperties)
+	{
+		ModifiedProperties.Add(MemberReference.ResolveSimpleMemberReference<UProperty>());
+	}
+}
+
 
 #undef LOCTEXT_NAMESPACE
diff --git a/Engine/Source/Runtime/Engine/Private/ActorConstruction.cpp b/Engine/Source/Runtime/Engine/Private/ActorConstruction.cpp
--- a/Engine/Source/Runtime/Engine/Private/ActorConstruction.cpp
+++ b/Engine/Source/Runtime/Engine/Private/ActorConstruction.cpp
@@ -456,6 +456,19 @@ void AActor::ExecuteConstruction(const FTransform& Transform, const FComponentIn
 				ProcessUserConstructionScript();
 			}
 
+			// Since re-run construction scripts will never be run and we want to keep dynamic spawning fast, don't spend time
+			// determining the UCS modified properties in game worlds
+			if (!GetWorld()->IsGameWorld())
+			{
+				for (UActorComponent* Component : GetComponents())
+				{
+					if (Component)
+					{
+						Component->DetermineUCSModifiedProperties();
+					}
+				}
+			}
+
 			// Bind any delegates on components			
 			((UBlueprintGeneratedClass*)GetClass())->BindDynamicDelegates(this); // We have a BP stack, we must have a UBlueprintGeneratedClass...
 
diff --git a/Engine/Source/Runtime/Engine/Private/ComponentInstanceDataCache.cpp b/Engine/Source/Runtime/Engine/Private/ComponentInstanceDataCache.cpp
--- a/Engine/Source/Runtime/Engine/Private/ComponentInstanceDataCache.cpp
+++ b/Engine/Source/Runtime/Engine/Private/ComponentInstanceDataCache.cpp
@@ -36,7 +36,7 @@ FActorComponentInstanceData::FActorComponentInstanceData(const UActorComponent*
 		}
 	}
 
-	/*if (SourceComponent->CreationMethod == EComponentCreationMethod::SimpleConstructionScript)
+	if (SourceComponent->IsEditableWhenInherited())
 	{
 		class FComponentPropertyWriter : public FObjectWriter
 		{
@@ -44,18 +44,24 @@ FActorComponentInstanceData::FActorComponentInstanceData(const UActorComponent*
 			FComponentPropertyWriter(const UActorComponent* Component, TArray<uint8>& InBytes)
 				: FObjectWriter(InBytes)
 			{
+				Component->GetUCSModifiedProperties(PropertiesToSkip);
+
 				UClass* ComponentClass = Component->GetClass();
 				ComponentClass->SerializeTaggedProperties(*this, (uint8*)Component, ComponentClass, (uint8*)Component->GetArchetype());
 			}
 
 			virtual bool ShouldSkipProperty(const UProperty* InProperty) const override
 			{
-				return (    InProperty->HasAnyPropertyFlags(CPF_Transient | CPF_ContainsInstancedReference | CPF_InstancedReference)
-						|| !InProperty->HasAnyPropertyFlags(CPF_Edit | CPF_Interp));
+				return (	InProperty->HasAnyPropertyFlags(CPF_Transient | CPF_ContainsInstancedReference | CPF_InstancedReference)
+						|| !InProperty->HasAnyPropertyFlags(CPF_Edit | CPF_Interp)
+						|| PropertiesToSkip.Contains(InProperty));
 			}
 
+		private:
+			TSet<const UProperty*> PropertiesToSkip;
+
 		} ComponentPropertyWriter(SourceComponent, SavedProperties);
-	}*/
+	}
 }
 
 bool FActorComponentInstanceData::MatchesComponent(const UActorComponent* Component) const
@@ -91,17 +97,31 @@ bool FActorComponentInstanceData::MatchesComponent(const UActorComponent* Compon
 
 void FActorComponentInstanceData::ApplyToComponent(UActorComponent* Component, const ECacheApplyPhase CacheApplyPhase)
 {
-	if (CacheApplyPhase == ECacheApplyPhase::PostSimpleConstructionScript && SavedProperties.Num() > 0)
+	// After the user construction script has run we will re-apply all the cached changes that do not conflict
+	// with a change that the user construction script made.
+	if (CacheApplyPhase == ECacheApplyPhase::PostUserConstructionScript && SavedProperties.Num() > 0)
 	{
+		Component->DetermineUCSModifiedProperties();
+
 		class FComponentPropertyReader : public FObjectReader
 		{
 		public:
 			FComponentPropertyReader(UActorComponent* InComponent, TArray<uint8>& InBytes)
 				: FObjectReader(InBytes)
 			{
+				InComponent->GetUCSModifiedProperties(PropertiesToSkip);
+
 				UClass* Class = InComponent->GetClass();
 				Class->SerializeTaggedProperties(*this, (uint8*)InComponent, Class, nullptr);
 			}
+
+			virtual bool ShouldSkipProperty(const UProperty* InProperty) const override
+			{
+				return PropertiesToSkip.Contains(InProperty);
+			}
+
+			TSet<const UProperty*> PropertiesToSkip;
+
 		} ComponentPropertyReader(Component, SavedProperties);	
 	}
 }
