Make UCSModified Properties read-only on actor instance components

diff --git a/Engine/Source/Editor/LevelEditor/Private/SActorDetails.cpp b/Engine/Source/Editor/LevelEditor/Private/SActorDetails.cpp
--- a/Engine/Source/Editor/LevelEditor/Private/SActorDetails.cpp
+++ b/Engine/Source/Editor/LevelEditor/Private/SActorDetails.cpp
@@ -95,6 +95,7 @@ void SActorDetails::Construct(const FArguments& InArgs, const FName TabIdentifie
 	};
 
 	DetailsView->SetIsPropertyVisibleDelegate(FIsPropertyVisible::CreateLambda(IsPropertyVisible));
+	DetailsView->SetIsPropertyReadOnlyDelegate(FIsPropertyReadOnly::CreateSP(this, &SActorDetails::IsPropertyReadOnly));
 
 	DetailsView->SetIsPropertyEditingEnabledDelegate(FIsPropertyEditingEnabled::CreateSP(this, &SActorDetails::IsPropertyEditingEnabled));
 
@@ -506,6 +507,27 @@ void SActorDetails::UpdateComponentTreeFromEditorSelection()
 	}
 }
 
+bool SActorDetails::IsPropertyReadOnly(const FPropertyAndParent& PropertyAndParent) const
+{
+	bool bIsReadOnly = false;
+	const TArray<FSCSEditorTreeNodePtrType> SelectedNodes = SCSEditor->GetSelectedNodes();
+	for (const auto& Node : SelectedNodes)
+	{
+		UActorComponent* Component = Node->GetComponentTemplate();
+		if (Component && Component->CreationMethod == EComponentCreationMethod::SimpleConstructionScript)
+		{
+			TSet<const UProperty*> UCSModifiedProperties;
+			Component->GetUCSModifiedProperties(UCSModifiedProperties);
+			if (UCSModifiedProperties.Contains(&PropertyAndParent.Property) || (PropertyAndParent.ParentProperty && UCSModifiedProperties.Contains(PropertyAndParent.ParentProperty)))
+			{
+				bIsReadOnly = true;
+				break;
+			}
+		}
+	}
+	return bIsReadOnly;
+}
+
 bool SActorDetails::IsPropertyEditingEnabled() const
 {
 	bool bIsEditable = true;
diff --git a/Engine/Source/Editor/LevelEditor/Private/SActorDetails.h b/Engine/Source/Editor/LevelEditor/Private/SActorDetails.h
--- a/Engine/Source/Editor/LevelEditor/Private/SActorDetails.h
+++ b/Engine/Source/Editor/LevelEditor/Private/SActorDetails.h
@@ -36,6 +36,7 @@ private:
 	void OnSCSEditorTreeViewSelectionChanged(const TArray<TSharedPtr<class FSCSEditorTreeNode> >& SelectedNodes);
 	void UpdateComponentTreeFromEditorSelection();
 
+	bool IsPropertyReadOnly(const struct FPropertyAndParent& PropertyAndParent) const;
 	bool IsPropertyEditingEnabled() const;
 	EVisibility GetUCSComponentWarningVisibility() const;
 	EVisibility GetInheritedBlueprintComponentWarningVisibility() const;
