Unify Slate's FInputGesture and Engine's FInputChord as FInputChord defined in Slate

diff --git a/Engine/Source/Editor/LevelEditor/Private/LevelEditorActions.cpp b/Engine/Source/Editor/LevelEditor/Private/LevelEditorActions.cpp
--- a/Engine/Source/Editor/LevelEditor/Private/LevelEditorActions.cpp
+++ b/Engine/Source/Editor/LevelEditor/Private/LevelEditorActions.cpp
@@ -2792,15 +2792,15 @@ class UWorld* FLevelEditorActionCallbacks::GetWorld()
 PRAGMA_DISABLE_OPTIMIZATION
 void FLevelEditorCommands::RegisterCommands()
 {
-	UI_COMMAND( BrowseDocumentation, "Documentation...", "Opens the main documentation page", EUserInterfaceActionType::Button, FInputGesture( EKeys::F1 ) );
-	UI_COMMAND( BrowseAPIReference, "API Reference...", "Opens the API reference documentation", EUserInterfaceActionType::Button, FInputGesture() );
-	UI_COMMAND( BrowseViewportControls, "Viewport Controls...", "Opens the viewport controls cheat sheet", EUserInterfaceActionType::Button, FInputGesture() );
-	UI_COMMAND( NewLevel, "New Level...", "Create a new level, or choose a level template to start from.", EUserInterfaceActionType::Button, FInputGesture( EModifierKey::Control, EKeys::N ) );
-	UI_COMMAND( OpenLevel, "Open Level...", "Loads an existing level", EUserInterfaceActionType::Button, FInputGesture( EModifierKey::Control, EKeys::O ) );
-	UI_COMMAND( Save, "Save", "Saves the current level to disk", EUserInterfaceActionType::Button, FInputGesture() );
-	UI_COMMAND( SaveAs, "Save As...", "Save the current level as...", EUserInterfaceActionType::Button, FInputGesture( EModifierKey::Shift|EModifierKey::Control, EKeys::S ) );
-	UI_COMMAND( SaveAllLevels, "Save All Levels", "Saves all unsaved levels to disk", EUserInterfaceActionType::Button, FInputGesture() );
-	UI_COMMAND( ToggleFavorite, "Toggle Favorite", "Sets whether the currently loaded level will appear in the list of favorite levels", EUserInterfaceActionType::Button, FInputGesture() );
+	UI_COMMAND( BrowseDocumentation, "Documentation...", "Opens the main documentation page", EUserInterfaceActionType::Button, FInputChord( EKeys::F1 ) );
+	UI_COMMAND( BrowseAPIReference, "API Reference...", "Opens the API reference documentation", EUserInterfaceActionType::Button, FInputChord() );
+	UI_COMMAND( BrowseViewportControls, "Viewport Controls...", "Opens the viewport controls cheat sheet", EUserInterfaceActionType::Button, FInputChord() );
+	UI_COMMAND( NewLevel, "New Level...", "Create a new level, or choose a level template to start from.", EUserInterfaceActionType::Button, FInputChord( EModifierKey::Control, EKeys::N ) );
+	UI_COMMAND( OpenLevel, "Open Level...", "Loads an existing level", EUserInterfaceActionType::Button, FInputChord( EModifierKey::Control, EKeys::O ) );
+	UI_COMMAND( Save, "Save", "Saves the current level to disk", EUserInterfaceActionType::Button, FInputChord() );
+	UI_COMMAND( SaveAs, "Save As...", "Save the current level as...", EUserInterfaceActionType::Button, FInputChord( EModifierKey::Shift|EModifierKey::Control, EKeys::S ) );
+	UI_COMMAND( SaveAllLevels, "Save All Levels", "Saves all unsaved levels to disk", EUserInterfaceActionType::Button, FInputChord() );
+	UI_COMMAND( ToggleFavorite, "Toggle Favorite", "Sets whether the currently loaded level will appear in the list of favorite levels", EUserInterfaceActionType::Button, FInputChord() );
 
 	for( int32 CurRecentIndex = 0; CurRecentIndex < FLevelEditorCommands::MaxRecentFiles; ++CurRecentIndex )
 	{
@@ -2813,251 +2813,251 @@ void FLevelEditorCommands::RegisterCommands()
 				FText::Format( NSLOCTEXT( "LevelEditorCommands", "OpenRecentFile", "Open Recent File {0}" ), FText::AsNumber( CurRecentIndex ) ),
 				NSLOCTEXT( "LevelEditorCommands", "OpenRecentFileToolTip", "Opens a recently opened file" ) )
 			.UserInterfaceType( EUserInterfaceActionType::Button )
-			.DefaultGesture( FInputGesture() );
+			.DefaultChord( FInputChord() );
 		OpenRecentFileCommands.Add( OpenRecentFile );
 	}
 
-	UI_COMMAND( Import, "Import...", "Imports objects and actors from a T3D format into the current level", EUserInterfaceActionType::Button, FInputGesture() );
-	UI_COMMAND( ExportAll, "Export All...", "Exports the entire level to a file on disk (multiple formats are supported.)", EUserInterfaceActionType::Button, FInputGesture() );
-	UI_COMMAND( ExportSelected, "Export Selected...", "Exports currently-selected objects to a file on disk (multiple formats are supported.)", EUserInterfaceActionType::Button, FInputGesture() );
-
-	UI_COMMAND( Build, "Build All Levels", "Builds all levels (precomputes lighting data and visibility data, generates navigation networks and updates brush models.)", EUserInterfaceActionType::Button, FInputGesture() );
-	UI_COMMAND( BuildAndSubmitToSourceControl, "Build and Submit...", "Displays a window that allows you to build all levels and submit them to source control", EUserInterfaceActionType::Button, FInputGesture() );
-	UI_COMMAND( BuildLightingOnly, "Build Lighting", "Only precomputes lighting (all levels.)", EUserInterfaceActionType::Button, FInputGesture(EModifierKey::Control|EModifierKey::Shift, EKeys::Semicolon) );
-	UI_COMMAND( BuildReflectionCapturesOnly, "Update Reflection Captures", "Only updates Reflection Captures (all levels.)", EUserInterfaceActionType::Button, FInputGesture() );
-	UI_COMMAND( BuildLightingOnly_VisibilityOnly, "Precompute Static Visibility", "Only precomputes static visibility data (all levels.)", EUserInterfaceActionType::Button, FInputGesture() );
-	UI_COMMAND( LightingBuildOptions_UseErrorColoring, "Use Error Coloring", "When enabled, errors during lighting precomputation will be baked as colors into light map data", EUserInterfaceActionType::ToggleButton, FInputGesture() );
-	UI_COMMAND( LightingBuildOptions_ShowLightingStats, "Show Lighting Stats", "When enabled, a window containing metrics about lighting performance and memory will be displayed after a successful build.", EUserInterfaceActionType::ToggleButton, FInputGesture() );
-	UI_COMMAND( BuildGeometryOnly, "Build Geometry", "Only builds geometry (all levels.)", EUserInterfaceActionType::Button, FInputGesture() );
-	UI_COMMAND( BuildGeometryOnly_OnlyCurrentLevel, "Build Geometry (Current Level)", "Builds geometry, only for the current level", EUserInterfaceActionType::Button, FInputGesture() );
-	UI_COMMAND( BuildPathsOnly, "Build Paths", "Only builds paths (all levels.)", EUserInterfaceActionType::Button, FInputGesture() );
-	UI_COMMAND( BuildLODsOnly, "Build LODs", "Only builds LODs (all levels.)", EUserInterfaceActionType::Button, FInputGesture() );
-	UI_COMMAND( LightingQuality_Production, "Production", "Sets precomputed lighting quality to highest possible quality (slowest computation time.)", EUserInterfaceActionType::RadioButton, FInputGesture() );
-	UI_COMMAND( LightingQuality_High, "High", "Sets precomputed lighting quality to high quality", EUserInterfaceActionType::RadioButton, FInputGesture() );
-	UI_COMMAND( LightingQuality_Medium, "Medium", "Sets precomputed lighting quality to medium quality", EUserInterfaceActionType::RadioButton, FInputGesture() );
-	UI_COMMAND( LightingQuality_Preview, "Preview", "Sets precomputed lighting quality to preview quality (fastest computation time.)", EUserInterfaceActionType::RadioButton, FInputGesture() );
-	UI_COMMAND( LightingTools_ShowBounds, "Show Lighting Bounds", "Draws a wireframe sphere for the lighting bounds of the Light Environment, and a point for the center, or multiple points if NumVolumeVisibilitySamples is greater than zero.", EUserInterfaceActionType::ToggleButton, FInputGesture() );
-	UI_COMMAND( LightingTools_ShowTraces, "Show Shadow Traces", "Draws lines for any shadow rays that the Light Environment traced. The lines will be white if the light is not shadowed, and red otherwise.", EUserInterfaceActionType::ToggleButton, FInputGesture() );
-	UI_COMMAND( LightingTools_ShowDirectOnly, "Show Direct Lighting", "Renders only the influence of direct lighting on the Light Environments.", EUserInterfaceActionType::ToggleButton, FInputGesture() );
-	UI_COMMAND( LightingTools_ShowIndirectOnly, "Show Indirect Lighting", "Renders only the influence of indirect lighting on the Light Environments.", EUserInterfaceActionType::ToggleButton, FInputGesture() );
-	UI_COMMAND( LightingTools_ShowIndirectSamples, "Show Indirect Lighting Samples", "Draws the indirect lighting samples that contributed to the indirect lighting for the Light Environments.", EUserInterfaceActionType::ToggleButton, FInputGesture() );
-	UI_COMMAND( LightingDensity_RenderGrayscale, "Render Grayscale", "Renders the lightmap density.", EUserInterfaceActionType::ToggleButton, FInputGesture() );
-	UI_COMMAND( LightingResolution_CurrentLevel, "Current Level", "Adjust only primitives in the current level.", EUserInterfaceActionType::RadioButton, FInputGesture() );
-	UI_COMMAND( LightingResolution_SelectedLevels, "Selected Levels", "Adjust only primitives in the selected levels.", EUserInterfaceActionType::RadioButton, FInputGesture() );
-	UI_COMMAND( LightingResolution_AllLoadedLevels, "All Loaded Levels", "Adjust primitives in all loaded levels.", EUserInterfaceActionType::RadioButton, FInputGesture() );
-	UI_COMMAND( LightingResolution_SelectedObjectsOnly, "Selected Objects Only", "Adjust only selected objects in the levels.", EUserInterfaceActionType::ToggleButton, FInputGesture() );
-	UI_COMMAND( LightingStaticMeshInfo, "Lighting StaticMesh Info...", "Shows the lighting information for the StaticMeshes.", EUserInterfaceActionType::Button, FInputGesture() );
-	UI_COMMAND( SceneStats, "Open Scene Stats", "Opens the Scene Stats viewer", EUserInterfaceActionType::Button, FInputGesture() );
-	UI_COMMAND( TextureStats, "Open Texture Stats", "Opens the Texture Stats viewer", EUserInterfaceActionType::Button, FInputGesture() );
-	UI_COMMAND( MapCheck, "Open Map Check", "Checks map for errors", EUserInterfaceActionType::Button, FInputGesture() );
-
-	UI_COMMAND( RecompileGameCode, "Recompile Game Code", "Recompiles and reloads C++ code for game systems on the fly", EUserInterfaceActionType::Button, FInputGesture( EKeys::P, EModifierKey::Alt | EModifierKey::Control | EModifierKey::Shift ) );
-
-	UI_COMMAND( EditAsset, "Edit Asset", "Edits the asset associated with the selected actor", EUserInterfaceActionType::Button, FInputGesture( EKeys::E, EModifierKey::Control ) );
-	UI_COMMAND( EditAssetNoConfirmMultiple, "Edit Asset", "Edits the asset associated with the selected actor", EUserInterfaceActionType::Button, FInputGesture( EKeys::E, EModifierKey::Control | EModifierKey::Shift ) );
-
-	UI_COMMAND( GoHere, "Go Here", "Moves the camera to the current mouse position", EUserInterfaceActionType::Button, FInputGesture() );
-
-	UI_COMMAND( SnapCameraToObject, "Snap View to Object", "Snaps the view to the selected object", EUserInterfaceActionType::Button, FInputGesture() );
-
-	UI_COMMAND( GoToCodeForActor, "Go to C++ Code for Actor", "Opens a code editing IDE and navigates to the source file associated with the seleced actor", EUserInterfaceActionType::Button, FInputGesture() );
-	UI_COMMAND( GoToDocsForActor, "Go to Documentation for Actor", "Opens documentation for the Actor in the default web browser", EUserInterfaceActionType::Button, FInputGesture() );
-
-	UI_COMMAND( PasteHere, "Paste Here", "Pastes the actor at the click location", EUserInterfaceActionType::Button, FInputGesture() );
-
-	UI_COMMAND( SnapOriginToGrid, "Snap Origin to Grid", "Snaps the actor to the nearest grid location at its origin", EUserInterfaceActionType::Button, FInputGesture( EModifierKey::Control, EKeys::End ) );
-	UI_COMMAND( SnapOriginToGridPerActor, "Snap Origin to Grid Per Actor", "Snaps each selected actor separately to the nearest grid location at its origin", EUserInterfaceActionType::Button, FInputGesture() );
-	UI_COMMAND( AlignOriginToGrid, "Align Origin to Grid", "Aligns the actor to the nearest grid location at its origin", EUserInterfaceActionType::Button, FInputGesture() );
-	UI_COMMAND( SnapToFloor, "Snap to Floor", "Snaps the actor or component to the floor below it", EUserInterfaceActionType::Button, FInputGesture( EKeys::End ) );
-	UI_COMMAND( AlignToFloor, "Align to Floor", "Aligns the actor or component with the floor", EUserInterfaceActionType::Button, FInputGesture() );
-	UI_COMMAND( SnapPivotToFloor, "Snap Pivot to Floor", "Snaps the actor to the floor at its pivot point", EUserInterfaceActionType::Button, FInputGesture( EModifierKey::Alt, EKeys::End ) );	
-	UI_COMMAND( AlignPivotToFloor, "Align Pivot to Floor", "Aligns the actor with the floor at its pivot point", EUserInterfaceActionType::Button, FInputGesture() );
-	UI_COMMAND( SnapBottomCenterBoundsToFloor, "Snap Bottom Center Bounds to Floor", "Snaps the actor to the floor at its bottom center bounds", EUserInterfaceActionType::Button, FInputGesture( EModifierKey::Shift, EKeys::End ) );
-	UI_COMMAND( AlignBottomCenterBoundsToFloor, "Align Bottom Center Bounds to Floor", "Aligns the actor with the floor at its bottom center bounds", EUserInterfaceActionType::Button, FInputGesture() );
-	UI_COMMAND( SnapOriginToActor, "Snap Origin to Actor", "SNaps the actor to another actor at its origin", EUserInterfaceActionType::Button, FInputGesture() );
-	UI_COMMAND( AlignOriginToActor, "Align Origin to Actor", "Aligns the actor to another actor at its origin", EUserInterfaceActionType::Button, FInputGesture() );
-	UI_COMMAND( SnapToActor, "Snap to Actor", "Snaps the actor to another actor", EUserInterfaceActionType::Button, FInputGesture() );
-	UI_COMMAND( AlignToActor, "Align to Actor", "Aligns the actor with another actor", EUserInterfaceActionType::Button, FInputGesture() );
-	UI_COMMAND( SnapPivotToActor, "Snap Pivot to Actor", "Snaps the actor to another actor at its pivot point", EUserInterfaceActionType::Button, FInputGesture() );
-	UI_COMMAND( AlignPivotToActor, "Align Pivot to Actor", "Aligns the actor with another actor at its pivot point", EUserInterfaceActionType::Button, FInputGesture() );
-	UI_COMMAND( SnapBottomCenterBoundsToActor, "Snap Bottom Center Bounds to Actor", "Snaps the actor to another actor at its bottom center bounds", EUserInterfaceActionType::Button, FInputGesture() );
-	UI_COMMAND( AlignBottomCenterBoundsToActor, "Align Bottom Center Bounds to Actor", "Aligns the actor with another actor at its bottom center bounds", EUserInterfaceActionType::Button, FInputGesture() );
-	UI_COMMAND( DeltaTransformToActors, "Delta Transform", "Apply Delta Transform to selected actors", EUserInterfaceActionType::Button, FInputGesture() );
-	UI_COMMAND( MirrorActorX, "Mirror X", "Mirrors the actor along the X axis", EUserInterfaceActionType::Button, FInputGesture() );
-	UI_COMMAND( MirrorActorY, "Mirror Y", "Mirrors the actor along the Y axis", EUserInterfaceActionType::Button, FInputGesture() );
-	UI_COMMAND( MirrorActorZ, "Mirror Z", "Mirrors the actor along the Z axis", EUserInterfaceActionType::Button, FInputGesture() );
-	UI_COMMAND( LockActorMovement, "Lock Actor Movement", "Locks the actor so it cannot be moved", EUserInterfaceActionType::ToggleButton, FInputGesture() );
-	UI_COMMAND( DetachFromParent, "Detach", "Detach the actor from its parent", EUserInterfaceActionType::Button, FInputGesture() );
-	UI_COMMAND( AttachSelectedActors, "Attach Selected Actors", "Attach the selected actors to the last selected actor", EUserInterfaceActionType::Button, FInputGesture(EModifierKey::Alt, EKeys::B) );
-	UI_COMMAND( AttachActorIteractive, "Attach Actor Interactive", "Start an interactive actor picker to let you choose a parent for the currently selected actor", EUserInterfaceActionType::Button, FInputGesture(EModifierKey::Alt, EKeys::A) );
-	UI_COMMAND( CreateNewOutlinerFolder, "Create Folder", "Place the selected actors in a new folder", EUserInterfaceActionType::Button, FInputGesture() );
-	UI_COMMAND( HoldToEnableVertexSnapping, "Hold to Enable Vertex Snapping", "When the key binding is pressed and held vertex snapping will be enabled", EUserInterfaceActionType::ToggleButton, FInputGesture(EKeys::V) );
+	UI_COMMAND( Import, "Import...", "Imports objects and actors from a T3D format into the current level", EUserInterfaceActionType::Button, FInputChord() );
+	UI_COMMAND( ExportAll, "Export All...", "Exports the entire level to a file on disk (multiple formats are supported.)", EUserInterfaceActionType::Button, FInputChord() );
+	UI_COMMAND( ExportSelected, "Export Selected...", "Exports currently-selected objects to a file on disk (multiple formats are supported.)", EUserInterfaceActionType::Button, FInputChord() );
+
+	UI_COMMAND( Build, "Build All Levels", "Builds all levels (precomputes lighting data and visibility data, generates navigation networks and updates brush models.)", EUserInterfaceActionType::Button, FInputChord() );
+	UI_COMMAND( BuildAndSubmitToSourceControl, "Build and Submit...", "Displays a window that allows you to build all levels and submit them to source control", EUserInterfaceActionType::Button, FInputChord() );
+	UI_COMMAND( BuildLightingOnly, "Build Lighting", "Only precomputes lighting (all levels.)", EUserInterfaceActionType::Button, FInputChord(EModifierKey::Control|EModifierKey::Shift, EKeys::Semicolon) );
+	UI_COMMAND( BuildReflectionCapturesOnly, "Update Reflection Captures", "Only updates Reflection Captures (all levels.)", EUserInterfaceActionType::Button, FInputChord() );
+	UI_COMMAND( BuildLightingOnly_VisibilityOnly, "Precompute Static Visibility", "Only precomputes static visibility data (all levels.)", EUserInterfaceActionType::Button, FInputChord() );
+	UI_COMMAND( LightingBuildOptions_UseErrorColoring, "Use Error Coloring", "When enabled, errors during lighting precomputation will be baked as colors into light map data", EUserInterfaceActionType::ToggleButton, FInputChord() );
+	UI_COMMAND( LightingBuildOptions_ShowLightingStats, "Show Lighting Stats", "When enabled, a window containing metrics about lighting performance and memory will be displayed after a successful build.", EUserInterfaceActionType::ToggleButton, FInputChord() );
+	UI_COMMAND( BuildGeometryOnly, "Build Geometry", "Only builds geometry (all levels.)", EUserInterfaceActionType::Button, FInputChord() );
+	UI_COMMAND( BuildGeometryOnly_OnlyCurrentLevel, "Build Geometry (Current Level)", "Builds geometry, only for the current level", EUserInterfaceActionType::Button, FInputChord() );
+	UI_COMMAND( BuildPathsOnly, "Build Paths", "Only builds paths (all levels.)", EUserInterfaceActionType::Button, FInputChord() );
+	UI_COMMAND( BuildLODsOnly, "Build LODs", "Only builds LODs (all levels.)", EUserInterfaceActionType::Button, FInputChord() );
+	UI_COMMAND( LightingQuality_Production, "Production", "Sets precomputed lighting quality to highest possible quality (slowest computation time.)", EUserInterfaceActionType::RadioButton, FInputChord() );
+	UI_COMMAND( LightingQuality_High, "High", "Sets precomputed lighting quality to high quality", EUserInterfaceActionType::RadioButton, FInputChord() );
+	UI_COMMAND( LightingQuality_Medium, "Medium", "Sets precomputed lighting quality to medium quality", EUserInterfaceActionType::RadioButton, FInputChord() );
+	UI_COMMAND( LightingQuality_Preview, "Preview", "Sets precomputed lighting quality to preview quality (fastest computation time.)", EUserInterfaceActionType::RadioButton, FInputChord() );
+	UI_COMMAND( LightingTools_ShowBounds, "Show Lighting Bounds", "Draws a wireframe sphere for the lighting bounds of the Light Environment, and a point for the center, or multiple points if NumVolumeVisibilitySamples is greater than zero.", EUserInterfaceActionType::ToggleButton, FInputChord() );
+	UI_COMMAND( LightingTools_ShowTraces, "Show Shadow Traces", "Draws lines for any shadow rays that the Light Environment traced. The lines will be white if the light is not shadowed, and red otherwise.", EUserInterfaceActionType::ToggleButton, FInputChord() );
+	UI_COMMAND( LightingTools_ShowDirectOnly, "Show Direct Lighting", "Renders only the influence of direct lighting on the Light Environments.", EUserInterfaceActionType::ToggleButton, FInputChord() );
+	UI_COMMAND( LightingTools_ShowIndirectOnly, "Show Indirect Lighting", "Renders only the influence of indirect lighting on the Light Environments.", EUserInterfaceActionType::ToggleButton, FInputChord() );
+	UI_COMMAND( LightingTools_ShowIndirectSamples, "Show Indirect Lighting Samples", "Draws the indirect lighting samples that contributed to the indirect lighting for the Light Environments.", EUserInterfaceActionType::ToggleButton, FInputChord() );
+	UI_COMMAND( LightingDensity_RenderGrayscale, "Render Grayscale", "Renders the lightmap density.", EUserInterfaceActionType::ToggleButton, FInputChord() );
+	UI_COMMAND( LightingResolution_CurrentLevel, "Current Level", "Adjust only primitives in the current level.", EUserInterfaceActionType::RadioButton, FInputChord() );
+	UI_COMMAND( LightingResolution_SelectedLevels, "Selected Levels", "Adjust only primitives in the selected levels.", EUserInterfaceActionType::RadioButton, FInputChord() );
+	UI_COMMAND( LightingResolution_AllLoadedLevels, "All Loaded Levels", "Adjust primitives in all loaded levels.", EUserInterfaceActionType::RadioButton, FInputChord() );
+	UI_COMMAND( LightingResolution_SelectedObjectsOnly, "Selected Objects Only", "Adjust only selected objects in the levels.", EUserInterfaceActionType::ToggleButton, FInputChord() );
+	UI_COMMAND( LightingStaticMeshInfo, "Lighting StaticMesh Info...", "Shows the lighting information for the StaticMeshes.", EUserInterfaceActionType::Button, FInputChord() );
+	UI_COMMAND( SceneStats, "Open Scene Stats", "Opens the Scene Stats viewer", EUserInterfaceActionType::Button, FInputChord() );
+	UI_COMMAND( TextureStats, "Open Texture Stats", "Opens the Texture Stats viewer", EUserInterfaceActionType::Button, FInputChord() );
+	UI_COMMAND( MapCheck, "Open Map Check", "Checks map for errors", EUserInterfaceActionType::Button, FInputChord() );
+
+	UI_COMMAND( RecompileGameCode, "Recompile Game Code", "Recompiles and reloads C++ code for game systems on the fly", EUserInterfaceActionType::Button, FInputChord( EKeys::P, EModifierKey::Alt | EModifierKey::Control | EModifierKey::Shift ) );
+
+	UI_COMMAND( EditAsset, "Edit Asset", "Edits the asset associated with the selected actor", EUserInterfaceActionType::Button, FInputChord( EKeys::E, EModifierKey::Control ) );
+	UI_COMMAND( EditAssetNoConfirmMultiple, "Edit Asset", "Edits the asset associated with the selected actor", EUserInterfaceActionType::Button, FInputChord( EKeys::E, EModifierKey::Control | EModifierKey::Shift ) );
+
+	UI_COMMAND( GoHere, "Go Here", "Moves the camera to the current mouse position", EUserInterfaceActionType::Button, FInputChord() );
+
+	UI_COMMAND( SnapCameraToObject, "Snap View to Object", "Snaps the view to the selected object", EUserInterfaceActionType::Button, FInputChord() );
+
+	UI_COMMAND( GoToCodeForActor, "Go to C++ Code for Actor", "Opens a code editing IDE and navigates to the source file associated with the seleced actor", EUserInterfaceActionType::Button, FInputChord() );
+	UI_COMMAND( GoToDocsForActor, "Go to Documentation for Actor", "Opens documentation for the Actor in the default web browser", EUserInterfaceActionType::Button, FInputChord() );
+
+	UI_COMMAND( PasteHere, "Paste Here", "Pastes the actor at the click location", EUserInterfaceActionType::Button, FInputChord() );
+
+	UI_COMMAND( SnapOriginToGrid, "Snap Origin to Grid", "Snaps the actor to the nearest grid location at its origin", EUserInterfaceActionType::Button, FInputChord( EModifierKey::Control, EKeys::End ) );
+	UI_COMMAND( SnapOriginToGridPerActor, "Snap Origin to Grid Per Actor", "Snaps each selected actor separately to the nearest grid location at its origin", EUserInterfaceActionType::Button, FInputChord() );
+	UI_COMMAND( AlignOriginToGrid, "Align Origin to Grid", "Aligns the actor to the nearest grid location at its origin", EUserInterfaceActionType::Button, FInputChord() );
+	UI_COMMAND( SnapToFloor, "Snap to Floor", "Snaps the actor or component to the floor below it", EUserInterfaceActionType::Button, FInputChord( EKeys::End ) );
+	UI_COMMAND( AlignToFloor, "Align to Floor", "Aligns the actor or component with the floor", EUserInterfaceActionType::Button, FInputChord() );
+	UI_COMMAND( SnapPivotToFloor, "Snap Pivot to Floor", "Snaps the actor to the floor at its pivot point", EUserInterfaceActionType::Button, FInputChord( EModifierKey::Alt, EKeys::End ) );	
+	UI_COMMAND( AlignPivotToFloor, "Align Pivot to Floor", "Aligns the actor with the floor at its pivot point", EUserInterfaceActionType::Button, FInputChord() );
+	UI_COMMAND( SnapBottomCenterBoundsToFloor, "Snap Bottom Center Bounds to Floor", "Snaps the actor to the floor at its bottom center bounds", EUserInterfaceActionType::Button, FInputChord( EModifierKey::Shift, EKeys::End ) );
+	UI_COMMAND( AlignBottomCenterBoundsToFloor, "Align Bottom Center Bounds to Floor", "Aligns the actor with the floor at its bottom center bounds", EUserInterfaceActionType::Button, FInputChord() );
+	UI_COMMAND( SnapOriginToActor, "Snap Origin to Actor", "SNaps the actor to another actor at its origin", EUserInterfaceActionType::Button, FInputChord() );
+	UI_COMMAND( AlignOriginToActor, "Align Origin to Actor", "Aligns the actor to another actor at its origin", EUserInterfaceActionType::Button, FInputChord() );
+	UI_COMMAND( SnapToActor, "Snap to Actor", "Snaps the actor to another actor", EUserInterfaceActionType::Button, FInputChord() );
+	UI_COMMAND( AlignToActor, "Align to Actor", "Aligns the actor with another actor", EUserInterfaceActionType::Button, FInputChord() );
+	UI_COMMAND( SnapPivotToActor, "Snap Pivot to Actor", "Snaps the actor to another actor at its pivot point", EUserInterfaceActionType::Button, FInputChord() );
+	UI_COMMAND( AlignPivotToActor, "Align Pivot to Actor", "Aligns the actor with another actor at its pivot point", EUserInterfaceActionType::Button, FInputChord() );
+	UI_COMMAND( SnapBottomCenterBoundsToActor, "Snap Bottom Center Bounds to Actor", "Snaps the actor to another actor at its bottom center bounds", EUserInterfaceActionType::Button, FInputChord() );
+	UI_COMMAND( AlignBottomCenterBoundsToActor, "Align Bottom Center Bounds to Actor", "Aligns the actor with another actor at its bottom center bounds", EUserInterfaceActionType::Button, FInputChord() );
+	UI_COMMAND( DeltaTransformToActors, "Delta Transform", "Apply Delta Transform to selected actors", EUserInterfaceActionType::Button, FInputChord() );
+	UI_COMMAND( MirrorActorX, "Mirror X", "Mirrors the actor along the X axis", EUserInterfaceActionType::Button, FInputChord() );
+	UI_COMMAND( MirrorActorY, "Mirror Y", "Mirrors the actor along the Y axis", EUserInterfaceActionType::Button, FInputChord() );
+	UI_COMMAND( MirrorActorZ, "Mirror Z", "Mirrors the actor along the Z axis", EUserInterfaceActionType::Button, FInputChord() );
+	UI_COMMAND( LockActorMovement, "Lock Actor Movement", "Locks the actor so it cannot be moved", EUserInterfaceActionType::ToggleButton, FInputChord() );
+	UI_COMMAND( DetachFromParent, "Detach", "Detach the actor from its parent", EUserInterfaceActionType::Button, FInputChord() );
+	UI_COMMAND( AttachSelectedActors, "Attach Selected Actors", "Attach the selected actors to the last selected actor", EUserInterfaceActionType::Button, FInputChord(EModifierKey::Alt, EKeys::B) );
+	UI_COMMAND( AttachActorIteractive, "Attach Actor Interactive", "Start an interactive actor picker to let you choose a parent for the currently selected actor", EUserInterfaceActionType::Button, FInputChord(EModifierKey::Alt, EKeys::A) );
+	UI_COMMAND( CreateNewOutlinerFolder, "Create Folder", "Place the selected actors in a new folder", EUserInterfaceActionType::Button, FInputChord() );
+	UI_COMMAND( HoldToEnableVertexSnapping, "Hold to Enable Vertex Snapping", "When the key binding is pressed and held vertex snapping will be enabled", EUserInterfaceActionType::ToggleButton, FInputChord(EKeys::V) );
 
 	//@ todo Slate better tooltips for pivot options
-	UI_COMMAND( SavePivotToPrePivot, "Save Pivot to Pre-Pivot", "Saves the pivot to the pre-pivot", EUserInterfaceActionType::Button, FInputGesture() );
-	UI_COMMAND( ResetPivot, "Reset Pivot", "Resets the pivot", EUserInterfaceActionType::Button, FInputGesture() );
-	UI_COMMAND( ResetPrePivot, "Reset Pre-Pivot", "Resets the pre-pivot", EUserInterfaceActionType::Button, FInputGesture() );
-	UI_COMMAND( MovePivotHere, "Move Here", "Moves the pivot to the clicked location", EUserInterfaceActionType::Button, FInputGesture() );
-	UI_COMMAND( MovePivotHereSnapped, "Move Here (Snapped)", "Moves the pivot to the nearest grid point to the clicked location", EUserInterfaceActionType::Button, FInputGesture() );
-	UI_COMMAND( MovePivotToCenter, "Center on Selection", "Centers the pivot to the middle of the selection", EUserInterfaceActionType::Button, FInputGesture() );
-
-	UI_COMMAND( ConvertToAdditive, "Additive", "Converts the selected brushes to additive brushes", EUserInterfaceActionType::Button, FInputGesture() );
-	UI_COMMAND( ConvertToSubtractive, "Subtractive", "Converts the selected brushes to subtractive brushes", EUserInterfaceActionType::Button, FInputGesture() );
-
-	UI_COMMAND( OrderFirst, "To First", "Changes the drawing order of the selected brushes so they are the first to draw", EUserInterfaceActionType::Button, FInputGesture() );
-	UI_COMMAND( OrderLast, "To Last", "Changes the drawing order of the selected brushes so they are the last to draw", EUserInterfaceActionType::Button, FInputGesture() );
-
-	UI_COMMAND( MakeSolid, "Solid", "Makes the selected brushes solid", EUserInterfaceActionType::Button, FInputGesture() );
-	UI_COMMAND( MakeSemiSolid, "Semi-Solid", "Makes the selected brushes semi-solid", EUserInterfaceActionType::Button, FInputGesture() );
-	UI_COMMAND( MakeNonSolid, "Non-Solid", "Makes the selected brushes non-solid", EUserInterfaceActionType::Button, FInputGesture() );
-
-	UI_COMMAND( MergePolys, "Merge", "Merges multiple polygons on a brush face into as few as possible", EUserInterfaceActionType::Button, FInputGesture() );
-	UI_COMMAND( SeparatePolys, "Separate", "Reverses the effect of a previous merge", EUserInterfaceActionType::Button, FInputGesture() );
-
-	// RegroupActors uses GroupActors for it's label and tooltip when simply grouping a selection of actors using overrides. This is to provide display of the gesture which is the same for both.
-	UI_COMMAND( GroupActors, "Group", "Groups the selected actors", EUserInterfaceActionType::Button, FInputGesture( /*EKeys::G, EModifierKey::Control*/ ) );
-	UI_COMMAND( RegroupActors, "Regroup", "Regroups the selected actors into a new group, removing any current groups in the selection", EUserInterfaceActionType::Button, FInputGesture( EKeys::G, EModifierKey::Control ) );
-	UI_COMMAND( UngroupActors, "Ungroup", "Ungroups the selected actors", EUserInterfaceActionType::Button, FInputGesture( EModifierKey::Shift, EKeys::G ) );
-	UI_COMMAND( AddActorsToGroup, "Add to Group", "Adds the selected actors to the selected group", EUserInterfaceActionType::Button, FInputGesture() );
-	UI_COMMAND( RemoveActorsFromGroup, "Remove from Group", "Removes the selected actors from the selected groups", EUserInterfaceActionType::Button, FInputGesture() );
-	UI_COMMAND( LockGroup, "Lock", "Locks the selected groups", EUserInterfaceActionType::Button, FInputGesture() );
-	UI_COMMAND( UnlockGroup, "Unlock", "Unlocks the selected groups", EUserInterfaceActionType::Button, FInputGesture() );
-	UI_COMMAND( MergeActors, "Create Mesh Proxy...", "Harvest geometry from selected actors and merge them into single mesh", EUserInterfaceActionType::Button, FInputGesture() );
-	UI_COMMAND( MergeActorsByMaterials, "Merge Actors...", "Harvest geometry from selected actors and merge grouping them by materials", EUserInterfaceActionType::Button, FInputGesture() );
+	UI_COMMAND( SavePivotToPrePivot, "Save Pivot to Pre-Pivot", "Saves the pivot to the pre-pivot", EUserInterfaceActionType::Button, FInputChord() );
+	UI_COMMAND( ResetPivot, "Reset Pivot", "Resets the pivot", EUserInterfaceActionType::Button, FInputChord() );
+	UI_COMMAND( ResetPrePivot, "Reset Pre-Pivot", "Resets the pre-pivot", EUserInterfaceActionType::Button, FInputChord() );
+	UI_COMMAND( MovePivotHere, "Move Here", "Moves the pivot to the clicked location", EUserInterfaceActionType::Button, FInputChord() );
+	UI_COMMAND( MovePivotHereSnapped, "Move Here (Snapped)", "Moves the pivot to the nearest grid point to the clicked location", EUserInterfaceActionType::Button, FInputChord() );
+	UI_COMMAND( MovePivotToCenter, "Center on Selection", "Centers the pivot to the middle of the selection", EUserInterfaceActionType::Button, FInputChord() );
+
+	UI_COMMAND( ConvertToAdditive, "Additive", "Converts the selected brushes to additive brushes", EUserInterfaceActionType::Button, FInputChord() );
+	UI_COMMAND( ConvertToSubtractive, "Subtractive", "Converts the selected brushes to subtractive brushes", EUserInterfaceActionType::Button, FInputChord() );
+
+	UI_COMMAND( OrderFirst, "To First", "Changes the drawing order of the selected brushes so they are the first to draw", EUserInterfaceActionType::Button, FInputChord() );
+	UI_COMMAND( OrderLast, "To Last", "Changes the drawing order of the selected brushes so they are the last to draw", EUserInterfaceActionType::Button, FInputChord() );
+
+	UI_COMMAND( MakeSolid, "Solid", "Makes the selected brushes solid", EUserInterfaceActionType::Button, FInputChord() );
+	UI_COMMAND( MakeSemiSolid, "Semi-Solid", "Makes the selected brushes semi-solid", EUserInterfaceActionType::Button, FInputChord() );
+	UI_COMMAND( MakeNonSolid, "Non-Solid", "Makes the selected brushes non-solid", EUserInterfaceActionType::Button, FInputChord() );
+
+	UI_COMMAND( MergePolys, "Merge", "Merges multiple polygons on a brush face into as few as possible", EUserInterfaceActionType::Button, FInputChord() );
+	UI_COMMAND( SeparatePolys, "Separate", "Reverses the effect of a previous merge", EUserInterfaceActionType::Button, FInputChord() );
+
+	// RegroupActors uses GroupActors for it's label and tooltip when simply grouping a selection of actors using overrides. This is to provide display of the chord which is the same for both.
+	UI_COMMAND( GroupActors, "Group", "Groups the selected actors", EUserInterfaceActionType::Button, FInputChord( /*EKeys::G, EModifierKey::Control*/ ) );
+	UI_COMMAND( RegroupActors, "Regroup", "Regroups the selected actors into a new group, removing any current groups in the selection", EUserInterfaceActionType::Button, FInputChord( EKeys::G, EModifierKey::Control ) );
+	UI_COMMAND( UngroupActors, "Ungroup", "Ungroups the selected actors", EUserInterfaceActionType::Button, FInputChord( EModifierKey::Shift, EKeys::G ) );
+	UI_COMMAND( AddActorsToGroup, "Add to Group", "Adds the selected actors to the selected group", EUserInterfaceActionType::Button, FInputChord() );
+	UI_COMMAND( RemoveActorsFromGroup, "Remove from Group", "Removes the selected actors from the selected groups", EUserInterfaceActionType::Button, FInputChord() );
+	UI_COMMAND( LockGroup, "Lock", "Locks the selected groups", EUserInterfaceActionType::Button, FInputChord() );
+	UI_COMMAND( UnlockGroup, "Unlock", "Unlocks the selected groups", EUserInterfaceActionType::Button, FInputChord() );
+	UI_COMMAND( MergeActors, "Create Mesh Proxy...", "Harvest geometry from selected actors and merge them into single mesh", EUserInterfaceActionType::Button, FInputChord() );
+	UI_COMMAND( MergeActorsByMaterials, "Merge Actors...", "Harvest geometry from selected actors and merge grouping them by materials", EUserInterfaceActionType::Button, FInputChord() );
 
 #if PLATFORM_MAC
-	UI_COMMAND( ShowAll, "Show All Actors", "Shows all actors", EUserInterfaceActionType::Button, FInputGesture( EModifierKey::Command, EKeys::H ) );
+	UI_COMMAND( ShowAll, "Show All Actors", "Shows all actors", EUserInterfaceActionType::Button, FInputChord( EModifierKey::Command, EKeys::H ) );
 #else
-	UI_COMMAND( ShowAll, "Show All Actors", "Shows all actors", EUserInterfaceActionType::Button, FInputGesture( EModifierKey::Control, EKeys::H ) );
+	UI_COMMAND( ShowAll, "Show All Actors", "Shows all actors", EUserInterfaceActionType::Button, FInputChord( EModifierKey::Control, EKeys::H ) );
 #endif
-	UI_COMMAND( ShowSelectedOnly, "Show Only Selected", "Shows only the selected actors", EUserInterfaceActionType::Button, FInputGesture() );
-	UI_COMMAND( ShowSelected, "Show Selected", "Shows the selected actors", EUserInterfaceActionType::Button, FInputGesture( EModifierKey::Shift, EKeys::H ) );
-	UI_COMMAND( HideSelected, "Hide Selected", "Hides the selected actors", EUserInterfaceActionType::Button, FInputGesture( EKeys::H ) );
-	UI_COMMAND( ShowAllStartup, "Show All At Startup", "Shows all actors at startup", EUserInterfaceActionType::Button, FInputGesture() );
-	UI_COMMAND( ShowSelectedStartup, "Show Selected At Startup", "Shows selected actors at startup", EUserInterfaceActionType::Button, FInputGesture() );
-	UI_COMMAND( HideSelectedStartup, "Hide Selected At Startup", "Hide selected actors at startup", EUserInterfaceActionType::Button, FInputGesture() );
-	UI_COMMAND( CycleNavigationDataDrawn, "Cycle Navigation Data Drawn", "Cycles through navigation data (navmeshes for example) to draw one at a time", EUserInterfaceActionType::Button, FInputGesture( EModifierKey::Alt, EKeys::N ) );
-
-	UI_COMMAND( SelectNone, "Unselect All", "Unselects all actors", EUserInterfaceActionType::Button, FInputGesture( EKeys::Escape ) ) ;
-	UI_COMMAND( InvertSelection, "Invert Selection", "Inverts the current selection", EUserInterfaceActionType::Button, FInputGesture() );
-
-	UI_COMMAND( SelectAllActorsOfSameClass, "Select All Actors of Same Class", "Selects all the actors that have the same class", EUserInterfaceActionType::Button, FInputGesture(EModifierKey::Shift|EModifierKey::Control, EKeys::A) );
-	UI_COMMAND( SelectAllActorsOfSameClassWithArchetype, "Select All Actors with Same Archetype", "Selects all the actors of the same class that have the same archetype", EUserInterfaceActionType::Button, FInputGesture() );
-	UI_COMMAND( SelectComponentOwnerActor, "Select Component Owner", "Select the actor that owns the currently selected component(s)", EUserInterfaceActionType::Button, FInputGesture() );
-	UI_COMMAND( SelectRelevantLights, "Select Relevant Lights", "Select all lights relevant to the current selection", EUserInterfaceActionType::Button, FInputGesture() ); 
-	UI_COMMAND( SelectStaticMeshesOfSameClass, "Select All Using Selected Static Meshes (Selected Actor Types)", "Selects all actors with the same static mesh and actor class as the selection", EUserInterfaceActionType::Button, FInputGesture() ); 
-	UI_COMMAND( SelectStaticMeshesAllClasses, "Select All Using Selected Static Meshes (All Actor Types)", "Selects all actors with the same static mesh as the selection", EUserInterfaceActionType::Button, FInputGesture( EModifierKey::Shift, EKeys::E ) ); 
-	UI_COMMAND( SelectSkeletalMeshesOfSameClass, "Select All Using Selected Skeletal Meshes (Selected Actor Types)", "Selects all actors with the same skeletal mesh and actor class as the selection", EUserInterfaceActionType::Button, FInputGesture() ); 
-	UI_COMMAND( SelectSkeletalMeshesAllClasses, "Select All Using Selected Skeletal Meshes (All Actor Types)", "Selects all actors with the same skeletal mesh as the selection", EUserInterfaceActionType::Button, FInputGesture() ); 
-	UI_COMMAND( SelectAllWithSameMaterial, "Select All With Same Material", "Selects all actors with the same material as the selection", EUserInterfaceActionType::Button, FInputGesture() ); 
-	UI_COMMAND( SelectMatchingEmitter, "Select All Matching Emitters", "Selects all emitters with the same particle system as the selection", EUserInterfaceActionType::Button, FInputGesture() ); 
-	UI_COMMAND( SelectAllLights, "Select All Lights", "Selects all lights", EUserInterfaceActionType::Button, FInputGesture() ); 
-	UI_COMMAND( SelectStationaryLightsExceedingOverlap, "Select Stationary Lights exceeding overlap", "Selects all stationary lights exceeding the overlap limit", EUserInterfaceActionType::Button, FInputGesture() ); 
-	UI_COMMAND( SelectAllAddditiveBrushes, "Select All Additive Brushes", "Selects all additive brushes", EUserInterfaceActionType::Button, FInputGesture() ); 
-	UI_COMMAND( SelectAllSubtractiveBrushes, "Select All Subtractive Brushes", "Selects all subtractive brushes", EUserInterfaceActionType::Button, FInputGesture() ); 
-	UI_COMMAND( SelectAllActorsControlledByMatinee, "Select Actors Used by This Matinee", "Selects all actors controlled by this Matinee", EUserInterfaceActionType::Button, FInputGesture() ); 
-
-	UI_COMMAND( SelectAllSurfaces, "Select All Surfaces", "Selects all bsp surfaces", EUserInterfaceActionType::Button, FInputGesture(EModifierKey::Shift, EKeys::S) ); 
-
-	UI_COMMAND( SurfSelectAllMatchingBrush, "Select Matching Brush", "Selects the surfaces belonging to the same brush as the selected surfaces", EUserInterfaceActionType::Button, FInputGesture(EModifierKey::Shift, EKeys::B) ); 
-	UI_COMMAND( SurfSelectAllMatchingTexture, "Select Matching Material", "Selects all surfaces with the same material as the selected surfaces", EUserInterfaceActionType::Button, FInputGesture(EModifierKey::Shift, EKeys::T) ); 
-	UI_COMMAND( SurfSelectAllAdjacents, "Select All Adjacent Surfaces", "Selects all surfaces adjacent to the currently selected surfaces", EUserInterfaceActionType::Button, FInputGesture(EModifierKey::Shift, EKeys::J) ); 
-	UI_COMMAND( SurfSelectAllAdjacentCoplanars, "Select All Coplanar Surfaces", "Selects all surfaces adjacent and coplanar with the selected surfaces", EUserInterfaceActionType::Button, FInputGesture(EModifierKey::Shift, EKeys::C) ); 
-	UI_COMMAND( SurfSelectAllAdjacentWalls, "Select All Adjacent Wall Surfaces", "Selects all adjacent upright surfaces", EUserInterfaceActionType::Button, FInputGesture(EModifierKey::Shift, EKeys::W) ); 
-	UI_COMMAND( SurfSelectAllAdjacentFloors, "Select All Adjacent Floor Surfaces", "Selects all adjacent floor sufaces(ones with normals pointing up)", EUserInterfaceActionType::Button, FInputGesture(EModifierKey::Shift, EKeys::U) ); 
-	UI_COMMAND( SurfSelectAllAdjacentSlants, "Select All Adjacent Slant Surfaces", "Selects all adjacent slant surfaces (surfaces that are not walls, floors, or ceilings according to their normals) to the currently selected surfaces.", EUserInterfaceActionType::Button, FInputGesture(EModifierKey::Shift, EKeys::Y) ); 
-	UI_COMMAND( SurfSelectReverse, "Invert Surface Selection", "Inverts the current surface selection", EUserInterfaceActionType::Button, FInputGesture(EModifierKey::Shift, EKeys::Q) );
-	UI_COMMAND( SurfSelectMemorize, "Memorize Surface Selection", "Stores the current surface selection in memory", EUserInterfaceActionType::Button, FInputGesture(EModifierKey::Shift, EKeys::M) ); 
-	UI_COMMAND( SurfSelectRecall, "Recall Surface Selection", "Replace the current selection with the selection saved in memory", EUserInterfaceActionType::Button, FInputGesture(EModifierKey::Shift, EKeys::R) ); 
-	UI_COMMAND( SurfSelectOr, "Surface Selection OR", "Replace the current selection with only the surfaces which are both currently selected and contained within the saved selection in memory", EUserInterfaceActionType::Button, FInputGesture(EModifierKey::Shift, EKeys::O) ); 
-	UI_COMMAND( SurfSelectAnd, "Surface Selection AND", "Add the selection of surfaces saved in memory to the current selection", EUserInterfaceActionType::Button, FInputGesture(EModifierKey::Shift, EKeys::A) ); 
-	UI_COMMAND( SurfSelectXor, "Surace Selection XOR", " Replace the current selection with only the surfaces that are not in both the current selection and the selection saved in memory", EUserInterfaceActionType::Button, FInputGesture(EModifierKey::Shift, EKeys::X) ); 
-	UI_COMMAND( SurfUnalign, "Align Surface Default", "Default surface alignment", EUserInterfaceActionType::Button, FInputGesture() );
-	UI_COMMAND( SurfAlignPlanarAuto, "Align Surface Planar", "Planar surface alignment", EUserInterfaceActionType::Button, FInputGesture() );
-	UI_COMMAND( SurfAlignPlanarWall, "Align Surface Planar Wall", "Planar wall surface alignment", EUserInterfaceActionType::Button, FInputGesture() );
-	UI_COMMAND( SurfAlignPlanarFloor, "Align Surface Planar Floor", "Planar floor surface alignment", EUserInterfaceActionType::Button, FInputGesture() );
-	UI_COMMAND( SurfAlignBox, "Align Surface Box", "Box surface alignment", EUserInterfaceActionType::Button, FInputGesture() );
-	UI_COMMAND( SurfAlignFit, "Align Surface Fit", "Best fit surface alignment", EUserInterfaceActionType::Button, FInputGesture() );
-	UI_COMMAND( ApplyMaterialToSurface, "Apply Material to Surface Selection", "Applies the selected material to the selected surfaces", EUserInterfaceActionType::Button, FInputGesture() );
-
-	UI_COMMAND( CreateBoundingBoxVolume, "Create Bounding Box Blocking Volume From Mesh", "Create a bounding box blocking volume from the static mesh", EUserInterfaceActionType::Button, FInputGesture() );
-	UI_COMMAND( CreateHeavyConvexVolume, "Heavy Convex Blocking Volume From Mesh", "Creates a heavy convex blocing volume from the static mesh", EUserInterfaceActionType::Button, FInputGesture() );
-	UI_COMMAND( CreateNormalConvexVolume, "Normal Convex Blocking Volume From Mesh", "Creates a normal convex blocking volume from the static mesh", EUserInterfaceActionType::Button, FInputGesture() ); 
-	UI_COMMAND( CreateLightConvexVolume, "Light Convex Blocking Volume From Mesh", "Creates a light convex blocking volume from the static mesh", EUserInterfaceActionType::Button, FInputGesture() );
-	UI_COMMAND( CreateRoughConvexVolume, "Rought Convex Blocking Volume From Mesh", "Creates a rough convex blocking volume from the static mesh", EUserInterfaceActionType::Button, FInputGesture() );
-
-	UI_COMMAND( SaveBrushAsCollision, "Save Collision from Builder Brush", "Creates a collision primitive on the selected static meshes based on the shape of the builder brush", EUserInterfaceActionType::Button, FInputGesture() );
-
-	UI_COMMAND( KeepSimulationChanges, "Keep Simulation Changes", "Saves the changes made to this actor in Simulate mode to the actor's default state.", EUserInterfaceActionType::Button, FInputGesture( EKeys::K ) );
-
-	UI_COMMAND( MakeActorLevelCurrent, "Make Selected Actor's Level Current", "Makes the selected actors level the current level", EUserInterfaceActionType::Button, FInputGesture( EKeys::M ) );
+	UI_COMMAND( ShowSelectedOnly, "Show Only Selected", "Shows only the selected actors", EUserInterfaceActionType::Button, FInputChord() );
+	UI_COMMAND( ShowSelected, "Show Selected", "Shows the selected actors", EUserInterfaceActionType::Button, FInputChord( EModifierKey::Shift, EKeys::H ) );
+	UI_COMMAND( HideSelected, "Hide Selected", "Hides the selected actors", EUserInterfaceActionType::Button, FInputChord( EKeys::H ) );
+	UI_COMMAND( ShowAllStartup, "Show All At Startup", "Shows all actors at startup", EUserInterfaceActionType::Button, FInputChord() );
+	UI_COMMAND( ShowSelectedStartup, "Show Selected At Startup", "Shows selected actors at startup", EUserInterfaceActionType::Button, FInputChord() );
+	UI_COMMAND( HideSelectedStartup, "Hide Selected At Startup", "Hide selected actors at startup", EUserInterfaceActionType::Button, FInputChord() );
+	UI_COMMAND( CycleNavigationDataDrawn, "Cycle Navigation Data Drawn", "Cycles through navigation data (navmeshes for example) to draw one at a time", EUserInterfaceActionType::Button, FInputChord( EModifierKey::Alt, EKeys::N ) );
+
+	UI_COMMAND( SelectNone, "Unselect All", "Unselects all actors", EUserInterfaceActionType::Button, FInputChord( EKeys::Escape ) ) ;
+	UI_COMMAND( InvertSelection, "Invert Selection", "Inverts the current selection", EUserInterfaceActionType::Button, FInputChord() );
+
+	UI_COMMAND( SelectAllActorsOfSameClass, "Select All Actors of Same Class", "Selects all the actors that have the same class", EUserInterfaceActionType::Button, FInputChord(EModifierKey::Shift|EModifierKey::Control, EKeys::A) );
+	UI_COMMAND( SelectAllActorsOfSameClassWithArchetype, "Select All Actors with Same Archetype", "Selects all the actors of the same class that have the same archetype", EUserInterfaceActionType::Button, FInputChord() );
+	UI_COMMAND( SelectComponentOwnerActor, "Select Component Owner", "Select the actor that owns the currently selected component(s)", EUserInterfaceActionType::Button, FInputChord() );
+	UI_COMMAND( SelectRelevantLights, "Select Relevant Lights", "Select all lights relevant to the current selection", EUserInterfaceActionType::Button, FInputChord() ); 
+	UI_COMMAND( SelectStaticMeshesOfSameClass, "Select All Using Selected Static Meshes (Selected Actor Types)", "Selects all actors with the same static mesh and actor class as the selection", EUserInterfaceActionType::Button, FInputChord() ); 
+	UI_COMMAND( SelectStaticMeshesAllClasses, "Select All Using Selected Static Meshes (All Actor Types)", "Selects all actors with the same static mesh as the selection", EUserInterfaceActionType::Button, FInputChord( EModifierKey::Shift, EKeys::E ) ); 
+	UI_COMMAND( SelectSkeletalMeshesOfSameClass, "Select All Using Selected Skeletal Meshes (Selected Actor Types)", "Selects all actors with the same skeletal mesh and actor class as the selection", EUserInterfaceActionType::Button, FInputChord() ); 
+	UI_COMMAND( SelectSkeletalMeshesAllClasses, "Select All Using Selected Skeletal Meshes (All Actor Types)", "Selects all actors with the same skeletal mesh as the selection", EUserInterfaceActionType::Button, FInputChord() ); 
+	UI_COMMAND( SelectAllWithSameMaterial, "Select All With Same Material", "Selects all actors with the same material as the selection", EUserInterfaceActionType::Button, FInputChord() ); 
+	UI_COMMAND( SelectMatchingEmitter, "Select All Matching Emitters", "Selects all emitters with the same particle system as the selection", EUserInterfaceActionType::Button, FInputChord() ); 
+	UI_COMMAND( SelectAllLights, "Select All Lights", "Selects all lights", EUserInterfaceActionType::Button, FInputChord() ); 
+	UI_COMMAND( SelectStationaryLightsExceedingOverlap, "Select Stationary Lights exceeding overlap", "Selects all stationary lights exceeding the overlap limit", EUserInterfaceActionType::Button, FInputChord() ); 
+	UI_COMMAND( SelectAllAddditiveBrushes, "Select All Additive Brushes", "Selects all additive brushes", EUserInterfaceActionType::Button, FInputChord() ); 
+	UI_COMMAND( SelectAllSubtractiveBrushes, "Select All Subtractive Brushes", "Selects all subtractive brushes", EUserInterfaceActionType::Button, FInputChord() ); 
+	UI_COMMAND( SelectAllActorsControlledByMatinee, "Select Actors Used by This Matinee", "Selects all actors controlled by this Matinee", EUserInterfaceActionType::Button, FInputChord() ); 
+
+	UI_COMMAND( SelectAllSurfaces, "Select All Surfaces", "Selects all bsp surfaces", EUserInterfaceActionType::Button, FInputChord(EModifierKey::Shift, EKeys::S) ); 
+
+	UI_COMMAND( SurfSelectAllMatchingBrush, "Select Matching Brush", "Selects the surfaces belonging to the same brush as the selected surfaces", EUserInterfaceActionType::Button, FInputChord(EModifierKey::Shift, EKeys::B) ); 
+	UI_COMMAND( SurfSelectAllMatchingTexture, "Select Matching Material", "Selects all surfaces with the same material as the selected surfaces", EUserInterfaceActionType::Button, FInputChord(EModifierKey::Shift, EKeys::T) ); 
+	UI_COMMAND( SurfSelectAllAdjacents, "Select All Adjacent Surfaces", "Selects all surfaces adjacent to the currently selected surfaces", EUserInterfaceActionType::Button, FInputChord(EModifierKey::Shift, EKeys::J) ); 
+	UI_COMMAND( SurfSelectAllAdjacentCoplanars, "Select All Coplanar Surfaces", "Selects all surfaces adjacent and coplanar with the selected surfaces", EUserInterfaceActionType::Button, FInputChord(EModifierKey::Shift, EKeys::C) ); 
+	UI_COMMAND( SurfSelectAllAdjacentWalls, "Select All Adjacent Wall Surfaces", "Selects all adjacent upright surfaces", EUserInterfaceActionType::Button, FInputChord(EModifierKey::Shift, EKeys::W) ); 
+	UI_COMMAND( SurfSelectAllAdjacentFloors, "Select All Adjacent Floor Surfaces", "Selects all adjacent floor sufaces(ones with normals pointing up)", EUserInterfaceActionType::Button, FInputChord(EModifierKey::Shift, EKeys::U) ); 
+	UI_COMMAND( SurfSelectAllAdjacentSlants, "Select All Adjacent Slant Surfaces", "Selects all adjacent slant surfaces (surfaces that are not walls, floors, or ceilings according to their normals) to the currently selected surfaces.", EUserInterfaceActionType::Button, FInputChord(EModifierKey::Shift, EKeys::Y) ); 
+	UI_COMMAND( SurfSelectReverse, "Invert Surface Selection", "Inverts the current surface selection", EUserInterfaceActionType::Button, FInputChord(EModifierKey::Shift, EKeys::Q) );
+	UI_COMMAND( SurfSelectMemorize, "Memorize Surface Selection", "Stores the current surface selection in memory", EUserInterfaceActionType::Button, FInputChord(EModifierKey::Shift, EKeys::M) ); 
+	UI_COMMAND( SurfSelectRecall, "Recall Surface Selection", "Replace the current selection with the selection saved in memory", EUserInterfaceActionType::Button, FInputChord(EModifierKey::Shift, EKeys::R) ); 
+	UI_COMMAND( SurfSelectOr, "Surface Selection OR", "Replace the current selection with only the surfaces which are both currently selected and contained within the saved selection in memory", EUserInterfaceActionType::Button, FInputChord(EModifierKey::Shift, EKeys::O) ); 
+	UI_COMMAND( SurfSelectAnd, "Surface Selection AND", "Add the selection of surfaces saved in memory to the current selection", EUserInterfaceActionType::Button, FInputChord(EModifierKey::Shift, EKeys::A) ); 
+	UI_COMMAND( SurfSelectXor, "Surace Selection XOR", " Replace the current selection with only the surfaces that are not in both the current selection and the selection saved in memory", EUserInterfaceActionType::Button, FInputChord(EModifierKey::Shift, EKeys::X) ); 
+	UI_COMMAND( SurfUnalign, "Align Surface Default", "Default surface alignment", EUserInterfaceActionType::Button, FInputChord() );
+	UI_COMMAND( SurfAlignPlanarAuto, "Align Surface Planar", "Planar surface alignment", EUserInterfaceActionType::Button, FInputChord() );
+	UI_COMMAND( SurfAlignPlanarWall, "Align Surface Planar Wall", "Planar wall surface alignment", EUserInterfaceActionType::Button, FInputChord() );
+	UI_COMMAND( SurfAlignPlanarFloor, "Align Surface Planar Floor", "Planar floor surface alignment", EUserInterfaceActionType::Button, FInputChord() );
+	UI_COMMAND( SurfAlignBox, "Align Surface Box", "Box surface alignment", EUserInterfaceActionType::Button, FInputChord() );
+	UI_COMMAND( SurfAlignFit, "Align Surface Fit", "Best fit surface alignment", EUserInterfaceActionType::Button, FInputChord() );
+	UI_COMMAND( ApplyMaterialToSurface, "Apply Material to Surface Selection", "Applies the selected material to the selected surfaces", EUserInterfaceActionType::Button, FInputChord() );
+
+	UI_COMMAND( CreateBoundingBoxVolume, "Create Bounding Box Blocking Volume From Mesh", "Create a bounding box blocking volume from the static mesh", EUserInterfaceActionType::Button, FInputChord() );
+	UI_COMMAND( CreateHeavyConvexVolume, "Heavy Convex Blocking Volume From Mesh", "Creates a heavy convex blocing volume from the static mesh", EUserInterfaceActionType::Button, FInputChord() );
+	UI_COMMAND( CreateNormalConvexVolume, "Normal Convex Blocking Volume From Mesh", "Creates a normal convex blocking volume from the static mesh", EUserInterfaceActionType::Button, FInputChord() ); 
+	UI_COMMAND( CreateLightConvexVolume, "Light Convex Blocking Volume From Mesh", "Creates a light convex blocking volume from the static mesh", EUserInterfaceActionType::Button, FInputChord() );
+	UI_COMMAND( CreateRoughConvexVolume, "Rought Convex Blocking Volume From Mesh", "Creates a rough convex blocking volume from the static mesh", EUserInterfaceActionType::Button, FInputChord() );
+
+	UI_COMMAND( SaveBrushAsCollision, "Save Collision from Builder Brush", "Creates a collision primitive on the selected static meshes based on the shape of the builder brush", EUserInterfaceActionType::Button, FInputChord() );
+
+	UI_COMMAND( KeepSimulationChanges, "Keep Simulation Changes", "Saves the changes made to this actor in Simulate mode to the actor's default state.", EUserInterfaceActionType::Button, FInputChord( EKeys::K ) );
+
+	UI_COMMAND( MakeActorLevelCurrent, "Make Selected Actor's Level Current", "Makes the selected actors level the current level", EUserInterfaceActionType::Button, FInputChord( EKeys::M ) );
 #if PLATFORM_MAC
-	UI_COMMAND( MoveSelectedToCurrentLevel, "Move Selection to Current Level", "Moves the selected actors to the current level", EUserInterfaceActionType::Button, FInputGesture( EModifierKey::Command, EKeys::M ) );
+	UI_COMMAND( MoveSelectedToCurrentLevel, "Move Selection to Current Level", "Moves the selected actors to the current level", EUserInterfaceActionType::Button, FInputChord( EModifierKey::Command, EKeys::M ) );
 #else
-	UI_COMMAND( MoveSelectedToCurrentLevel, "Move Selection to Current Level", "Moves the selected actors to the current level", EUserInterfaceActionType::Button, FInputGesture( EModifierKey::Control, EKeys::M ) );
+	UI_COMMAND( MoveSelectedToCurrentLevel, "Move Selection to Current Level", "Moves the selected actors to the current level", EUserInterfaceActionType::Button, FInputChord( EModifierKey::Control, EKeys::M ) );
 #endif
-	UI_COMMAND( FindLevelsInLevelBrowser, "Find Levels in Level Browser", "Finds the selected actors level in the level browser", EUserInterfaceActionType::Button, FInputGesture() );
-	UI_COMMAND( AddLevelsToSelection, "Add Levels to Selection", "Adds the selected actors levels  to the current level browser selection", EUserInterfaceActionType::Button, FInputGesture() );
-	UI_COMMAND( RemoveLevelsFromSelection, "Remove Levels from Selection", "Removes the selected actors levels from the current level browser selection", EUserInterfaceActionType::Button, FInputGesture() );
-
-	UI_COMMAND( FindActorInLevelScript, "Find in Level Blueprint", "Finds any references to the selected actor in its level's blueprint", EUserInterfaceActionType::Button, FInputGesture() );
-
-	UI_COMMAND( WorldProperties, "World Settings", "Displays the world settings", EUserInterfaceActionType::Button, FInputGesture() );
-	UI_COMMAND( OpenContentBrowser, "Open Content Browser", "Opens the Content Browser", EUserInterfaceActionType::Button, FInputGesture(EModifierKey::Control|EModifierKey::Shift, EKeys::F) );
-	UI_COMMAND( OpenMarketplace, "Open Marketplace", "Opens the Marketplace", EUserInterfaceActionType::Button, FInputGesture() );
-	UI_COMMAND( AddMatinee, "Add Matinee", "Creates a new matinee actor to edit", EUserInterfaceActionType::Button, FInputGesture() );
-	UI_COMMAND( EditMatinee, "Edit Matinee", "Selects a Matinee to edit", EUserInterfaceActionType::Button, FInputGesture() );
-
-	UI_COMMAND( OpenLevelBlueprint, "Open Level Blueprint", "Edit the Level Blueprint for the current level", EUserInterfaceActionType::Button, FInputGesture() );
-	UI_COMMAND( CheckOutProjectSettingsConfig, "Check Out", "Checks out the project settings config file so the game mode can be set.", EUserInterfaceActionType::Button, FInputGesture() );
-	UI_COMMAND( CreateBlankBlueprintClass, "New Empty Blueprint Class...", "Create a new Blueprint Class", EUserInterfaceActionType::Button, FInputGesture() );
-	UI_COMMAND( ConvertSelectionToBlueprintViaHarvest, "Convert Selected Components to Blueprint Class...", "Replace all of the selected actors with a new Blueprint Class based on Actor that contains the components", EUserInterfaceActionType::Button, FInputGesture() );
-	UI_COMMAND( ConvertSelectionToBlueprintViaSubclass, "Convert Selected Actor to Blueprint Class...", "Replace the selected actor with a new Blueprint subclass based on the class of the selected Actor", EUserInterfaceActionType::Button, FInputGesture() );
-
-	UI_COMMAND( ShowTransformWidget, "Show Transform Widget", "Toggles the visibility of the transform widgets", EUserInterfaceActionType::ToggleButton, FInputGesture() );
-	UI_COMMAND( AllowTranslucentSelection, "Allow Translucent Selection", "Allows translucent objects to be selected", EUserInterfaceActionType::ToggleButton, FInputGesture(EKeys::T) );
-	UI_COMMAND( AllowGroupSelection, "Allow Group Selection", "Allows actor groups to be selected", EUserInterfaceActionType::ToggleButton, FInputGesture(EModifierKey::Control|EModifierKey::Shift, EKeys::G) );
-	UI_COMMAND( StrictBoxSelect, "Strict Box Selection", "When enabled an object must be entirely encompassed by the selection box when marquee box selecting", EUserInterfaceActionType::ToggleButton, FInputGesture() );
-	UI_COMMAND( TransparentBoxSelect, "Transparent Box Selection", "When enabled, marquee box select operations will also select objects that are occluded by other objects.", EUserInterfaceActionType::ToggleButton, FInputGesture() );
-	UI_COMMAND( DrawBrushMarkerPolys, "Draw Brush Polys", "Draws semi-transparent polygons around a brush when selected", EUserInterfaceActionType::ToggleButton, FInputGesture() );
-	UI_COMMAND( OnlyLoadVisibleInPIE, "Only Load Visible Levels in Game Preview", "If enabled, when game preview starts, only visible levels will be loaded", EUserInterfaceActionType::ToggleButton, FInputGesture() );
-	UI_COMMAND( ToggleSocketSnapping, "Enable Socket Snapping", "Enables or disables snapping to sockets", EUserInterfaceActionType::ToggleButton, FInputGesture() ); 
-	UI_COMMAND( ToggleParticleSystemLOD, "Enable Particle System LOD Switching", "If enabled particle systems will use distance LOD switching in perspective viewports", EUserInterfaceActionType::ToggleButton, FInputGesture() );
-	UI_COMMAND( ToggleFreezeParticleSimulation, "Freeze Particle Simulation", "If enabled particle systems will freeze their simulation state", EUserInterfaceActionType::ToggleButton, FInputGesture() );
-	UI_COMMAND( ToggleParticleSystemHelpers, "Toggle Particle System Helpers", "Toggles showing particle system helper widgets in viewports", EUserInterfaceActionType::ToggleButton, FInputGesture() );
-	UI_COMMAND( ToggleLODViewLocking, "Enable LOD View Locking", "If enabled viewports of the same type will use the same LOD", EUserInterfaceActionType::ToggleButton, FInputGesture() );
-	UI_COMMAND( LevelStreamingVolumePrevis, "Enable Automatic Level Streaming", "If enabled, the viewport will stream in levels automatically when the camera is moved", EUserInterfaceActionType::ToggleButton, FInputGesture() );
-	UI_COMMAND( EnableActorSnap, "Enable Actor Snapping", "If enabled, actors will snap to the location of other actors when they are within distance", EUserInterfaceActionType::ToggleButton, FInputGesture(EModifierKey::Control|EModifierKey::Shift, EKeys::K) );
-	UI_COMMAND( EnableVertexSnap, "Enable Vertex Snapping","If enabled, actors will snap to the location of the nearest vertex on another actor in the direction of movement", EUserInterfaceActionType::ToggleButton, FInputGesture() );
-	UI_COMMAND( ToggleHideViewportUI, "Hide Viewport UI", "Toggles hidden viewport UI mode. Hides all overlaid viewport UI widgets", EUserInterfaceActionType::ToggleButton, FInputGesture() );
+	UI_COMMAND( FindLevelsInLevelBrowser, "Find Levels in Level Browser", "Finds the selected actors level in the level browser", EUserInterfaceActionType::Button, FInputChord() );
+	UI_COMMAND( AddLevelsToSelection, "Add Levels to Selection", "Adds the selected actors levels  to the current level browser selection", EUserInterfaceActionType::Button, FInputChord() );
+	UI_COMMAND( RemoveLevelsFromSelection, "Remove Levels from Selection", "Removes the selected actors levels from the current level browser selection", EUserInterfaceActionType::Button, FInputChord() );
+
+	UI_COMMAND( FindActorInLevelScript, "Find in Level Blueprint", "Finds any references to the selected actor in its level's blueprint", EUserInterfaceActionType::Button, FInputChord() );
+
+	UI_COMMAND( WorldProperties, "World Settings", "Displays the world settings", EUserInterfaceActionType::Button, FInputChord() );
+	UI_COMMAND( OpenContentBrowser, "Open Content Browser", "Opens the Content Browser", EUserInterfaceActionType::Button, FInputChord(EModifierKey::Control|EModifierKey::Shift, EKeys::F) );
+	UI_COMMAND( OpenMarketplace, "Open Marketplace", "Opens the Marketplace", EUserInterfaceActionType::Button, FInputChord() );
+	UI_COMMAND( AddMatinee, "Add Matinee", "Creates a new matinee actor to edit", EUserInterfaceActionType::Button, FInputChord() );
+	UI_COMMAND( EditMatinee, "Edit Matinee", "Selects a Matinee to edit", EUserInterfaceActionType::Button, FInputChord() );
+
+	UI_COMMAND( OpenLevelBlueprint, "Open Level Blueprint", "Edit the Level Blueprint for the current level", EUserInterfaceActionType::Button, FInputChord() );
+	UI_COMMAND( CheckOutProjectSettingsConfig, "Check Out", "Checks out the project settings config file so the game mode can be set.", EUserInterfaceActionType::Button, FInputChord() );
+	UI_COMMAND( CreateBlankBlueprintClass, "New Empty Blueprint Class...", "Create a new Blueprint Class", EUserInterfaceActionType::Button, FInputChord() );
+	UI_COMMAND( ConvertSelectionToBlueprintViaHarvest, "Convert Selected Components to Blueprint Class...", "Replace all of the selected actors with a new Blueprint Class based on Actor that contains the components", EUserInterfaceActionType::Button, FInputChord() );
+	UI_COMMAND( ConvertSelectionToBlueprintViaSubclass, "Convert Selected Actor to Blueprint Class...", "Replace the selected actor with a new Blueprint subclass based on the class of the selected Actor", EUserInterfaceActionType::Button, FInputChord() );
+
+	UI_COMMAND( ShowTransformWidget, "Show Transform Widget", "Toggles the visibility of the transform widgets", EUserInterfaceActionType::ToggleButton, FInputChord() );
+	UI_COMMAND( AllowTranslucentSelection, "Allow Translucent Selection", "Allows translucent objects to be selected", EUserInterfaceActionType::ToggleButton, FInputChord(EKeys::T) );
+	UI_COMMAND( AllowGroupSelection, "Allow Group Selection", "Allows actor groups to be selected", EUserInterfaceActionType::ToggleButton, FInputChord(EModifierKey::Control|EModifierKey::Shift, EKeys::G) );
+	UI_COMMAND( StrictBoxSelect, "Strict Box Selection", "When enabled an object must be entirely encompassed by the selection box when marquee box selecting", EUserInterfaceActionType::ToggleButton, FInputChord() );
+	UI_COMMAND( TransparentBoxSelect, "Transparent Box Selection", "When enabled, marquee box select operations will also select objects that are occluded by other objects.", EUserInterfaceActionType::ToggleButton, FInputChord() );
+	UI_COMMAND( DrawBrushMarkerPolys, "Draw Brush Polys", "Draws semi-transparent polygons around a brush when selected", EUserInterfaceActionType::ToggleButton, FInputChord() );
+	UI_COMMAND( OnlyLoadVisibleInPIE, "Only Load Visible Levels in Game Preview", "If enabled, when game preview starts, only visible levels will be loaded", EUserInterfaceActionType::ToggleButton, FInputChord() );
+	UI_COMMAND( ToggleSocketSnapping, "Enable Socket Snapping", "Enables or disables snapping to sockets", EUserInterfaceActionType::ToggleButton, FInputChord() ); 
+	UI_COMMAND( ToggleParticleSystemLOD, "Enable Particle System LOD Switching", "If enabled particle systems will use distance LOD switching in perspective viewports", EUserInterfaceActionType::ToggleButton, FInputChord() );
+	UI_COMMAND( ToggleFreezeParticleSimulation, "Freeze Particle Simulation", "If enabled particle systems will freeze their simulation state", EUserInterfaceActionType::ToggleButton, FInputChord() );
+	UI_COMMAND( ToggleParticleSystemHelpers, "Toggle Particle System Helpers", "Toggles showing particle system helper widgets in viewports", EUserInterfaceActionType::ToggleButton, FInputChord() );
+	UI_COMMAND( ToggleLODViewLocking, "Enable LOD View Locking", "If enabled viewports of the same type will use the same LOD", EUserInterfaceActionType::ToggleButton, FInputChord() );
+	UI_COMMAND( LevelStreamingVolumePrevis, "Enable Automatic Level Streaming", "If enabled, the viewport will stream in levels automatically when the camera is moved", EUserInterfaceActionType::ToggleButton, FInputChord() );
+	UI_COMMAND( EnableActorSnap, "Enable Actor Snapping", "If enabled, actors will snap to the location of other actors when they are within distance", EUserInterfaceActionType::ToggleButton, FInputChord(EModifierKey::Control|EModifierKey::Shift, EKeys::K) );
+	UI_COMMAND( EnableVertexSnap, "Enable Vertex Snapping","If enabled, actors will snap to the location of the nearest vertex on another actor in the direction of movement", EUserInterfaceActionType::ToggleButton, FInputChord() );
+	UI_COMMAND( ToggleHideViewportUI, "Hide Viewport UI", "Toggles hidden viewport UI mode. Hides all overlaid viewport UI widgets", EUserInterfaceActionType::ToggleButton, FInputChord() );
 
 	//if (FParse::Param( FCommandLine::Get(), TEXT( "editortoolbox" ) ))
 	//{
-	//	UI_COMMAND( BspMode, "Enable Bsp Mode", "Enables BSP mode", EUserInterfaceActionType::ToggleButton, FInputGesture( EModifierKey::Shift, EKeys::One ) );
-	//	UI_COMMAND( MeshPaintMode, "Enable Mesh Paint Mode", "Enables mesh paint mode", EUserInterfaceActionType::ToggleButton, FInputGesture( EModifierKey::Shift, EKeys::Two ) );
-	//	UI_COMMAND( LandscapeMode, "Enable Landscape Mode", "Enables landscape editing", EUserInterfaceActionType::ToggleButton, FInputGesture( EModifierKey::Shift, EKeys::Three ) );
-	//	UI_COMMAND( FoliageMode, "Enable Foliage Mode", "Enables foliage editing", EUserInterfaceActionType::ToggleButton, FInputGesture( EModifierKey::Shift, EKeys::Four ) );
+	//	UI_COMMAND( BspMode, "Enable Bsp Mode", "Enables BSP mode", EUserInterfaceActionType::ToggleButton, FInputChord( EModifierKey::Shift, EKeys::One ) );
+	//	UI_COMMAND( MeshPaintMode, "Enable Mesh Paint Mode", "Enables mesh paint mode", EUserInterfaceActionType::ToggleButton, FInputChord( EModifierKey::Shift, EKeys::Two ) );
+	//	UI_COMMAND( LandscapeMode, "Enable Landscape Mode", "Enables landscape editing", EUserInterfaceActionType::ToggleButton, FInputChord( EModifierKey::Shift, EKeys::Three ) );
+	//	UI_COMMAND( FoliageMode, "Enable Foliage Mode", "Enables foliage editing", EUserInterfaceActionType::ToggleButton, FInputChord( EModifierKey::Shift, EKeys::Four ) );
 	//}
 
-	UI_COMMAND( ShowSelectedDetails, "Show Actor Details", "Opens a details panel for the selected actors", EUserInterfaceActionType::Button, FInputGesture( EKeys::F4 ) );
+	UI_COMMAND( ShowSelectedDetails, "Show Actor Details", "Opens a details panel for the selected actors", EUserInterfaceActionType::Button, FInputChord( EKeys::F4 ) );
 
-	UI_COMMAND( RecompileShaders, "Recompile Changed Shaders", "Recompiles shaders which are out of date", EUserInterfaceActionType::Button, FInputGesture( EModifierKey::Shift|EModifierKey::Control, EKeys::Period ) );
-	UI_COMMAND( ProfileGPU, "Profile GPU", "Profiles the GPU for the next frame and opens a window with profiled data", EUserInterfaceActionType::Button, FInputGesture( EModifierKey::Shift|EModifierKey::Control, EKeys::Comma ) );
+	UI_COMMAND( RecompileShaders, "Recompile Changed Shaders", "Recompiles shaders which are out of date", EUserInterfaceActionType::Button, FInputChord( EModifierKey::Shift|EModifierKey::Control, EKeys::Period ) );
+	UI_COMMAND( ProfileGPU, "Profile GPU", "Profiles the GPU for the next frame and opens a window with profiled data", EUserInterfaceActionType::Button, FInputChord( EModifierKey::Shift|EModifierKey::Control, EKeys::Comma ) );
 
-	UI_COMMAND( ResetAllParticleSystems, "Reset All Particle Systems", "Resets all particle system emitters (removes all active particles and restarts them)", EUserInterfaceActionType::Button, FInputGesture( EModifierKey::Shift, EKeys::Slash ) );
-	UI_COMMAND( ResetSelectedParticleSystem, "Resets Selected Particle Systems" , "Resets selected particle system emitters (removes all active particles and restarts them)", EUserInterfaceActionType::Button, FInputGesture( EKeys::Slash ) );
+	UI_COMMAND( ResetAllParticleSystems, "Reset All Particle Systems", "Resets all particle system emitters (removes all active particles and restarts them)", EUserInterfaceActionType::Button, FInputChord( EModifierKey::Shift, EKeys::Slash ) );
+	UI_COMMAND( ResetSelectedParticleSystem, "Resets Selected Particle Systems" , "Resets selected particle system emitters (removes all active particles and restarts them)", EUserInterfaceActionType::Button, FInputChord( EKeys::Slash ) );
 
-	UI_COMMAND( SelectActorsInLayers, "Select all actors in selected actor's layers", "Selects all actors belonging to the layers of the currently selected actors", EUserInterfaceActionType::Button, FInputGesture( EModifierKey::Control, EKeys::L ) );
+	UI_COMMAND( SelectActorsInLayers, "Select all actors in selected actor's layers", "Selects all actors belonging to the layers of the currently selected actors", EUserInterfaceActionType::Button, FInputChord( EModifierKey::Control, EKeys::L ) );
 
-	UI_COMMAND( FocusAllViewportsToSelection, "Focus Selected Actors in All Viewports", "Moves the camera in front of the selected actors in all open viewports", EUserInterfaceActionType::Button, FInputGesture( EKeys::F, EModifierKey::Shift ) );
+	UI_COMMAND( FocusAllViewportsToSelection, "Focus Selected Actors in All Viewports", "Moves the camera in front of the selected actors in all open viewports", EUserInterfaceActionType::Button, FInputChord( EKeys::F, EModifierKey::Shift ) );
 
-	UI_COMMAND( MaterialQualityLevel_Low, "Low", "Sets material quality in the scene to low.", EUserInterfaceActionType::RadioButton, FInputGesture() );
-	UI_COMMAND( MaterialQualityLevel_High, "High", "Sets material quality in the scene to high.", EUserInterfaceActionType::RadioButton, FInputGesture() );
+	UI_COMMAND( MaterialQualityLevel_Low, "Low", "Sets material quality in the scene to low.", EUserInterfaceActionType::RadioButton, FInputChord() );
+	UI_COMMAND( MaterialQualityLevel_High, "High", "Sets material quality in the scene to high.", EUserInterfaceActionType::RadioButton, FInputChord() );
 
-	UI_COMMAND( ConnectToSourceControl, "Connect to Source Control...", "Opens a dialog to connect to source control.", EUserInterfaceActionType::Button, FInputGesture());
-	UI_COMMAND( ChangeSourceControlSettings, "Change Source Control Settings...", "Opens a dialog to change source control settings.", EUserInterfaceActionType::Button, FInputGesture());
-	UI_COMMAND( CheckOutModifiedFiles, "Check Out Modified Files...", "Opens a dialog to check out any assets which have been modified.", EUserInterfaceActionType::Button, FInputGesture());
-	UI_COMMAND( SubmitToSourceControl, "Submit to Source Control...", "Opens a dialog with check in options for content and levels.", EUserInterfaceActionType::Button, FInputGesture());
+	UI_COMMAND( ConnectToSourceControl, "Connect to Source Control...", "Opens a dialog to connect to source control.", EUserInterfaceActionType::Button, FInputChord());
+	UI_COMMAND( ChangeSourceControlSettings, "Change Source Control Settings...", "Opens a dialog to change source control settings.", EUserInterfaceActionType::Button, FInputChord());
+	UI_COMMAND( CheckOutModifiedFiles, "Check Out Modified Files...", "Opens a dialog to check out any assets which have been modified.", EUserInterfaceActionType::Button, FInputChord());
+	UI_COMMAND( SubmitToSourceControl, "Submit to Source Control...", "Opens a dialog with check in options for content and levels.", EUserInterfaceActionType::Button, FInputChord());
 
 	static const FText FeatureLevelLabels[ERHIFeatureLevel::Num] = 
 	{
@@ -3087,7 +3087,7 @@ void FLevelEditorCommands::RegisterCommands()
 			FeatureLevelLabels[i],
 			FeatureLevelToolTips[i])
 			.UserInterfaceType(EUserInterfaceActionType::RadioButton)
-			.DefaultGesture(FInputGesture());
+			.DefaultChord(FInputChord());
 	}
 }
 
diff --git a/Engine/Source/Editor/LevelEditor/Private/LevelEditorModesActions.cpp b/Engine/Source/Editor/LevelEditor/Private/LevelEditorModesActions.cpp
--- a/Engine/Source/Editor/LevelEditor/Private/LevelEditorModesActions.cpp
+++ b/Engine/Source/Editor/LevelEditor/Private/LevelEditorModesActions.cpp
@@ -44,7 +44,7 @@ void FLevelEditorModesCommands::RegisterCommands()
 				Tooltip,
 				Mode.IconBrush,
 				EUserInterfaceActionType::ToggleButton,
-				editorMode < 9 ? FInputGesture( EModifierKey::Shift, EdModeKeys[editorMode] ) : FInputGesture() );
+				editorMode < 9 ? FInputChord( EModifierKey::Shift, EdModeKeys[editorMode] ) : FInputChord() );
 
 			EditorModeCommands.Add(EditorModeCommand);
 		}
diff --git a/Engine/Source/Editor/LevelEditor/Private/LevelViewportActions.cpp b/Engine/Source/Editor/LevelEditor/Private/LevelViewportActions.cpp
--- a/Engine/Source/Editor/LevelEditor/Private/LevelViewportActions.cpp
+++ b/Engine/Source/Editor/LevelEditor/Private/LevelViewportActions.cpp
@@ -24,38 +24,38 @@ FLevelViewportCommands::~FLevelViewportCommands()
 PRAGMA_DISABLE_OPTIMIZATION
 void FLevelViewportCommands::RegisterCommands()
 {
-	UI_COMMAND( ToggleMaximize, "Maximize Viewport", "Toggles the Maximize state of the current viewport", EUserInterfaceActionType::ToggleButton, FInputGesture() );
-	UI_COMMAND( ToggleGameView, "Game View", "Toggles game view.  Game view shows the scene as it appears in game", EUserInterfaceActionType::ToggleButton, FInputGesture( EKeys::G ) );
-	UI_COMMAND( ToggleImmersive, "Immersive Mode", "Switches this viewport between immersive mode and regular mode", EUserInterfaceActionType::ToggleButton, FInputGesture( EKeys::F11 ) );
+	UI_COMMAND( ToggleMaximize, "Maximize Viewport", "Toggles the Maximize state of the current viewport", EUserInterfaceActionType::ToggleButton, FInputChord() );
+	UI_COMMAND( ToggleGameView, "Game View", "Toggles game view.  Game view shows the scene as it appears in game", EUserInterfaceActionType::ToggleButton, FInputChord( EKeys::G ) );
+	UI_COMMAND( ToggleImmersive, "Immersive Mode", "Switches this viewport between immersive mode and regular mode", EUserInterfaceActionType::ToggleButton, FInputChord( EKeys::F11 ) );
 
-	UI_COMMAND( CreateCamera, "Create Camera Here", "Creates a new camera actor at the current location of this viewport's camera", EUserInterfaceActionType::Button, FInputGesture() );
-	UI_COMMAND( HighResScreenshot, "High Resolution Screenshot...", "Opens the control panel for high resolution screenshots", EUserInterfaceActionType::Button, FInputGesture() );
+	UI_COMMAND( CreateCamera, "Create Camera Here", "Creates a new camera actor at the current location of this viewport's camera", EUserInterfaceActionType::Button, FInputChord() );
+	UI_COMMAND( HighResScreenshot, "High Resolution Screenshot...", "Opens the control panel for high resolution screenshots", EUserInterfaceActionType::Button, FInputChord() );
 	
-	UI_COMMAND( UseDefaultShowFlags, "Use Defaults", "Resets all show flags to default", EUserInterfaceActionType::Button, FInputGesture() );
+	UI_COMMAND( UseDefaultShowFlags, "Use Defaults", "Resets all show flags to default", EUserInterfaceActionType::Button, FInputChord() );
 
-	UI_COMMAND( PilotSelectedActor, "Pilot Selected Actor", "Move the selected actor around using the viewport controls, and bind the viewport to the actor's location and orientation.", EUserInterfaceActionType::Button, FInputGesture( EModifierKey::Control | EModifierKey::Shift, EKeys::P ) );
-	UI_COMMAND( EjectActorPilot, "Eject from Actor Pilot", "Stop piloting an actor with the current viewport. Unlocks the viewport's position and orientation from the actor the viewport is currently piloting.", EUserInterfaceActionType::Button, FInputGesture() );
-	UI_COMMAND( ToggleActorPilotCameraView, "Actor Pilot Camera View", "Toggles showing the exact camera view when using the viewport to pilot a camera", EUserInterfaceActionType::ToggleButton, FInputGesture( EModifierKey::Control | EModifierKey::Shift, EKeys::C ) );
+	UI_COMMAND( PilotSelectedActor, "Pilot Selected Actor", "Move the selected actor around using the viewport controls, and bind the viewport to the actor's location and orientation.", EUserInterfaceActionType::Button, FInputChord( EModifierKey::Control | EModifierKey::Shift, EKeys::P ) );
+	UI_COMMAND( EjectActorPilot, "Eject from Actor Pilot", "Stop piloting an actor with the current viewport. Unlocks the viewport's position and orientation from the actor the viewport is currently piloting.", EUserInterfaceActionType::Button, FInputChord() );
+	UI_COMMAND( ToggleActorPilotCameraView, "Actor Pilot Camera View", "Toggles showing the exact camera view when using the viewport to pilot a camera", EUserInterfaceActionType::ToggleButton, FInputChord( EModifierKey::Control | EModifierKey::Shift, EKeys::C ) );
 
-	UI_COMMAND( ViewportConfig_OnePane, "Layout One Pane", "Changes the viewport arrangement to one pane", EUserInterfaceActionType::ToggleButton, FInputGesture() );
-	UI_COMMAND( ViewportConfig_TwoPanesH, "Layout Two Panes (horizontal)", "Changes the viewport arrangement to two panes, side-by-side", EUserInterfaceActionType::ToggleButton, FInputGesture() );
-	UI_COMMAND( ViewportConfig_TwoPanesV, "Layout Two Panes (vertical)", "Changes the viewport arrangement to two panes, one above the other", EUserInterfaceActionType::ToggleButton, FInputGesture() );
-	UI_COMMAND( ViewportConfig_ThreePanesLeft, "Layout Three Panes (one left, two right)", "Changes the viewport arrangement to three panes, one on the left, two on the right", EUserInterfaceActionType::ToggleButton, FInputGesture() );
-	UI_COMMAND( ViewportConfig_ThreePanesRight, "Layout Three Panes (one right, two left)", "Changes the viewport arrangement to three panes, one on the right, two on the left", EUserInterfaceActionType::ToggleButton, FInputGesture() );
-	UI_COMMAND( ViewportConfig_ThreePanesTop, "Layout Three Panes (one top, two bottom)", "Changes the viewport arrangement to three panes, one on the top, two on the bottom", EUserInterfaceActionType::ToggleButton, FInputGesture() );
-	UI_COMMAND( ViewportConfig_ThreePanesBottom, "Layout Three Panes (one bottom, two top)", "Changes the viewport arrangement to three panes, one on the bottom, two on the top", EUserInterfaceActionType::ToggleButton, FInputGesture() );
-	UI_COMMAND( ViewportConfig_FourPanesLeft, "Layout Four Panes (one left, three right)", "Changes the viewport arrangement to four panes, one on the left, three on the right", EUserInterfaceActionType::ToggleButton, FInputGesture() );
-	UI_COMMAND( ViewportConfig_FourPanesRight, "Layout Four Panes (one right, three left)", "Changes the viewport arrangement to four panes, one on the right, three on the left", EUserInterfaceActionType::ToggleButton, FInputGesture() );
-	UI_COMMAND( ViewportConfig_FourPanesTop, "Layout Four Panes (one top, three bottom)", "Changes the viewport arrangement to four panes, one on the top, three on the bottom", EUserInterfaceActionType::ToggleButton, FInputGesture() );
-	UI_COMMAND( ViewportConfig_FourPanesBottom, "Layout Four Panes (one bottom, three top)", "Changes the viewport arrangement to four panes, one on the bottom, three on the top", EUserInterfaceActionType::ToggleButton, FInputGesture() );
-	UI_COMMAND( ViewportConfig_FourPanes2x2, "Layout Four Panes (2x2)", "Changes the viewport arrangement to four panes, in a 2x2 grid", EUserInterfaceActionType::ToggleButton, FInputGesture() );
+	UI_COMMAND( ViewportConfig_OnePane, "Layout One Pane", "Changes the viewport arrangement to one pane", EUserInterfaceActionType::ToggleButton, FInputChord() );
+	UI_COMMAND( ViewportConfig_TwoPanesH, "Layout Two Panes (horizontal)", "Changes the viewport arrangement to two panes, side-by-side", EUserInterfaceActionType::ToggleButton, FInputChord() );
+	UI_COMMAND( ViewportConfig_TwoPanesV, "Layout Two Panes (vertical)", "Changes the viewport arrangement to two panes, one above the other", EUserInterfaceActionType::ToggleButton, FInputChord() );
+	UI_COMMAND( ViewportConfig_ThreePanesLeft, "Layout Three Panes (one left, two right)", "Changes the viewport arrangement to three panes, one on the left, two on the right", EUserInterfaceActionType::ToggleButton, FInputChord() );
+	UI_COMMAND( ViewportConfig_ThreePanesRight, "Layout Three Panes (one right, two left)", "Changes the viewport arrangement to three panes, one on the right, two on the left", EUserInterfaceActionType::ToggleButton, FInputChord() );
+	UI_COMMAND( ViewportConfig_ThreePanesTop, "Layout Three Panes (one top, two bottom)", "Changes the viewport arrangement to three panes, one on the top, two on the bottom", EUserInterfaceActionType::ToggleButton, FInputChord() );
+	UI_COMMAND( ViewportConfig_ThreePanesBottom, "Layout Three Panes (one bottom, two top)", "Changes the viewport arrangement to three panes, one on the bottom, two on the top", EUserInterfaceActionType::ToggleButton, FInputChord() );
+	UI_COMMAND( ViewportConfig_FourPanesLeft, "Layout Four Panes (one left, three right)", "Changes the viewport arrangement to four panes, one on the left, three on the right", EUserInterfaceActionType::ToggleButton, FInputChord() );
+	UI_COMMAND( ViewportConfig_FourPanesRight, "Layout Four Panes (one right, three left)", "Changes the viewport arrangement to four panes, one on the right, three on the left", EUserInterfaceActionType::ToggleButton, FInputChord() );
+	UI_COMMAND( ViewportConfig_FourPanesTop, "Layout Four Panes (one top, three bottom)", "Changes the viewport arrangement to four panes, one on the top, three on the bottom", EUserInterfaceActionType::ToggleButton, FInputChord() );
+	UI_COMMAND( ViewportConfig_FourPanesBottom, "Layout Four Panes (one bottom, three top)", "Changes the viewport arrangement to four panes, one on the bottom, three on the top", EUserInterfaceActionType::ToggleButton, FInputChord() );
+	UI_COMMAND( ViewportConfig_FourPanes2x2, "Layout Four Panes (2x2)", "Changes the viewport arrangement to four panes, in a 2x2 grid", EUserInterfaceActionType::ToggleButton, FInputChord() );
 
-	UI_COMMAND( ApplyMaterialToActor, "Apply Material", "Attempts to apply a dropped material to this object", EUserInterfaceActionType::Button, FInputGesture() );
+	UI_COMMAND( ApplyMaterialToActor, "Apply Material", "Attempts to apply a dropped material to this object", EUserInterfaceActionType::Button, FInputChord() );
 
-	UI_COMMAND( AllowMatineePreview, "Allow Matinee Preview", "If enabled, allows matinee to be previewed in this viewport", EUserInterfaceActionType::ToggleButton, FInputGesture() );
+	UI_COMMAND( AllowMatineePreview, "Allow Matinee Preview", "If enabled, allows matinee to be previewed in this viewport", EUserInterfaceActionType::ToggleButton, FInputChord() );
 
-	UI_COMMAND( FindInLevelScriptBlueprint, "Find In Level Script", "Finds references of a selected actor in the level script blueprint", EUserInterfaceActionType::Button, FInputGesture(EModifierKey::Control, EKeys::K) );
-	UI_COMMAND( AdvancedSettings, "Advanced Settings...", "Opens the advanced viewport settings", EUserInterfaceActionType::Button, FInputGesture());
+	UI_COMMAND( FindInLevelScriptBlueprint, "Find In Level Script", "Finds references of a selected actor in the level script blueprint", EUserInterfaceActionType::Button, FInputChord(EModifierKey::Control, EKeys::K) );
+	UI_COMMAND( AdvancedSettings, "Advanced Settings...", "Opens the advanced viewport settings", EUserInterfaceActionType::Button, FInputChord());
 
 	// Generate a command for each buffer visualization mode
 	{
@@ -79,7 +79,7 @@ void FLevelViewportCommands::RegisterCommands()
 				const FText MaterialNameText = FText::FromString( InMaterialName );
 				Record.Command = FUICommandInfoDecl( Parent, ViewportCommandName, MaterialNameText, MaterialNameText )
 					.UserInterfaceType( EUserInterfaceActionType::RadioButton )
-					.DefaultGesture( FInputGesture() );
+					.DefaultChord( FInputChord() );
 			}
 		};
 
@@ -90,7 +90,7 @@ void FLevelViewportCommands::RegisterCommands()
 		OverviewRecord.Name = NAME_None;
 		OverviewRecord.Command = FUICommandInfoDecl( this->AsShared(), ViewportCommandName, LOCTEXT("BufferVisualization", "Overview"), LOCTEXT("BufferVisualization", "Overview") )
 			.UserInterfaceType( EUserInterfaceActionType::RadioButton )
-			.DefaultGesture( FInputGesture() );
+			.DefaultChord( FInputChord() );
 
 		FMaterialIterator It(this->AsShared(), BufferVisualizationModeCommands);
 		GetBufferVisualizationData().IterateOverAvailableMaterials(It);
@@ -122,7 +122,7 @@ void FLevelViewportCommands::RegisterCommands()
 		TSharedPtr<FUICommandInfo> ShowFlagCommand 
 			= FUICommandInfoDecl( this->AsShared(), SFData.ShowFlagName, LocalizedName, ShowFlagDesc )
 			.UserInterfaceType( EUserInterfaceActionType::ToggleButton )
-			.DefaultGesture( SFData.InputGesture )
+			.DefaultChord( SFData.InputChord )
 			.Icon(SFData.Group == EShowFlagGroup::SFG_Normal ?
 						FSlateIcon(FEditorStyle::GetStyleSetName(), FEditorStyle::Join( GetContextName(), TCHAR_TO_ANSI( *FString::Printf( TEXT(".%s"), *SFData.ShowFlagName.ToString() ) ) ) ) :
 						FSlateIcon());
@@ -132,8 +132,8 @@ void FLevelViewportCommands::RegisterCommands()
 
 	// Generate a command for each volume class
 	{
-		UI_COMMAND( ShowAllVolumes, "Show All Volumes", "Shows all volumes", EUserInterfaceActionType::Button, FInputGesture() );
-		UI_COMMAND( HideAllVolumes, "Hide All Volumes", "Hides all volumes", EUserInterfaceActionType::Button, FInputGesture() );
+		UI_COMMAND( ShowAllVolumes, "Show All Volumes", "Shows all volumes", EUserInterfaceActionType::Button, FInputChord() );
+		UI_COMMAND( HideAllVolumes, "Hide All Volumes", "Hides all volumes", EUserInterfaceActionType::Button, FInputChord() );
 
 		TArray< UClass* > VolumeClasses;
 		UUnrealEdEngine::GetSortedVolumeClasses(&VolumeClasses);
@@ -161,14 +161,14 @@ void FLevelViewportCommands::RegisterCommands()
 
 	// Generate a command for show/hide all layers
 	{
-		UI_COMMAND( ShowAllLayers, "Show All Layers", "Shows all layers", EUserInterfaceActionType::Button, FInputGesture() );
-		UI_COMMAND( HideAllLayers, "Hide All Layers", "Hides all layers", EUserInterfaceActionType::Button, FInputGesture() );
+		UI_COMMAND( ShowAllLayers, "Show All Layers", "Shows all layers", EUserInterfaceActionType::Button, FInputChord() );
+		UI_COMMAND( HideAllLayers, "Hide All Layers", "Hides all layers", EUserInterfaceActionType::Button, FInputChord() );
 	}
 
 	// Generate a command for each sprite category
 	{
-		UI_COMMAND( ShowAllSprites, "Show All Sprites", "Shows all sprites", EUserInterfaceActionType::Button, FInputGesture() );
-		UI_COMMAND( HideAllSprites, "Hide All Sprites", "Hides all sprites", EUserInterfaceActionType::Button, FInputGesture() );
+		UI_COMMAND( ShowAllSprites, "Show All Sprites", "Shows all sprites", EUserInterfaceActionType::Button, FInputChord() );
+		UI_COMMAND( HideAllSprites, "Hide All Sprites", "Hides all sprites", EUserInterfaceActionType::Button, FInputChord() );
 
 		// get all the known layers
 		// Get a fresh list as GUnrealEd->SortedSpriteInfo may not yet be built.
@@ -196,7 +196,7 @@ void FLevelViewportCommands::RegisterCommands()
 
 	// Generate a command for each Stat category
 	{
-		UI_COMMAND(HideAllStats, "Hide All Stats", "Hides all Stats", EUserInterfaceActionType::Button, FInputGesture());
+		UI_COMMAND(HideAllStats, "Hide All Stats", "Hides all Stats", EUserInterfaceActionType::Button, FInputChord());
 
 		// Bind a listener here for any additional stat commands that get registered later.
 		UEngine::NewStatDelegate.AddRaw(this, &FLevelViewportCommands::HandleNewStat);
@@ -228,7 +228,7 @@ void FLevelViewportCommands::RegisterCommands()
 			FText::Format( NSLOCTEXT("LevelEditorCommands", "JumpToBookmark", "Jump to Bookmark {0}"), FText::AsNumber( BookmarkIndex ) ), //Localized label
 			FText::Format( NSLOCTEXT("LevelEditorCommands", "JumpToBookmark_ToolTip", "Moves the viewport to the location and orientation stored at bookmark {0}"), FText::AsNumber( BookmarkIndex ) ) )//Localized tooltip
 			.UserInterfaceType( EUserInterfaceActionType::Button ) //interface type
-			.DefaultGesture( FInputGesture( NumberKeyNames.IsValidIndex( BookmarkIndex ) ? NumberKeyNames[BookmarkIndex] : EKeys::Invalid ) ); //default gesture
+			.DefaultChord( FInputChord( NumberKeyNames.IsValidIndex( BookmarkIndex ) ? NumberKeyNames[BookmarkIndex] : EKeys::Invalid ) ); //default chord
 
 		JumpToBookmarkCommands.Add( JumpToBookmark );
 
@@ -239,7 +239,7 @@ void FLevelViewportCommands::RegisterCommands()
 			FText::Format( NSLOCTEXT("LevelEditorCommands", "SetBookmark", "Set Bookmark {0}"), FText::AsNumber( BookmarkIndex ) ), //Localized label
 			FText::Format( NSLOCTEXT("LevelEditorCommands", "SetBookmark_ToolTip", "Stores the viewports location and orientation in bookmark {0}"), FText::AsNumber( BookmarkIndex ) ) )//Localized tooltip
 			.UserInterfaceType( EUserInterfaceActionType::Button ) //interface type
-			.DefaultGesture( FInputGesture( EModifierKey::Control, NumberKeyNames.IsValidIndex( BookmarkIndex ) ? NumberKeyNames[BookmarkIndex] : EKeys::Invalid ) ); //default gesture
+			.DefaultChord( FInputChord( EModifierKey::Control, NumberKeyNames.IsValidIndex( BookmarkIndex ) ? NumberKeyNames[BookmarkIndex] : EKeys::Invalid ) ); //default chord
 
 		SetBookmarkCommands.Add( SetBookmark );
 
@@ -250,14 +250,14 @@ void FLevelViewportCommands::RegisterCommands()
 			FText::Format( NSLOCTEXT("LevelEditorCommands", "ClearBookmark", "Clear Bookmark {0}"), FText::AsNumber( BookmarkIndex ) ), //Localized label
 			FText::Format( NSLOCTEXT("LevelEditorCommands", "ClearBookmark_ToolTip", "Clears the viewports location and orientation in bookmark {0}"), FText::AsNumber( BookmarkIndex ) ) )//Localized tooltip
 			.UserInterfaceType( EUserInterfaceActionType::Button ) //interface type
-			.DefaultGesture( FInputGesture() ); //default gesture 
+			.DefaultChord( FInputChord() ); //default chord 
 
 		ClearBookmarkCommands.Add( ClearBookMark );
 	}
-	UI_COMMAND( ClearAllBookMarks, "Clear All Bookmarks", "Clears all the bookmarks", EUserInterfaceActionType::Button, FInputGesture() );
+	UI_COMMAND( ClearAllBookMarks, "Clear All Bookmarks", "Clears all the bookmarks", EUserInterfaceActionType::Button, FInputChord() );
 
-	UI_COMMAND( EnablePreviewMesh, "Hold To Enable Preview Mesh", "When held down a preview mesh appears under the cursor", EUserInterfaceActionType::Button, FInputGesture(EKeys::Backslash) );
-	UI_COMMAND( CyclePreviewMesh, "Cycles Preview Mesh", "Cycles available preview meshes", EUserInterfaceActionType::Button, FInputGesture( EModifierKey::Shift, EKeys::Backslash ) );
+	UI_COMMAND( EnablePreviewMesh, "Hold To Enable Preview Mesh", "When held down a preview mesh appears under the cursor", EUserInterfaceActionType::Button, FInputChord(EKeys::Backslash) );
+	UI_COMMAND( CyclePreviewMesh, "Cycles Preview Mesh", "Cycles available preview meshes", EUserInterfaceActionType::Button, FInputChord( EModifierKey::Shift, EKeys::Backslash ) );
 
 }
 
diff --git a/Engine/Source/Editor/LevelEditor/Private/SLevelViewport.cpp b/Engine/Source/Editor/LevelEditor/Private/SLevelViewport.cpp
--- a/Engine/Source/Editor/LevelEditor/Private/SLevelViewport.cpp
+++ b/Engine/Source/Editor/LevelEditor/Private/SLevelViewport.cpp
@@ -3341,21 +3341,21 @@ FText SLevelViewport::GetMouseCaptureLabelText() const
 	{
 		// Only need to fetch the exec binding once since it can't change
 		// but better than hard coding to the current binding.
-		static FInputGesture Gesture(EKeys::F1, EModifierKey::Shift);
-		static bool bInitedGesture = false;
-		if(!bInitedGesture)
+		static FInputChord Chord(EKeys::F1, EModifierKey::Shift);
+		static bool bInitedChord = false;
+		if(!bInitedChord)
 		{
 			ULocalPlayer* const TargetPlayer = GEngine->GetLocalPlayerFromControllerId(GetWorld(), 0);
 			if (TargetPlayer && TargetPlayer->PlayerController && TargetPlayer->PlayerController->PlayerInput)
 			{
 				FKeyBind Binding = TargetPlayer->PlayerController->PlayerInput->GetExecBind(TEXT("ShowMouseCursor"));
-				Gesture = FInputGesture(Binding.Key, EModifierKey::FromBools(Binding.Control, Binding.Alt, Binding.Shift, Binding.Cmd));
-				bInitedGesture = true;
+				Chord = FInputChord(Binding.Key, EModifierKey::FromBools(Binding.Control, Binding.Alt, Binding.Shift, Binding.Cmd));
+				bInitedChord = true;
 			}
 		}
 
 		FFormatNamedArguments Args;
-		Args.Add( TEXT("InputText"), Gesture.GetInputText() );
+		Args.Add( TEXT("InputText"), Chord.GetInputText() );
 		return FText::Format( LOCTEXT("ShowMouseCursorLabel", "{InputText} for Mouse Cursor"), Args );
 	}
 	else
