Editor sessions now remember the view modes set for perspective / ortho viewports.

diff --git a/Engine/Source/Editor/LevelEditor/Private/SLevelViewport.cpp b/Engine/Source/Editor/LevelEditor/Private/SLevelViewport.cpp
--- a/Engine/Source/Editor/LevelEditor/Private/SLevelViewport.cpp
+++ b/Engine/Source/Editor/LevelEditor/Private/SLevelViewport.cpp
@@ -3440,9 +3440,6 @@ void SLevelViewport::ResetNewLevelViewFlags()
 {
 	ChangeExposureSetting(FEditorViewportCommands::AutoExposureRadioID);
 
-	const bool bIsPerspective = LevelViewportClient->IsPerspective();
-	LevelViewportClient->SetViewMode(bIsPerspective ? FEditorViewportClient::DefaultPerspectiveViewMode : FEditorViewportClient::DefaultOrthoViewMode);
-
 	const bool bUseSavedDefaults = true;
 	OnUseDefaultShowFlags(bUseSavedDefaults);
 }
