PR #868: Add "Snap Object To View" to Editor (Contributed by [REDACTED_PERSON])

This allows you to align your view in the editor, select an actor in the
world outliner or in scene, right-click and select "Snap Actor to
Camera".  Actor will then set its location and rotation to match the
view.  This is super handy for quickly finding an awesome camera view in
the editor and then making your CameraActors match that view.

diff --git a/Engine/Source/Editor/LevelEditor/Private/LevelEditor.cpp b/Engine/Source/Editor/LevelEditor/Private/LevelEditor.cpp
--- a/Engine/Source/Editor/LevelEditor/Private/LevelEditor.cpp
+++ b/Engine/Source/Editor/LevelEditor/Private/LevelEditor.cpp
@@ -629,7 +629,12 @@ void FLevelEditorModule::BindGlobalLevelEditorCommands()
 		FExecuteAction::CreateStatic( &FLevelEditorActionCallbacks::ExecuteExecCommand, FString( TEXT("CAMERA SNAP") ) )
 		);
 
-	ActionList.MapAction( 
+	ActionList.MapAction(
+		Commands.SnapObjectToCamera,
+		FExecuteAction::CreateStatic(&FLevelEditorActionCallbacks::SnapObjectToView_Clicked)
+		);
+
+	ActionList.MapAction(
 		Commands.GoToCodeForActor, 
 		FExecuteAction::CreateStatic( &FLevelEditorActionCallbacks::GoToCodeForActor_Clicked )
 		);
diff --git a/Engine/Source/Editor/LevelEditor/Private/LevelEditorActions.cpp b/Engine/Source/Editor/LevelEditor/Private/LevelEditorActions.cpp
--- a/Engine/Source/Editor/LevelEditor/Private/LevelEditorActions.cpp
+++ b/Engine/Source/Editor/LevelEditor/Private/LevelEditorActions.cpp
@@ -2324,6 +2324,20 @@ void FLevelEditorActionCallbacks::OnAudioMutedChanged(bool bMuted)
 	GEditor->MuteRealTimeAudio(bMuted);
 }
 
+void FLevelEditorActionCallbacks::SnapObjectToView_Clicked()
+{
+	for (FSelectionIterator It(GEditor->GetSelectedActorIterator()); It; ++It)
+	{
+		AActor* Actor = Cast<AActor>(*It);
+		FVector location = GCurrentLevelEditingViewportClient->GetViewLocation();
+		FRotator rotation = GCurrentLevelEditingViewportClient->GetViewRotation();
+
+		Actor->SetActorLocation(location);
+		Actor->SetActorRotation(rotation);
+	}
+
+}
+
 void FLevelEditorActionCallbacks::OnEnableActorSnap()
 {
 	FSnappingUtils::EnableActorSnap( !FSnappingUtils::IsSnapToActorEnabled() );
@@ -2859,6 +2873,7 @@ void FLevelEditorCommands::RegisterCommands()
 	UI_COMMAND( GoHere, "Go Here", "Moves the camera to the current mouse position", EUserInterfaceActionType::Button, FInputChord() );
 
 	UI_COMMAND( SnapCameraToObject, "Snap View to Object", "Snaps the view to the selected object", EUserInterfaceActionType::Button, FInputChord() );
+	UI_COMMAND( SnapObjectToCamera, "Snap Object to View", "Snaps the selected object to the view", EUserInterfaceActionType::Button, FInputChord() );
 
 	UI_COMMAND( GoToCodeForActor, "Go to C++ Code for Actor", "Opens a code editing IDE and navigates to the source file associated with the seleced actor", EUserInterfaceActionType::Button, FInputChord() );
 	UI_COMMAND( GoToDocsForActor, "Go to Documentation for Actor", "Opens documentation for the Actor in the default web browser", EUserInterfaceActionType::Button, FInputChord() );
diff --git a/Engine/Source/Editor/LevelEditor/Private/LevelEditorContextMenu.cpp b/Engine/Source/Editor/LevelEditor/Private/LevelEditorContextMenu.cpp
--- a/Engine/Source/Editor/LevelEditor/Private/LevelEditorContextMenu.cpp
+++ b/Engine/Source/Editor/LevelEditor/Private/LevelEditorContextMenu.cpp
@@ -185,6 +185,7 @@ void FLevelEditorContextMenu::FillMenu( FMenuBuilder& MenuBuilder, TWeakPtr<SLev
 
 			MenuBuilder.AddMenuEntry(FLevelEditorCommands::Get().GoHere);
 			MenuBuilder.AddMenuEntry(FLevelEditorCommands::Get().SnapCameraToObject);
+			MenuBuilder.AddMenuEntry(FLevelEditorCommands::Get().SnapObjectToCamera);
 		}
 		MenuBuilder.EndSection();
 
@@ -276,6 +277,7 @@ void FLevelEditorContextMenu::FillMenu( FMenuBuilder& MenuBuilder, TWeakPtr<SLev
 
 			MenuBuilder.AddMenuEntry(FLevelEditorCommands::Get().GoHere);
 			MenuBuilder.AddMenuEntry(FLevelEditorCommands::Get().SnapCameraToObject);
+			MenuBuilder.AddMenuEntry(FLevelEditorCommands::Get().SnapObjectToCamera);
 
 			if (SelectedActors.Num() == 1)
 			{
diff --git a/Engine/Source/Editor/LevelEditor/Public/LevelEditorActions.h b/Engine/Source/Editor/LevelEditor/Public/LevelEditorActions.h
--- a/Engine/Source/Editor/LevelEditor/Public/LevelEditorActions.h
+++ b/Engine/Source/Editor/LevelEditor/Public/LevelEditorActions.h
@@ -113,6 +113,9 @@ public:
 	/** Snaps the camera to the selected object. */
 	TSharedPtr< FUICommandInfo > SnapCameraToObject;
 
+	/** Snaps the selected actor to the camera. */
+	TSharedPtr< FUICommandInfo > SnapObjectToCamera;
+
 	/** Goes to the source code for the selected actor's class. */
 	TSharedPtr< FUICommandInfo > GoToCodeForActor;
 
@@ -1161,6 +1164,11 @@ public:
 	 */
 	static void MoveActorToGrid_Clicked( bool InAlign, bool bInPerActor );
 
+	/**
+	* Snaps a selected actor to the camera view.
+	*/
+	static void SnapObjectToView_Clicked();
+
 	/** 
 	 * Moves an actor to another actor.
 	 */
