No longer focus the level viewport when PIE ends.  It causes annoyances when PIE is ended from other editors such as blueprints are umg.

diff --git a/Engine/Source/Editor/LevelEditor/Private/SLevelViewport.cpp b/Engine/Source/Editor/LevelEditor/Private/SLevelViewport.cpp
--- a/Engine/Source/Editor/LevelEditor/Private/SLevelViewport.cpp
+++ b/Engine/Source/Editor/LevelEditor/Private/SLevelViewport.cpp
@@ -3491,8 +3491,6 @@ void SLevelViewport::EndPlayInEditorSession()
 		HideMouseCaptureLabel();
 	}
 
-	FSlateApplication::Get().SetKeyboardFocus( ViewportWidget );
-
 	// Kick off a quick transition effect (border graphics)
 	ViewTransitionType = EViewTransition::ReturningToEditor;
 	ViewTransitionAnim = FCurveSequence( 0.0f, 1.5f, ECurveEaseFunction::CubicOut );
