Editor: Passed in a host command list in SActorDetails uses

diff --git a/Engine/Source/Editor/LevelEditor/Private/SActorDetails.cpp b/Engine/Source/Editor/LevelEditor/Private/SActorDetails.cpp
--- a/Engine/Source/Editor/LevelEditor/Private/SActorDetails.cpp
+++ b/Engine/Source/Editor/LevelEditor/Private/SActorDetails.cpp
@@ -60,7 +60,7 @@ public:
 	}
 };
 
-void SActorDetails::Construct(const FArguments& InArgs, const FName TabIdentifier)
+void SActorDetails::Construct(const FArguments& InArgs, const FName TabIdentifier, TSharedPtr<FUICommandList> InCommandList)
 {
 	bSelectionGuard = false;
 	bShowingRootActorNodeSelected = false;
@@ -80,7 +80,7 @@ void SActorDetails::Construct(const FArguments& InArgs, const FName TabIdentifie
 	DetailsViewArgs.bCustomNameAreaLocation = true;
 	DetailsViewArgs.bCustomFilterAreaLocation = true;
 	DetailsViewArgs.DefaultsOnlyVisibility = FDetailsViewArgs::EEditDefaultsOnlyNodeVisibility::Hide;
-
+	DetailsViewArgs.HostCommandList = InCommandList;
 	DetailsView = PropPlugin.CreateDetailView(DetailsViewArgs);
 
 	auto IsPropertyVisible = [](const FPropertyAndParent& PropertyAndParent)
diff --git a/Engine/Source/Editor/LevelEditor/Private/SActorDetails.h b/Engine/Source/Editor/LevelEditor/Private/SActorDetails.h
--- a/Engine/Source/Editor/LevelEditor/Private/SActorDetails.h
+++ b/Engine/Source/Editor/LevelEditor/Private/SActorDetails.h
@@ -11,7 +11,7 @@ public:
 	SLATE_BEGIN_ARGS(SActorDetails) {}
 	SLATE_END_ARGS()
 
-	void Construct(const FArguments& InArgs, const FName TabIdentifier);
+	void Construct(const FArguments& InArgs, const FName TabIdentifier, TSharedPtr<FUICommandList> InCommandList);
 	~SActorDetails();
 
 	/**
diff --git a/Engine/Source/Editor/LevelEditor/Private/SLevelEditor.cpp b/Engine/Source/Editor/LevelEditor/Private/SLevelEditor.cpp
--- a/Engine/Source/Editor/LevelEditor/Private/SLevelEditor.cpp
+++ b/Engine/Source/Editor/LevelEditor/Private/SLevelEditor.cpp
@@ -513,9 +513,9 @@ TSharedRef<FTabManager> SLevelEditor::GetTabManager() const
 
 TSharedRef<SDockTab> SLevelEditor::SummonDetailsPanel( FName TabIdentifier )
 {
-	TSharedPtr<SActorDetails> ActorDetails;
+	TSharedRef<SActorDetails> ActorDetails = SNew(SActorDetails, TabIdentifier, LevelEditorCommands);
 
-	FText Label = NSLOCTEXT( "LevelEditor", "DetailsTabTitle", "Details" );
+	const FText Label = NSLOCTEXT( "LevelEditor", "DetailsTabTitle", "Details" );
 
 	TSharedRef<SDockTab> DocTab = SNew(SDockTab)
 		.Icon( FEditorStyle::GetBrush( "LevelEditor.Tabs.Details" ) )
@@ -525,7 +525,7 @@ TSharedRef<SDockTab> SLevelEditor::SummonDetailsPanel( FName TabIdentifier )
 			SNew( SBox )
 			.AddMetaData<FTutorialMetaData>(FTutorialMetaData(TEXT("ActorDetails"), TEXT("LevelEditorSelectionDetails")))
 			[
-				SAssignNew( ActorDetails, SActorDetails, TabIdentifier )
+				ActorDetails
 			]
 		];
 
