Removed obsolete "Levels" module

diff --git a/Engine/Build/JunkManifest.txt b/Engine/Build/JunkManifest.txt
--- a/Engine/Build/JunkManifest.txt
+++ b/Engine/Build/JunkManifest.txt
@@ -18,6 +18,7 @@ Engine\Binaries\Win64\UE4Editor-LinuxTarget*.*
 Engine\Binaries\Win64\UE4Editor-AndroidTarget*.*
 Engine\Binaries\Win64\UE4Editor-IOSPlatformEditor*.*
 Engine\Binaries\Win64\UE4Editor-TaskBrowser*.*
+Engine\Binaries\Win64\UE4Editor-Levels*.*
 Engine\Binaries\Win64\*-ProjectWorkflows.*
 Engine\Binaries\Win64\*-WorkflowEditor.*
 Engine\Binaries\Win64\*-Workflows.*
@@ -37,6 +38,7 @@ Engine\Binaries\Mac\UE4Editor.app\Contents\MacOS\*-CameraDebugger.*
 Engine\Binaries\Mac\UE4Editor.app\Contents\MacOS\*-CameraEditor.*
 Engine\Binaries\Mac\UE4Editor.app\Contents\MacOS\*-SessionLauncher.*
 Engine\Binaries\Mac\UE4Editor.app\Contents\MacOS\*-TaskBrowser.*
+Engine\Binaries\Mac\UE4Editor.app\Contents\MacOS\*-Levels.*
 Engine\Binaries\Mac\*-MessagingDebugger.*
 Engine\Binaries\Win64\*-DDCUtils*.*
 Engine\Binaries\Win64\*-MessagingDebugger.*
diff --git a/Engine/Source/Editor/LevelEditor/LevelEditor.Build.cs b/Engine/Source/Editor/LevelEditor/LevelEditor.Build.cs
--- a/Engine/Source/Editor/LevelEditor/LevelEditor.Build.cs
+++ b/Engine/Source/Editor/LevelEditor/LevelEditor.Build.cs
@@ -72,7 +72,6 @@ public class LevelEditor : ModuleRules
 				"SessionFrontend",
 				"AutomationWindow",
 				"Layers",
-				"Levels",
                 "WorldBrowser",
 				"EditorWidgets",
 				"AssetTools",
diff --git a/Engine/Source/Editor/LevelEditor/Private/SLevelEditor.cpp b/Engine/Source/Editor/LevelEditor/Private/SLevelEditor.cpp
--- a/Engine/Source/Editor/LevelEditor/Private/SLevelEditor.cpp
+++ b/Engine/Source/Editor/LevelEditor/Private/SLevelEditor.cpp
@@ -19,7 +19,6 @@
 #include "Editor/UnrealEd/Public/Kismet2/DebuggerCommands.h"
 #include "Editor/SceneOutliner/Public/SceneOutliner.h"
 #include "Editor/Layers/Public/LayersModule.h"
-#include "Editor/Levels/Public/LevelsModule.h"
 #include "Editor/WorldBrowser/Public/WorldBrowserModule.h"
 #include "Editor/ClassViewer/Public/ClassViewerModule.h"
 #include "Toolkits/IToolkit.h"
@@ -670,34 +669,15 @@ TSharedRef<SDockTab> SLevelEditor::SpawnLevelEditorTab( const FSpawnTabArgs& Arg
 	}
 	else if( TabIdentifier == WorldBrowserHierarchyTab || TabIdentifier == TEXT("LevelEditorLevelBrowser") )
 	{
-		if (FParse::Param(FCommandLine::Get(), TEXT("oldlevels")))
-		{
-			FLevelsModule& LevelsModule = FModuleManager::LoadModuleChecked<FLevelsModule>( "Levels" );
-			return SNew( SDockTab )
-				.Icon( FEditorStyle::GetBrush( "LevelEditor.Tabs.Levels" ) )
-				.Label( NSLOCTEXT("LevelEditor", "LevelsTabTitle", "Levels") )
-				[
-					SNew(SBorder)
-					.Padding( 0 )
-					.BorderImage( FEditorStyle::GetBrush("ToolPanel.GroupBorder") )
-					[
-						LevelsModule.CreateLevelBrowser()
-					]
-				];
-		}
-		else
-		{
-			FWorldBrowserModule& WorldBrowserModule = FModuleManager::LoadModuleChecked<FWorldBrowserModule>( "WorldBrowser" );
-			return SNew( SDockTab )
-				.Icon( FEditorStyle::GetBrush( "LevelEditor.Tabs.WorldBrowser" ) )
-				.Label( NSLOCTEXT("LevelEditor", "WorldBrowserHierarchyTabTitle", "Levels") )
-				[
-					WorldBrowserModule.CreateWorldBrowserHierarchy()
-				];
-		}
-				
+		FWorldBrowserModule& WorldBrowserModule = FModuleManager::LoadModuleChecked<FWorldBrowserModule>( "WorldBrowser" );
+		return SNew( SDockTab )
+			.Icon( FEditorStyle::GetBrush( "LevelEditor.Tabs.WorldBrowser" ) )
+			.Label( NSLOCTEXT("LevelEditor", "WorldBrowserHierarchyTabTitle", "Levels") )
+			[
+				WorldBrowserModule.CreateWorldBrowserHierarchy()
+			];
 	}
-	else if( TabIdentifier == WorldBrowserDetailsTab && !FParse::Param(FCommandLine::Get(), TEXT("oldlevels")) )
+	else if( TabIdentifier == WorldBrowserDetailsTab )
 	{
 		FWorldBrowserModule& WorldBrowserModule = FModuleManager::LoadModuleChecked<FWorldBrowserModule>( "WorldBrowser" );
 		return SNew( SDockTab )
@@ -707,7 +687,7 @@ TSharedRef<SDockTab> SLevelEditor::SpawnLevelEditorTab( const FSpawnTabArgs& Arg
 				WorldBrowserModule.CreateWorldBrowserDetails()
 			];
 	}
-	else if( TabIdentifier == WorldBrowserCompositionTab && !FParse::Param(FCommandLine::Get(), TEXT("oldlevels")) )
+	else if( TabIdentifier == WorldBrowserCompositionTab )
 	{
 		FWorldBrowserModule& WorldBrowserModule = FModuleManager::LoadModuleChecked<FWorldBrowserModule>( "WorldBrowser" );
 		return SNew( SDockTab )
@@ -1020,16 +1000,6 @@ TSharedRef<SWidget> SLevelEditor::RestoreContentArea( const TSharedRef<SDockTab>
 				.SetIcon( LayersIcon );
 		}
 		
-		if (FParse::Param(FCommandLine::Get(), TEXT("oldlevels")))
-		{
-			const FSlateIcon LevelsIcon(FEditorStyle::GetStyleSetName(), "LevelEditor.Tabs.Levels");
-			LevelEditorTabManager->RegisterTabSpawner( "LevelEditorLevelBrowser", FOnSpawnTab::CreateSP<SLevelEditor, FName, FString>(this, &SLevelEditor::SpawnLevelEditorTab, FName("LevelEditorLevelBrowser"), FString()) )
-				.SetDisplayName(NSLOCTEXT("LevelEditorTabs", "LevelEditorLevelBrowser", "Levels"))
-				.SetTooltipText(NSLOCTEXT("LevelEditorTabs", "LevelEditorLevelBrowserTooltipText", "Open the Levels tab. Use this to manage the levels in the current project."))
-				.SetGroup( MenuStructure.GetLevelEditorCategory() )
-				.SetIcon( LevelsIcon );
-		}
-		else
 		{
 			LevelEditorTabManager->RegisterTabSpawner( WorldBrowserHierarchyTab, FOnSpawnTab::CreateSP<SLevelEditor, FName, FString>(this, &SLevelEditor::SpawnLevelEditorTab, WorldBrowserHierarchyTab, FString()) )
 				.SetDisplayName(NSLOCTEXT("LevelEditorTabs", "WorldBrowserHierarchy", "Levels"))
diff --git a/Engine/Source/Editor/Levels/Levels.Build.cs b/Engine/Source/Editor/Levels/Levels.Build.cs
deleted file mode 100644
diff --git a/Engine/Source/Editor/Levels/Private/ActorsAssignedToSpecificLevelsFilter.h b/Engine/Source/Editor/Levels/Private/ActorsAssignedToSpecificLevelsFilter.h
deleted file mode 100644
diff --git a/Engine/Source/Editor/Levels/Private/LevelCollectionViewCommands.h b/Engine/Source/Editor/Levels/Private/LevelCollectionViewCommands.h
deleted file mode 100644
diff --git a/Engine/Source/Editor/Levels/Private/LevelCollectionViewModel.cpp b/Engine/Source/Editor/Levels/Private/LevelCollectionViewModel.cpp
deleted file mode 100644
diff --git a/Engine/Source/Editor/Levels/Private/LevelCollectionViewModel.h b/Engine/Source/Editor/Levels/Private/LevelCollectionViewModel.h
deleted file mode 100644
diff --git a/Engine/Source/Editor/Levels/Private/LevelEdMode.cpp b/Engine/Source/Editor/Levels/Private/LevelEdMode.cpp
deleted file mode 100644
diff --git a/Engine/Source/Editor/Levels/Private/LevelViewModel.cpp b/Engine/Source/Editor/Levels/Private/LevelViewModel.cpp
deleted file mode 100644
diff --git a/Engine/Source/Editor/Levels/Private/LevelViewModel.h b/Engine/Source/Editor/Levels/Private/LevelViewModel.h
deleted file mode 100644
diff --git a/Engine/Source/Editor/Levels/Private/LevelsModule.cpp b/Engine/Source/Editor/Levels/Private/LevelsModule.cpp
deleted file mode 100644
diff --git a/Engine/Source/Editor/Levels/Private/LevelsPrivatePCH.h b/Engine/Source/Editor/Levels/Private/LevelsPrivatePCH.h
deleted file mode 100644
diff --git a/Engine/Source/Editor/Levels/Private/SLevelBrowser.cpp b/Engine/Source/Editor/Levels/Private/SLevelBrowser.cpp
deleted file mode 100644
diff --git a/Engine/Source/Editor/Levels/Private/SLevelBrowser.h b/Engine/Source/Editor/Levels/Private/SLevelBrowser.h
deleted file mode 100644
diff --git a/Engine/Source/Editor/Levels/Private/SLevelsCommandsMenu.h b/Engine/Source/Editor/Levels/Private/SLevelsCommandsMenu.h
deleted file mode 100644
diff --git a/Engine/Source/Editor/Levels/Private/SLevelsView.h b/Engine/Source/Editor/Levels/Private/SLevelsView.h
deleted file mode 100644
diff --git a/Engine/Source/Editor/Levels/Private/SLevelsViewRow.cpp b/Engine/Source/Editor/Levels/Private/SLevelsViewRow.cpp
deleted file mode 100644
diff --git a/Engine/Source/Editor/Levels/Private/SLevelsViewRow.h b/Engine/Source/Editor/Levels/Private/SLevelsViewRow.h
deleted file mode 100644
diff --git a/Engine/Source/Editor/Levels/Public/LevelBrowserSettings.h b/Engine/Source/Editor/Levels/Public/LevelBrowserSettings.h
deleted file mode 100644
diff --git a/Engine/Source/Editor/Levels/Public/LevelEdMode.h b/Engine/Source/Editor/Levels/Public/LevelEdMode.h
deleted file mode 100644
diff --git a/Engine/Source/Editor/Levels/Public/LevelsModule.h b/Engine/Source/Editor/Levels/Public/LevelsModule.h
deleted file mode 100644
diff --git a/Engine/Source/Editor/UnrealEd/Private/EditorEngine.cpp b/Engine/Source/Editor/UnrealEd/Private/EditorEngine.cpp
--- a/Engine/Source/Editor/UnrealEd/Private/EditorEngine.cpp
+++ b/Engine/Source/Editor/UnrealEd/Private/EditorEngine.cpp
@@ -2711,14 +2711,7 @@ void UEditorEngine::SelectLevelInLevelBrowser( bool bDeselectOthers )
 	}
 
 	FLevelEditorModule& LevelEditorModule = FModuleManager::LoadModuleChecked<FLevelEditorModule>("LevelEditor");
-	if (FParse::Param(FCommandLine::Get(), TEXT("oldlevels")))
-	{
-		LevelEditorModule.SummonLevelBrowser();
-	}
-	else
-	{
-		LevelEditorModule.SummonWorldBrowserHierarchy();
-	}
+	LevelEditorModule.SummonWorldBrowserHierarchy();
 }
 
 void UEditorEngine::DeselectLevelInLevelBrowser()
@@ -2733,14 +2726,7 @@ void UEditorEngine::DeselectLevelInLevelBrowser()
 	}
 	
 	FLevelEditorModule& LevelEditorModule = FModuleManager::LoadModuleChecked<FLevelEditorModule>("LevelEditor");
-	if (FParse::Param(FCommandLine::Get(), TEXT("oldlevels")))
-	{
-		LevelEditorModule.SummonLevelBrowser();
-	}
-	else
-	{
-		LevelEditorModule.SummonWorldBrowserHierarchy();
-	}
+	LevelEditorModule.SummonWorldBrowserHierarchy();
 }
 
 void UEditorEngine::SelectAllActorsControlledByMatinee()
diff --git a/Engine/Source/Editor/UnrealEd/Private/EditorLevelUtils.cpp b/Engine/Source/Editor/UnrealEd/Private/EditorLevelUtils.cpp
--- a/Engine/Source/Editor/UnrealEd/Private/EditorLevelUtils.cpp
+++ b/Engine/Source/Editor/UnrealEd/Private/EditorLevelUtils.cpp
@@ -15,7 +15,6 @@
 #include "LevelEditor.h"
 #include "ScopedTransaction.h"
 #include "ActorEditorUtils.h"
-#include "Editor/Levels/Public/LevelEdMode.h"
 #include "ContentStreaming.h"
 #include "PackageTools.h"
 
@@ -388,17 +387,6 @@ namespace EditorLevelUtils
 	
 	bool RemoveLevelFromWorld(ULevel* InLevel)
 	{
-		// If we're removing a level, lets make sure to close the level transform mode if its the same level currently selected for edit
-		FEdModeLevel* LevelMode = static_cast<FEdModeLevel*>(GLevelEditorModeTools().GetActiveMode( FBuiltinEditorModes::EM_Level ));		
-		if( LevelMode )
-		{
-			ULevelStreaming* LevelStream = FLevelUtils::FindStreamingLevel( InLevel ); 
-			if( LevelMode->IsEditing( LevelStream ) )
-			{
-				GLevelEditorModeTools().DeactivateMode( FBuiltinEditorModes::EM_Level );
-			}
-		}
-
 		if (GEditor->Layers.IsValid())
 		{
 			GEditor->Layers->RemoveLevelLayerInformation( InLevel );
@@ -709,13 +697,6 @@ namespace EditorLevelUtils
 				StreamingLevel = FLevelUtils::FindStreamingLevel( Level );
 			}
 
-			// If were hiding a level, lets make sure to close the level transform mode if its the same level currently selected for edit
-			FEdModeLevel* LevelMode = static_cast<FEdModeLevel*>(GLevelEditorModeTools().GetActiveMode( FBuiltinEditorModes::EM_Level ));
-			if( LevelMode && LevelMode->IsEditing( StreamingLevel ) )
-			{
-				GLevelEditorModeTools().DeactivateMode( FBuiltinEditorModes::EM_Level );
-			}
-
 			//create a transaction so we can undo the visibilty toggle
 			const FScopedTransaction Transaction( LOCTEXT( "ToggleLevelVisibility", "Toggle Level Visibility" ) );		
 
diff --git a/Engine/Source/Editor/WorldBrowser/Private/LevelModel.cpp b/Engine/Source/Editor/WorldBrowser/Private/LevelModel.cpp
--- a/Engine/Source/Editor/WorldBrowser/Private/LevelModel.cpp
+++ b/Engine/Source/Editor/WorldBrowser/Private/LevelModel.cpp
@@ -5,7 +5,6 @@
 #include "Toolkits/AssetEditorManager.h"
 #include "LevelModel.h"
 #include "LevelCollectionModel.h"
-#include "LevelBrowserSettings.h"
 #include "Engine/Selection.h"
 #include "Engine/LevelStreaming.h"
 #include "AssetRegistryModule.h"
@@ -744,14 +743,14 @@ FString FLevelModel::GetLightmassSizeString() const
 	FString MemorySizeString;
 	ULevel* Level = GetLevelObject();
 
-	if (Level && GetDefault<ULevelBrowserSettings>()->bDisplayLightmassSize)
-	{
-		// Update metrics
-		static const float ByteConversion = 1.0f / 1024.0f;
-		float LightmapSize = Level->LightmapTotalSize * ByteConversion;
-		
-		MemorySizeString += FString::Printf(TEXT( "%.2f" ), LightmapSize);
-	}
+	//if (Level && GetDefault<ULevelBrowserSettings>()->bDisplayLightmassSize)
+	//{
+	//	// Update metrics
+	//	static const float ByteConversion = 1.0f / 1024.0f;
+	//	float LightmapSize = Level->LightmapTotalSize * ByteConversion;
+	//	
+	//	MemorySizeString += FString::Printf(TEXT( "%.2f" ), LightmapSize);
+	//}
 
 	return MemorySizeString;
 }
@@ -761,14 +760,14 @@ FString FLevelModel::GetFileSizeString() const
 	FString MemorySizeString;
 	ULevel* Level = GetLevelObject();
 
-	if (Level && GetDefault<ULevelBrowserSettings>()->bDisplayFileSize)
-	{
-		// Update metrics
-		static const float ByteConversion = 1.0f / 1024.0f;
-		float FileSize = Level->GetOutermost()->GetFileSize() * ByteConversion * ByteConversion;
-		
-		MemorySizeString += FString::Printf(TEXT( "%.2f" ), FileSize);
-	}
+	//if (Level && GetDefault<ULevelBrowserSettings>()->bDisplayFileSize)
+	//{
+	//	// Update metrics
+	//	static const float ByteConversion = 1.0f / 1024.0f;
+	//	float FileSize = Level->GetOutermost()->GetFileSize() * ByteConversion * ByteConversion;
+	//	
+	//	MemorySizeString += FString::Printf(TEXT( "%.2f" ), FileSize);
+	//}
 
 	return MemorySizeString;
 }
diff --git a/Engine/Source/Editor/WorldBrowser/WorldBrowser.Build.cs b/Engine/Source/Editor/WorldBrowser/WorldBrowser.Build.cs
--- a/Engine/Source/Editor/WorldBrowser/WorldBrowser.Build.cs
+++ b/Engine/Source/Editor/WorldBrowser/WorldBrowser.Build.cs
@@ -41,7 +41,6 @@ public class WorldBrowser : ModuleRules
                 "RawMesh",
                 "MeshUtilities",
                 "LandscapeEditor",
-				"Levels",
                 "ImageWrapper",
                 "Foliage",
             }
