Mark actors dirty when modifying GameMode properties via the level editor toolbar

diff --git a/Engine/Source/Editor/LevelEditor/Private/LevelEditorToolBar.cpp b/Engine/Source/Editor/LevelEditor/Private/LevelEditorToolBar.cpp
--- a/Engine/Source/Editor/LevelEditor/Private/LevelEditorToolBar.cpp
+++ b/Engine/Source/Editor/LevelEditor/Private/LevelEditorToolBar.cpp
@@ -814,7 +814,11 @@ void LevelEditorActionHelpers::OnSelectGameModeClassPicked(UClass* InChosenClass
 	}
 	else
 	{
-		InLevelEditor.Pin()->GetWorld()->GetWorldSettings()->DefaultGameMode = InChosenClass;
+		const FScopedTransaction Transaction( NSLOCTEXT("LevelEditorCommands", "SelectGameModeClassAction", "Set Override Game Mode Class") );
+
+		AWorldSettings* WorldSettings = InLevelEditor.Pin()->GetWorld()->GetWorldSettings();
+		WorldSettings->Modify();
+		WorldSettings->DefaultGameMode = InChosenClass;
 	}
 	FSlateApplication::Get().DismissAllMenus();
 }
@@ -875,7 +879,9 @@ void LevelEditorActionHelpers::OnSelectGameStateClassPicked(UClass* InChosenClas
 {
 	if(UClass* GameModeClass = GetGameModeClass(InLevelEditor, bInIsProjectSettings))
 	{
+		const FScopedTransaction Transaction( NSLOCTEXT("LevelEditorCommands", "SelectGameStateClassAction", "Set Game State Class") );
 		AGameMode* ActiveGameMode = Cast<AGameMode>(GameModeClass->GetDefaultObject());
+		ActiveGameMode->Modify();
 		ActiveGameMode->GameStateClass = InChosenClass;
 	}
 
@@ -939,7 +945,10 @@ void LevelEditorActionHelpers::OnSelectPawnClassPicked(UClass* InChosenClass, TW
 {
 	if(UClass* GameModeClass = GetGameModeClass(InLevelEditor, bInIsProjectSettings))
 	{
+		const FScopedTransaction Transaction( NSLOCTEXT("LevelEditorCommands", "SelectPawnClassAction", "Set Pawn Class") );
+
 		AGameMode* ActiveGameMode = Cast<AGameMode>(GameModeClass->GetDefaultObject());
+		ActiveGameMode->Modify();
 		ActiveGameMode->DefaultPawnClass = InChosenClass;
 	}
 
@@ -1002,7 +1011,10 @@ void LevelEditorActionHelpers::OnSelectHUDClassPicked(UClass* InChosenClass, TWe
 {
 	if(UClass* GameModeClass = GetGameModeClass(InLevelEditor, bInIsProjectSettings))
 	{
+		const FScopedTransaction Transaction( NSLOCTEXT("LevelEditorCommands", "SelectHUDClassAction", "Set HUD Class") );
+
 		AGameMode* ActiveGameMode = Cast<AGameMode>(GameModeClass->GetDefaultObject());
+		ActiveGameMode->Modify();
 		ActiveGameMode->HUDClass = InChosenClass;
 	}
 
@@ -1065,7 +1077,10 @@ void LevelEditorActionHelpers::OnSelectPlayerControllerClassPicked(UClass* InCho
 {
 	if(UClass* GameModeClass = GetGameModeClass(InLevelEditor, bInIsProjectSettings))
 	{
+		const FScopedTransaction Transaction( NSLOCTEXT("LevelEditorCommands", "SelectPlayerControllerClassAction", "Set Player Controller Class") );
+
 		AGameMode* ActiveGameMode = Cast<AGameMode>(GameModeClass->GetDefaultObject());
+		ActiveGameMode->Modify();
 		ActiveGameMode->PlayerControllerClass = InChosenClass;
 	}
 	FSlateApplication::Get().DismissAllMenus();
