Fixed non unity compilation.

diff --git a/Engine/Source/Editor/LevelEditor/Private/LevelEditorToolBar.cpp b/Engine/Source/Editor/LevelEditor/Private/LevelEditorToolBar.cpp
--- a/Engine/Source/Editor/LevelEditor/Private/LevelEditorToolBar.cpp
+++ b/Engine/Source/Editor/LevelEditor/Private/LevelEditorToolBar.cpp
@@ -35,7 +35,7 @@
 #include "GameFramework/WorldSettings.h"
 #include "EngineUtils.h"
 #include "GameMapsSettings.h"
-
+#include "ScopedTransaction.h"
 
 namespace LevelEditorActionHelpers
 {
