Safeguard against unreproducible crash on mac opening the favorites menu

diff --git a/Engine/Source/Editor/LevelEditor/Private/LevelEditorActions.cpp b/Engine/Source/Editor/LevelEditor/Private/LevelEditorActions.cpp
--- a/Engine/Source/Editor/LevelEditor/Private/LevelEditorActions.cpp
+++ b/Engine/Source/Editor/LevelEditor/Private/LevelEditorActions.cpp
@@ -309,11 +309,16 @@ void FLevelEditorActionCallbacks::RemoveFavorite( int32 FavoriteFileIndex )
 
 bool FLevelEditorActionCallbacks::ToggleFavorite_CanExecute()
 {
-	FString FileName;
-	const bool bMapFileExists = FPackageName::DoesPackageExist(GetWorld()->GetOutermost()->GetName(), NULL, &FileName);
-	
-	// Disable the favorites button if the map isn't associated to a file yet (new map, never before saved, etc.)
-	return bMapFileExists;
+	if( GetWorld() && GetWorld()->GetOutermost() )
+	{
+		FString FileName;
+		const bool bMapFileExists = FPackageName::DoesPackageExist(GetWorld()->GetOutermost()->GetName(), NULL, &FileName);
+
+		// Disable the favorites button if the map isn't associated to a file yet (new map, never before saved, etc.)
+		return bMapFileExists;
+	}
+
+	return false;
 }
 
 
