Removed defunct LightingTools menu

diff --git a/Engine/Source/Editor/LevelEditor/Private/LevelEditor.cpp b/Engine/Source/Editor/LevelEditor/Private/LevelEditor.cpp
--- a/Engine/Source/Editor/LevelEditor/Private/LevelEditor.cpp
+++ b/Engine/Source/Editor/LevelEditor/Private/LevelEditor.cpp
@@ -1378,32 +1378,6 @@ void FLevelEditorModule::BindGlobalLevelEditorCommands()
 		FCanExecuteAction(),
 		FIsActionChecked::CreateStatic( &FLevelEditorActionCallbacks::IsLightingQualityChecked, (ELightingBuildQuality)Quality_Preview) );
 
-	ActionList.MapAction( 
-		Commands.LightingTools_ShowBounds, 
-		FExecuteAction::CreateStatic( &FLevelEditorActionCallbacks::SetLightingToolShowBounds ),
-		FCanExecuteAction(),
-		FIsActionChecked::CreateStatic( &FLevelEditorActionCallbacks::IsLightingToolShowBoundsChecked ) );
-	ActionList.MapAction( 
-		Commands.LightingTools_ShowTraces, 
-		FExecuteAction::CreateStatic( &FLevelEditorActionCallbacks::SetLightingToolShowTraces ),
-		FCanExecuteAction(),
-		FIsActionChecked::CreateStatic( &FLevelEditorActionCallbacks::IsLightingToolShowTracesChecked ) );
-	ActionList.MapAction( 
-		Commands.LightingTools_ShowDirectOnly, 
-		FExecuteAction::CreateStatic( &FLevelEditorActionCallbacks::SetLightingToolShowDirectOnly ),
-		FCanExecuteAction(),
-		FIsActionChecked::CreateStatic( &FLevelEditorActionCallbacks::IsLightingToolShowDirectOnlyChecked ) );
-	ActionList.MapAction( 
-		Commands.LightingTools_ShowIndirectOnly, 
-		FExecuteAction::CreateStatic( &FLevelEditorActionCallbacks::SetLightingToolShowIndirectOnly ),
-		FCanExecuteAction(),
-		FIsActionChecked::CreateStatic( &FLevelEditorActionCallbacks::IsLightingToolShowIndirectOnlyChecked ) );
-	ActionList.MapAction( 
-		Commands.LightingTools_ShowIndirectSamples, 
-		FExecuteAction::CreateStatic( &FLevelEditorActionCallbacks::SetLightingToolShowIndirectSamples ),
-		FCanExecuteAction(),
-		FIsActionChecked::CreateStatic( &FLevelEditorActionCallbacks::IsLightingToolShowIndirectSamplesChecked ) );
-
 	ActionList.MapAction( 
 		Commands.LightingDensity_RenderGrayscale, 
 		FExecuteAction::CreateStatic( &FLevelEditorActionCallbacks::SetLightingDensityRenderGrayscale ),
diff --git a/Engine/Source/Editor/LevelEditor/Private/LevelEditorActions.cpp b/Engine/Source/Editor/LevelEditor/Private/LevelEditorActions.cpp
--- a/Engine/Source/Editor/LevelEditor/Private/LevelEditorActions.cpp
+++ b/Engine/Source/Editor/LevelEditor/Private/LevelEditorActions.cpp
@@ -23,7 +23,6 @@
 #include "Editor/NewLevelDialog/Public/NewLevelDialogModule.h"
 #include "DelegateFilter.h"
 #include "BlueprintUtilities.h"
-#include "LightingTools.h"
 #include "MRUFavoritesList.h"
 #include "Editor/SceneOutliner/Private/SSocketChooser.h"
 #include "SnappingUtils.h"
@@ -646,66 +645,6 @@ void FLevelEditorActionCallbacks::SetLightingQuality( ELightingBuildQuality NewQ
 	GConfig->SetInt(TEXT("LightingBuildOptions"), TEXT("QualityLevel"), (int32)NewQuality, GEditorPerProjectIni);
 }
 
-void FLevelEditorActionCallbacks::SetLightingToolShowBounds()
-{
-	FLightingToolsSettings& Settings = FLightingToolsSettings::Get();
-	Settings.bShowLightingBounds = !Settings.bShowLightingBounds;
-	Settings.ApplyToggle();
-}
-
-bool FLevelEditorActionCallbacks::IsLightingToolShowBoundsChecked()
-{
-	return FLightingToolsSettings::Get().bShowLightingBounds;
-}
-
-void FLevelEditorActionCallbacks::SetLightingToolShowTraces()
-{
-	FLightingToolsSettings& Settings = FLightingToolsSettings::Get();
-	Settings.bShowShadowTraces = !Settings.bShowShadowTraces;
-	Settings.ApplyToggle();
-}
-
-bool FLevelEditorActionCallbacks::IsLightingToolShowTracesChecked()
-{
-	return FLightingToolsSettings::Get().bShowShadowTraces;
-}
-
-void FLevelEditorActionCallbacks::SetLightingToolShowDirectOnly()
-{
-	FLightingToolsSettings& Settings = FLightingToolsSettings::Get();
-	Settings.bShowDirectOnly = !Settings.bShowDirectOnly;
-	Settings.ApplyToggle();
-}
-
-bool FLevelEditorActionCallbacks::IsLightingToolShowDirectOnlyChecked()
-{
-	return FLightingToolsSettings::Get().bShowDirectOnly;
-}
-
-void FLevelEditorActionCallbacks::SetLightingToolShowIndirectOnly()
-{
-	FLightingToolsSettings& Settings = FLightingToolsSettings::Get();
-	Settings.bShowIndirectOnly = !Settings.bShowIndirectOnly;
-	Settings.ApplyToggle();
-}
-
-bool FLevelEditorActionCallbacks::IsLightingToolShowIndirectOnlyChecked()
-{
-	return FLightingToolsSettings::Get().bShowIndirectOnly;
-}
-
-void FLevelEditorActionCallbacks::SetLightingToolShowIndirectSamples()
-{
-	FLightingToolsSettings& Settings = FLightingToolsSettings::Get();
-	Settings.bShowIndirectSamples = !Settings.bShowIndirectSamples;
-	Settings.ApplyToggle();
-}
-
-bool FLevelEditorActionCallbacks::IsLightingToolShowIndirectSamplesChecked()
-{
-	return FLightingToolsSettings::Get().bShowIndirectSamples;
-}
-
 float FLevelEditorActionCallbacks::GetLightingDensityIdeal()
 {
 	return ( GEngine->IdealLightMapDensity );
@@ -2891,11 +2830,6 @@ void FLevelEditorCommands::RegisterCommands()
 	UI_COMMAND( LightingQuality_High, "High", "Sets precomputed lighting quality to high quality", EUserInterfaceActionType::RadioButton, FInputChord() );
 	UI_COMMAND( LightingQuality_Medium, "Medium", "Sets precomputed lighting quality to medium quality", EUserInterfaceActionType::RadioButton, FInputChord() );
 	UI_COMMAND( LightingQuality_Preview, "Preview", "Sets precomputed lighting quality to preview quality (fastest computation time.)", EUserInterfaceActionType::RadioButton, FInputChord() );
-	UI_COMMAND( LightingTools_ShowBounds, "Show Lighting Bounds", "Draws a wireframe sphere for the lighting bounds of the Light Environment, and a point for the center, or multiple points if NumVolumeVisibilitySamples is greater than zero.", EUserInterfaceActionType::ToggleButton, FInputChord() );
-	UI_COMMAND( LightingTools_ShowTraces, "Show Shadow Traces", "Draws lines for any shadow rays that the Light Environment traced. The lines will be white if the light is not shadowed, and red otherwise.", EUserInterfaceActionType::ToggleButton, FInputChord() );
-	UI_COMMAND( LightingTools_ShowDirectOnly, "Show Direct Lighting", "Renders only the influence of direct lighting on the Light Environments.", EUserInterfaceActionType::ToggleButton, FInputChord() );
-	UI_COMMAND( LightingTools_ShowIndirectOnly, "Show Indirect Lighting", "Renders only the influence of indirect lighting on the Light Environments.", EUserInterfaceActionType::ToggleButton, FInputChord() );
-	UI_COMMAND( LightingTools_ShowIndirectSamples, "Show Indirect Lighting Samples", "Draws the indirect lighting samples that contributed to the indirect lighting for the Light Environments.", EUserInterfaceActionType::ToggleButton, FInputChord() );
 	UI_COMMAND( LightingDensity_RenderGrayscale, "Render Grayscale", "Renders the lightmap density.", EUserInterfaceActionType::ToggleButton, FInputChord() );
 	UI_COMMAND( LightingResolution_CurrentLevel, "Current Level", "Adjust only primitives in the current level.", EUserInterfaceActionType::RadioButton, FInputChord() );
 	UI_COMMAND( LightingResolution_SelectedLevels, "Selected Levels", "Adjust only primitives in the selected levels.", EUserInterfaceActionType::RadioButton, FInputChord() );
diff --git a/Engine/Source/Editor/LevelEditor/Private/LevelEditorToolBar.cpp b/Engine/Source/Editor/LevelEditor/Private/LevelEditorToolBar.cpp
--- a/Engine/Source/Editor/LevelEditor/Private/LevelEditorToolBar.cpp
+++ b/Engine/Source/Editor/LevelEditor/Private/LevelEditorToolBar.cpp
@@ -1328,20 +1328,6 @@ TSharedRef< SWidget > FLevelEditorToolBar::GenerateBuildMenuContent( TSharedRef<
 			InMenuBuilder.EndSection();
 		}
 
-		/** Generates a lighting tools sub-menu */
-		static void MakeLightingToolsMenu( FMenuBuilder& InMenuBuilder )
-		{
-			InMenuBuilder.BeginSection("LevelEditorBuildLightingTools", LOCTEXT( "LightingToolsHeading", "Light Environments" ) );
-			{
-				InMenuBuilder.AddMenuEntry( FLevelEditorCommands::Get().LightingTools_ShowBounds );
-				InMenuBuilder.AddMenuEntry( FLevelEditorCommands::Get().LightingTools_ShowTraces );
-				InMenuBuilder.AddMenuEntry( FLevelEditorCommands::Get().LightingTools_ShowDirectOnly );
-				InMenuBuilder.AddMenuEntry( FLevelEditorCommands::Get().LightingTools_ShowIndirectOnly );
-				InMenuBuilder.AddMenuEntry( FLevelEditorCommands::Get().LightingTools_ShowIndirectSamples );
-			}
-			InMenuBuilder.EndSection();
-		}
-
 		/** Generates a lighting density sub-menu */
 		static void MakeLightingDensityMenu( FMenuBuilder& InMenuBuilder )
 		{
@@ -1514,11 +1500,6 @@ TSharedRef< SWidget > FLevelEditorToolBar::GenerateBuildMenuContent( TSharedRef<
 		{
 			InMenuBuilder.BeginSection("LevelEditorBuildLightingInfo", LOCTEXT( "LightingInfoHeading", "Lighting Info Dialogs" ) );
 			{
-				InMenuBuilder.AddSubMenu(
-					LOCTEXT( "LightingToolsSubMenu", "Lighting Tools" ),
-					LOCTEXT( "LightingToolsSubMenu_ToolTip", "Shows the Lighting Tools options." ),
-					FNewMenuDelegate::CreateStatic( &FLightingMenus::MakeLightingToolsMenu ) );
-					
 				InMenuBuilder.AddSubMenu(
 					LOCTEXT( "LightingDensityRenderingSubMenu", "LightMap Density Rendering Options" ),
 					LOCTEXT( "LightingDensityRenderingSubMenu_ToolTip", "Shows the LightMap Density Rendering viewmode options." ),
diff --git a/Engine/Source/Editor/LevelEditor/Private/LightingTools.cpp b/Engine/Source/Editor/LevelEditor/Private/LightingTools.cpp
deleted file mode 100644
diff --git a/Engine/Source/Editor/LevelEditor/Public/LevelEditorActions.h b/Engine/Source/Editor/LevelEditor/Public/LevelEditorActions.h
--- a/Engine/Source/Editor/LevelEditor/Public/LevelEditorActions.h
+++ b/Engine/Source/Editor/LevelEditor/Public/LevelEditorActions.h
@@ -75,11 +75,6 @@ public:
 	TSharedPtr< FUICommandInfo > LightingQuality_High;
 	TSharedPtr< FUICommandInfo > LightingQuality_Medium;
 	TSharedPtr< FUICommandInfo > LightingQuality_Preview;
-	TSharedPtr< FUICommandInfo > LightingTools_ShowBounds;
-	TSharedPtr< FUICommandInfo > LightingTools_ShowTraces;
-	TSharedPtr< FUICommandInfo > LightingTools_ShowDirectOnly;
-	TSharedPtr< FUICommandInfo > LightingTools_ShowIndirectOnly;
-	TSharedPtr< FUICommandInfo > LightingTools_ShowIndirectSamples;
 	TSharedPtr< FUICommandInfo > LightingDensity_RenderGrayscale;
 	TSharedPtr< FUICommandInfo > LightingResolution_CurrentLevel;
 	TSharedPtr< FUICommandInfo > LightingResolution_SelectedLevels;
@@ -691,16 +686,6 @@ public:
 	static void BuildLODsOnly_Execute();
 	static void SetLightingQuality( ELightingBuildQuality NewQuality );
 	static bool IsLightingQualityChecked( ELightingBuildQuality TestQuality );
-	static void SetLightingToolShowBounds();
-	static bool IsLightingToolShowBoundsChecked();
-	static void SetLightingToolShowTraces();
-	static bool IsLightingToolShowTracesChecked();
-	static void SetLightingToolShowDirectOnly();
-	static bool IsLightingToolShowDirectOnlyChecked();
-	static void SetLightingToolShowIndirectOnly();
-	static bool IsLightingToolShowIndirectOnlyChecked();
-	static void SetLightingToolShowIndirectSamples();
-	static bool IsLightingToolShowIndirectSamplesChecked();
 	static float GetLightingDensityIdeal();
 	static void SetLightingDensityIdeal( float Value );
 	static float GetLightingDensityMaximum();
diff --git a/Engine/Source/Editor/LevelEditor/Public/LightingTools.h b/Engine/Source/Editor/LevelEditor/Public/LightingTools.h
deleted file mode 100644
