Added placed cameras to the camera type dropdown menu in the level editor viewport

diff --git a/Engine/Source/Editor/LevelEditor/Private/LevelEditorToolBar.cpp b/Engine/Source/Editor/LevelEditor/Private/LevelEditorToolBar.cpp
--- a/Engine/Source/Editor/LevelEditor/Private/LevelEditorToolBar.cpp
+++ b/Engine/Source/Editor/LevelEditor/Private/LevelEditorToolBar.cpp
@@ -1983,7 +1983,8 @@ TSharedRef< SWidget > FLevelEditorToolBar::GenerateMatineeMenuContent( TSharedRe
 		// We hide the header row to keep the UI compact.
 		// @todo: Might be useful to have this sometimes, actually.  Ideally the user could summon it.
 		InitOptions.bShowHeaderRow = false;
-
+		InitOptions.bShowSearchBox = false;
+		InitOptions.bShowCreateNewFolder = false;
 		struct Local
 		{
 			static bool IsMatineeActor( const AActor* const Actor )
@@ -2056,7 +2057,7 @@ TSharedRef< SWidget > FLevelEditorToolBar::GenerateMatineeMenuContent( TSharedRe
 	{
 		if( bHasAnyMatineeActors )
 		{
-			MenuBuilder.AddWidget(MiniSceneOutliner, FText::GetEmpty(), false);
+			MenuBuilder.AddWidget(MiniSceneOutliner, FText::GetEmpty(), true);
 		}
 	}
 	MenuBuilder.EndSection();
diff --git a/Engine/Source/Editor/LevelEditor/Private/SLevelViewport.cpp b/Engine/Source/Editor/LevelEditor/Private/SLevelViewport.cpp
--- a/Engine/Source/Editor/LevelEditor/Private/SLevelViewport.cpp
+++ b/Engine/Source/Editor/LevelEditor/Private/SLevelViewport.cpp
@@ -2074,18 +2074,15 @@ void SLevelViewport::OnActorLockToggleFromMenu(AActor* Actor)
 	{
 		const bool bLockNewActor = Actor != LevelViewportClient->GetActiveActorLock().Get();
 
-		// Unlock the previous actor
-		OnActorUnlock();
-
 		// Lock the new actor if it wasn't the same actor that we just unlocked
 		if (bLockNewActor)
 		{
+			// Unlock the previous actor
+			OnActorUnlock();
+
 			LockActorInternal(Actor);
 		}
 	}
-
-	// this is currently called from a context menu containing a scene outliner widget that needs closing
-	FSlateApplication::Get().DismissAllMenus();
 }
 
 bool SLevelViewport::IsActorLocked(const TWeakObjectPtr<AActor> Actor) const
diff --git a/Engine/Source/Editor/LevelEditor/Private/SLevelViewportToolBar.cpp b/Engine/Source/Editor/LevelEditor/Private/SLevelViewportToolBar.cpp
--- a/Engine/Source/Editor/LevelEditor/Private/SLevelViewportToolBar.cpp
+++ b/Engine/Source/Editor/LevelEditor/Private/SLevelViewportToolBar.cpp
@@ -671,6 +671,23 @@ void SLevelViewportToolBar::SetLevelProfile( FString DeviceProfileName )
 	UIManager->SetProfile( DeviceProfileName );
 }
 
+void SLevelViewportToolBar::GeneratePlacedCameraMenuEntries( FMenuBuilder& Builder, TArray<ACameraActor*> Cameras ) const
+{
+	FSlateIcon CameraIcon( FEditorStyle::GetStyleSetName(), "ClassIcon.CameraComponent" );
+
+	for( ACameraActor* CameraActor : Cameras )
+	{
+		// Needed for the delegate hookup to work below
+		AActor* GenericActor = CameraActor;
+
+		FText ActorDisplayName = FText::FromString(CameraActor->GetActorLabel());
+		FUIAction LookThroughCameraAction;
+		LookThroughCameraAction.ExecuteAction.BindSP(Viewport.Pin().ToSharedRef(), &SLevelViewport::OnActorLockToggleFromMenu, GenericActor);
+		LookThroughCameraAction.IsCheckedDelegate.BindSP(Viewport.Pin().ToSharedRef(), &SLevelViewport::IsActorLocked, TWeakObjectPtr<AActor>(GenericActor) );
+
+		Builder.AddMenuEntry( ActorDisplayName, FText::Format(LOCTEXT("LookThroughCameraActor_ToolTip", "Look through and pilot {0}"), ActorDisplayName), CameraIcon, LookThroughCameraAction, NAME_None, EUserInterfaceActionType::RadioButton );
+	}
+}
 
 TSharedRef<SWidget> SLevelViewportToolBar::GenerateCameraMenu() const
 {
@@ -691,6 +708,31 @@ TSharedRef<SWidget> SLevelViewportToolBar::GenerateCameraMenu() const
 		CameraMenuBuilder.AddMenuEntry(FEditorViewportCommands::Get().Back);
 	CameraMenuBuilder.EndSection();
 
+
+	TArray<ACameraActor*> Cameras;
+
+	for( TActorIterator<ACameraActor> It(GetWorld().Get()); It; ++It )
+	{
+		Cameras.Add( *It );
+	}
+
+	FText CameraActorsHeading = LOCTEXT("CameraActorsHeading", "Placed Cameras");
+
+	// Don't add too many cameras to the top level menu or else it becomes too large
+	const uint32 MaxCamerasInTopLevelMenu = 10;
+	if( Cameras.Num() > MaxCamerasInTopLevelMenu )
+	{
+		CameraMenuBuilder.BeginSection("CameraActors");
+			CameraMenuBuilder.AddSubMenu( CameraActorsHeading, LOCTEXT("LookThroughCameraActor_ToolTip", "Look through and pilot placed cameras"), FNewMenuDelegate::CreateSP(this, &SLevelViewportToolBar::GeneratePlacedCameraMenuEntries, Cameras ) );
+		CameraMenuBuilder.EndSection();
+	}
+	else
+	{
+		CameraMenuBuilder.BeginSection("CameraActors", CameraActorsHeading );
+			GeneratePlacedCameraMenuEntries( CameraMenuBuilder, Cameras );
+		CameraMenuBuilder.EndSection();
+	}
+
 	return CameraMenuBuilder.MakeWidget();
 }
 
diff --git a/Engine/Source/Editor/LevelEditor/Private/SLevelViewportToolBar.h b/Engine/Source/Editor/LevelEditor/Private/SLevelViewportToolBar.h
--- a/Engine/Source/Editor/LevelEditor/Private/SLevelViewportToolBar.h
+++ b/Engine/Source/Editor/LevelEditor/Private/SLevelViewportToolBar.h
@@ -101,6 +101,14 @@ private:
 	 */
 	TSharedRef<SWidget> GenerateCameraMenu() const;
 
+	/**
+	 * Generates menu entries for placed cameras (e.g CameraActors
+	 *
+	 * @param Builder	The menu builder to add menu entries to
+	 * @param Cameras	The list of cameras to add
+	 */
+	void GeneratePlacedCameraMenuEntries( FMenuBuilder& Builder, TArray<ACameraActor*> Cameras ) const;
+
 	/**
 	 * Generates the toolbar view menu content 
 	 *
