Slate : Tabs now provide a content slot for TabWell background content.
Editor : Editor puts the project watermark as a badge in the TabWell background.
Project name prefix, badge color and text color can be controlled via:

EditorPerProjectUserSettings.ini:
---------------------------------
[LevelEditor]
ProjectNameWatermarkPrefix=DEV
ProjectBadgeBackgroundColor=(R=237,G=25,B=104,A=255)
ProjectBadgeTextColor=(R=0,G=0,B=0,A=255)

diff --git a/Engine/Source/Editor/EditorStyle/Private/SlateEditorStyle.cpp b/Engine/Source/Editor/EditorStyle/Private/SlateEditorStyle.cpp
--- a/Engine/Source/Editor/EditorStyle/Private/SlateEditorStyle.cpp
+++ b/Engine/Source/Editor/EditorStyle/Private/SlateEditorStyle.cpp
@@ -133,6 +133,7 @@ void FSlateEditorStyle::FStyle::Initialize()
 	SetupGeneralStyles();
 	SetupGeneralIcons();
 	SetupWindowStyles();
+	SetupProjectBadgeStyle();
 	SetupDockingStyles();
 	SetupTutorialStyles();
 	SetupPropertyEditorStyles();
@@ -2567,8 +2568,18 @@ void FSlateEditorStyle::FStyle::SetupWindowStyles()
 			.SetChildBackgroundBrush( IMAGE_BRUSH( "Common/NoiseBackground", FVector2D(64, 64), FLinearColor::White, ESlateBrushTileType::Both) )
 			);
 	}
-	}
-	
+}
+
+void FSlateEditorStyle::FStyle::SetupProjectBadgeStyle()
+{
+	Set("SProjectBadge.Text", FTextBlockStyle(NormalText)
+		.SetFont(TTF_CORE_FONT("Fonts/Roboto-Regular", 12))
+		.SetColorAndOpacity(FLinearColor(0.0f, 0.0f, 0.0f, 1.f))
+	);
+
+	Set("SProjectBadge.BadgeShape", new BOX_BRUSH("ProjectBadge/Badge", Icon16x16, FMargin(6.0f/16)));
+}
+
 void FSlateEditorStyle::FStyle::SetupDockingStyles()
 {
 #if WITH_EDITOR || IS_PROGRAM
diff --git a/Engine/Source/Editor/EditorStyle/Private/SlateEditorStyle.h b/Engine/Source/Editor/EditorStyle/Private/SlateEditorStyle.h
--- a/Engine/Source/Editor/EditorStyle/Private/SlateEditorStyle.h
+++ b/Engine/Source/Editor/EditorStyle/Private/SlateEditorStyle.h
@@ -64,6 +64,7 @@ public:
 		void SetupGeneralStyles();
 		void SetupGeneralIcons();
 		void SetupWindowStyles();
+		void SetupProjectBadgeStyle();
 		void SetupDockingStyles();
 		void SetupTutorialStyles();
 		void SetupTranslationEditorStyles();
diff --git a/Engine/Source/Editor/LevelEditor/Private/LevelEditor.cpp b/Engine/Source/Editor/LevelEditor/Private/LevelEditor.cpp
--- a/Engine/Source/Editor/LevelEditor/Private/LevelEditor.cpp
+++ b/Engine/Source/Editor/LevelEditor/Private/LevelEditor.cpp
@@ -42,10 +42,10 @@ FLevelEditorModule::FLevelEditorModule()
 }
 
 
-class SLevelEditorWatermark : public SCompoundWidget
+class SProjectBadge : public SBox
 {
 public:
-	SLATE_BEGIN_ARGS(SLevelEditorWatermark) {}
+	SLATE_BEGIN_ARGS(SProjectBadge) {}
 	SLATE_END_ARGS()
 
 	void Construct(const FArguments& InArgs)
@@ -53,6 +53,12 @@ public:
 		FString OptionalBranchPrefix;
 		GConfig->GetString(TEXT("LevelEditor"), TEXT("ProjectNameWatermarkPrefix"), /*out*/ OptionalBranchPrefix, GEditorPerProjectIni);
 
+		FColor BadgeBackgroundColor = FColor::Black;
+		GConfig->GetColor(TEXT("LevelEditor"), TEXT("ProjectBadgeBackgroundColor"), /*out*/ BadgeBackgroundColor, GEditorPerProjectIni);
+
+		FColor BadgeTextColor = FColor(128,128,128,255);
+		GConfig->GetColor(TEXT("LevelEditor"), TEXT("ProjectBadgeTextColor"), /*out*/ BadgeTextColor, GEditorPerProjectIni);
+
 		FFormatNamedArguments Args;
 		Args.Add(TEXT("Branch"), FText::FromString(OptionalBranchPrefix));
 		Args.Add(TEXT("GameName"), FText::FromString(FString(FApp::GetGameName())));
@@ -82,16 +88,27 @@ public:
 
 		SetToolTipText(RightContentTooltip);
 
-		ChildSlot
-		[
-			SNew(STextBlock)
-			.Text(RightContentText)
-			.Visibility(EVisibility::HitTestInvisible)
-			.Font(FSlateFontInfo(FPaths::EngineContentDir() / TEXT("Slate/Fonts/Roboto-Bold.ttf"), 14))
-			.ColorAndOpacity(FLinearColor(1.0f, 1.0f, 1.0f, 0.4f))
-		];
+		SBox::Construct(SBox::FArguments()
+			.HAlign(HAlign_Right)
+			.VAlign(VAlign_Top)
+			.Padding(FMargin(0.0f, 0.0f, 15.0f, 0.0f))
+			[
+				SNew(SBorder)
+				.BorderImage(FEditorStyle::GetBrush("SProjectBadge.BadgeShape"))
+				.Padding(FMargin(10.0f, 2.5f))
+				.BorderBackgroundColor(BadgeBackgroundColor)
+				.VAlign(VAlign_Top)
+				[
+					SNew(STextBlock)
+					.Text(RightContentText)
+					.Visibility(EVisibility::HitTestInvisible)
+					.TextStyle(FEditorStyle::Get(),"SProjectBadge.Text")
+					.ColorAndOpacity(BadgeTextColor)
+				]
+			]);
 	}
-	
+
+
 	// SWidget interface
 	virtual EWindowZone::Type GetWindowZoneOverride() const override
 	{
@@ -132,45 +149,41 @@ TSharedRef<SDockTab> FLevelEditorModule::SpawnLevelEditor( const FSpawnTabArgs&
 	IIntroTutorials& IntroTutorials = FModuleManager::LoadModuleChecked<IIntroTutorials>(TEXT("IntroTutorials"));
 	TSharedRef<SWidget> TutorialWidget = IntroTutorials.CreateTutorialsWidget(TEXT("LevelEditor"), OwnerWindow);
 
-	TSharedPtr< SWidget > RightContent;
-	{
-		RightContent =
-				SNew( SHorizontalBox )
+	TSharedPtr< SWidget > RightContent=
+		SNew( SHorizontalBox )
 
-				+SHorizontalBox::Slot()
-				.Padding(0.0f, 0.0f, 14.0f, 0.0f)
-				.AutoWidth()
-				[
-					SNew(SLevelEditorWatermark)
-				]
 // Put the level editor stats/notification widgets on the main window title bar since we don't have a menu bar on OS X
 #if PLATFORM_MAC
-				+SHorizontalBox::Slot()
-				.AutoWidth()
-				.Padding(8.0f, 0.0f, 0.0f, 0.0f)
-				.VAlign(VAlign_Center)
-				[
-					LevelEditorTab->GetRightContent()
-				]
+		+SHorizontalBox::Slot()
+		.AutoWidth()
+		.Padding(8.0f, 0.0f, 0.0f, 0.0f)
+		.VAlign(VAlign_Center)
+		[
+			LevelEditorTab->GetRightContent()
+		]
 #endif
 		
-				+SHorizontalBox::Slot()
-				.AutoWidth()
-				.Padding(8.0f, 0.0f, 0.0f, 0.0f)
-				.VAlign(VAlign_Center)
-				[
-					UserFeedbackWidget
-				]
-				+ SHorizontalBox::Slot()
-				.AutoWidth()
-				.Padding(8.0f, 0.0f, 8.0f, 0.0f)
-				.VAlign(VAlign_Center)
-				[
-					TutorialWidget
-				]
-		;
-	}
+		+SHorizontalBox::Slot()
+		.AutoWidth()
+		.Padding(8.0f, 0.0f, 0.0f, 0.0f)
+		.VAlign(VAlign_Center)
+		[
+			UserFeedbackWidget
+		]
+		+ SHorizontalBox::Slot()
+		.AutoWidth()
+		.Padding(8.0f, 0.0f, 8.0f, 0.0f)
+		.VAlign(VAlign_Center)
+		[
+			TutorialWidget
+		];
+
 	LevelEditorTab->SetRightContent( RightContent.ToSharedRef() );
+
+	LevelEditorTab->SetBackgroundContent(
+			SNew(SProjectBadge)
+	);
+	
 	
 	return LevelEditorTab;
 }
diff --git a/Engine/Source/Runtime/Slate/Private/Framework/Docking/SDockingTabStack.cpp b/Engine/Source/Runtime/Slate/Private/Framework/Docking/SDockingTabStack.cpp
--- a/Engine/Source/Runtime/Slate/Private/Framework/Docking/SDockingTabStack.cpp
+++ b/Engine/Source/Runtime/Slate/Private/Framework/Docking/SDockingTabStack.cpp
@@ -43,6 +43,7 @@ void SDockingTabStack::Construct( const FArguments& InArgs, const TSharedRef<FTa
 
 	InlineContentAreaLeft = NULL;
 	InlineContentAreaRight = NULL;
+	BackgroundContentArea = NULL;
 	TitleBarSlot = NULL;
 
 	this->TabStackGeometry = FGeometry();
@@ -90,10 +91,14 @@ void SDockingTabStack::Construct( const FArguments& InArgs, const TSharedRef<FTa
 		.VAlign(VAlign_Bottom)
 		.Padding(4.0f, 0.0f, 0.0f, 0.0f)
 		[
-			SNew(SVerticalBox)
+			SNew(SOverlay)
+			+ SOverlay::Slot().Expose(BackgroundContentArea)
+			+ SOverlay::Slot()
+			[
+				SNew(SVerticalBox)
 				.Visibility(EVisibility::SelfHitTestInvisible)
 
-			+ SVerticalBox::Slot()
+				+ SVerticalBox::Slot()
 				.AutoHeight()
 				[
 					SNew(SSpacer)
@@ -101,13 +106,14 @@ void SDockingTabStack::Construct( const FArguments& InArgs, const TSharedRef<FTa
 						.Size(FVector2D(0.0f, 10.0f))
 				]
 
-			+ SVerticalBox::Slot()
+				+ SVerticalBox::Slot()
 				.AutoHeight()
 				[
 					// TabWell
 					SAssignNew(TabWell, SDockingTabWell)
 						.ParentStackNode(SharedThis(this))
 				]
+			]
 		]
 
 	+ SHorizontalBox::Slot()
@@ -137,8 +143,8 @@ void SDockingTabStack::Construct( const FArguments& InArgs, const TSharedRef<FTa
 				.Visibility(EVisibility::SelfHitTestInvisible)
 
 				+ SVerticalBox::Slot()
-				.AutoHeight()
 				.Expose(TitleBarSlot)
+				.AutoHeight()				
 
 				+ SVerticalBox::Slot()
 				.AutoHeight()
@@ -208,7 +214,7 @@ void SDockingTabStack::Construct( const FArguments& InArgs, const TSharedRef<FTa
 
 	if (bIsDocumentArea)
 	{
-		this->SetNodeContent(SDocumentAreaWidget::MakeDocumentAreaWidget(), SNullWidget::NullWidget, SNullWidget::NullWidget);
+		this->SetNodeContent(SDocumentAreaWidget::MakeDocumentAreaWidget(), SNullWidget::NullWidget, SNullWidget::NullWidget, SNullWidget::NullWidget);
 	}
 }
 
@@ -219,11 +225,11 @@ void SDockingTabStack::OnLastTabRemoved()
 	{
 		// Stop holding onto any meaningful window content.
 		// The user should not see any content in this DockNode.
-		this->SetNodeContent( SNullWidget::NullWidget, SNullWidget::NullWidget, SNullWidget::NullWidget );
+		this->SetNodeContent(SNullWidget::NullWidget, SNullWidget::NullWidget, SNullWidget::NullWidget, SNullWidget::NullWidget);
 	}
 	else
 	{
-		this->SetNodeContent( SDocumentAreaWidget::MakeDocumentAreaWidget(), SNullWidget::NullWidget, SNullWidget::NullWidget );
+		this->SetNodeContent(SDocumentAreaWidget::MakeDocumentAreaWidget(), SNullWidget::NullWidget, SNullWidget::NullWidget, SNullWidget::NullWidget);
 	}
 }
 
@@ -328,11 +334,12 @@ void SDockingTabStack::BringToFront( const TSharedRef<SDockTab>& TabToBringToFro
 	TabWell->BringTabToFront(TabToBringToFront);
 }
 
-void SDockingTabStack::SetNodeContent( const TSharedRef<SWidget>& InContent, const TSharedRef<SWidget>& ContentLeft, const TSharedRef<SWidget>& ContentRight )
+void SDockingTabStack::SetNodeContent(const TSharedRef<SWidget>& InContent, const TSharedRef<SWidget>& ContentLeft, const TSharedRef<SWidget>& ContentRight, const TSharedRef<SWidget>& InContentBackground)
 {
 	ContentSlot->SetContent(InContent);
 	(*InlineContentAreaLeft)[ContentLeft];
 	(*InlineContentAreaRight)[ContentRight];
+	(*BackgroundContentArea)[InContentBackground];
 }
 
 FReply SDockingTabStack::OnDragOver( const FGeometry& MyGeometry, const FDragDropEvent& DragDropEvent )
diff --git a/Engine/Source/Runtime/Slate/Private/Framework/Docking/SDockingTabStack.h b/Engine/Source/Runtime/Slate/Private/Framework/Docking/SDockingTabStack.h
--- a/Engine/Source/Runtime/Slate/Private/Framework/Docking/SDockingTabStack.h
+++ b/Engine/Source/Runtime/Slate/Private/Framework/Docking/SDockingTabStack.h
@@ -53,7 +53,7 @@ public:
 	void BringToFront( const TSharedRef<SDockTab>& TabToBringToFront );
 
 	/** Set the content that the DockNode is presenting. */
-	void SetNodeContent( const TSharedRef<SWidget>& InContent, const TSharedRef<SWidget>& InContentLeft, const TSharedRef<SWidget>& InContentRight );
+	void SetNodeContent( const TSharedRef<SWidget>& InContent, const TSharedRef<SWidget>& InContentLeft, const TSharedRef<SWidget>& InContentRight, const TSharedRef<SWidget>& InContentBackground );
 
 	virtual FReply OnUserAttemptingDock( SDockingNode::RelativeDirection Direction, const FDragDropEvent& DragDropEvent ) override;
 	
@@ -192,6 +192,7 @@ private:
 	/** The borders that hold any potential inline content areas. */
 	SHorizontalBox::FSlot* InlineContentAreaLeft;
 	SHorizontalBox::FSlot* InlineContentAreaRight;
+	SOverlay::FOverlaySlot* BackgroundContentArea;
 	SVerticalBox::FSlot* TitleBarSlot;
 	TSharedPtr<SWidget> TitleBarContent;
 
diff --git a/Engine/Source/Runtime/Slate/Private/Framework/Docking/SDockingTabWell.cpp b/Engine/Source/Runtime/Slate/Private/Framework/Docking/SDockingTabWell.cpp
--- a/Engine/Source/Runtime/Slate/Private/Framework/Docking/SDockingTabWell.cpp
+++ b/Engine/Source/Runtime/Slate/Private/Framework/Docking/SDockingTabWell.cpp
@@ -304,7 +304,7 @@ void SDockingTabWell::OnDragEnter( const FGeometry& MyGeometry, const FDragDropE
 			this->TabGrabOffsetFraction = DragDropOperation->GetTabGrabOffsetFraction();
 			
 			// The user should see the contents of the tab that we're dragging.
-			ParentTabStackPtr.Pin()->SetNodeContent(DragDropOperation->GetTabBeingDragged()->GetContent(), SNullWidget::NullWidget, SNullWidget::NullWidget);
+			ParentTabStackPtr.Pin()->SetNodeContent(DragDropOperation->GetTabBeingDragged()->GetContent(), SNullWidget::NullWidget, SNullWidget::NullWidget, SNullWidget::NullWidget);
 		}
 	}
 }
@@ -580,11 +580,11 @@ void SDockingTabWell::RefreshParentContent()
 			ParentWindowPtr->SetTitle( ForegroundTab->GetTabLabel() );
 		}
 
-		ParentTabStackPtr.Pin()->SetNodeContent( ForegroundTab->GetContent(), ForegroundTab->GetLeftContent(), ForegroundTab->GetRightContent() );
+		ParentTabStackPtr.Pin()->SetNodeContent( ForegroundTab->GetContent(), ForegroundTab->GetLeftContent(), ForegroundTab->GetRightContent(), ForegroundTab->GetBackgrounfContent() );
 	}
 	else
 	{
-		ParentTabStackPtr.Pin()->SetNodeContent( SNullWidget::NullWidget, SNullWidget::NullWidget, SNullWidget::NullWidget );
+		ParentTabStackPtr.Pin()->SetNodeContent(SNullWidget::NullWidget, SNullWidget::NullWidget, SNullWidget::NullWidget, SNullWidget::NullWidget);
 	}
 }
 
diff --git a/Engine/Source/Runtime/Slate/Private/Widgets/Docking/SDockTab.cpp b/Engine/Source/Runtime/Slate/Private/Widgets/Docking/SDockTab.cpp
--- a/Engine/Source/Runtime/Slate/Private/Widgets/Docking/SDockTab.cpp
+++ b/Engine/Source/Runtime/Slate/Private/Widgets/Docking/SDockTab.cpp
@@ -177,6 +177,15 @@ void SDockTab::SetRightContent( TSharedRef<SWidget> InContent )
 	}
 }
 
+void SDockTab::SetBackgroundContent(TSharedRef<SWidget> InContent)
+{
+	this->TabWellContentBackground = InContent;
+	if (ParentPtr.IsValid())
+	{
+		ParentPtr.Pin()->RefreshParentContent();
+	}
+}
+
 
 bool SDockTab::IsActive() const
 {
@@ -248,6 +257,11 @@ TSharedRef<SWidget> SDockTab::GetRightContent()
 	return TabWellContentRight;
 }
 
+TSharedRef<SWidget> SDockTab::GetBackgrounfContent()
+{
+	return TabWellContentBackground;
+}
+
 FMargin SDockTab::GetContentPadding() const
 {
 	return ContentAreaPadding.Get();
@@ -307,6 +321,7 @@ SDockTab::SDockTab()
 	: Content(SNew(SSpacer))
 	, TabWellContentLeft(SNullWidget::NullWidget)
 	, TabWellContentRight(SNullWidget::NullWidget)
+	, TabWellContentBackground(SNullWidget::NullWidget)
 	, LayoutIdentifier(NAME_None)
 	, TabRole(ETabRole::PanelTab)
 	, ParentPtr()
diff --git a/Engine/Source/Runtime/Slate/Public/Widgets/Docking/SDockTab.h b/Engine/Source/Runtime/Slate/Public/Widgets/Docking/SDockTab.h
--- a/Engine/Source/Runtime/Slate/Public/Widgets/Docking/SDockTab.h
+++ b/Engine/Source/Runtime/Slate/Public/Widgets/Docking/SDockTab.h
@@ -48,6 +48,7 @@ public:
 		: _Content()
 		, _TabWellContentLeft()
 		, _TabWellContentRight()
+		, _TabWellContentBackground()
 		, _ContentPadding( FMargin( 2 ) )
 		, _TabRole(ETabRole::PanelTab)
 		, _Label()
@@ -63,6 +64,7 @@ public:
 		SLATE_DEFAULT_SLOT( FArguments, Content )
 		SLATE_NAMED_SLOT( FArguments, TabWellContentLeft )
 		SLATE_NAMED_SLOT( FArguments, TabWellContentRight )
+		SLATE_NAMED_SLOT(FArguments, TabWellContentBackground)
 		SLATE_ATTRIBUTE( FMargin, ContentPadding )
 		SLATE_ARGUMENT( ETabRole, TabRole )
 		SLATE_ATTRIBUTE( FText, Label )
@@ -94,8 +96,12 @@ public:
 	virtual void SetContent(TSharedRef<SWidget> InContent) override;
 	// End of SBorder interface
 
+	/** Content that appears in the TabWell to the left of the tabs */
 	void SetLeftContent( TSharedRef<SWidget> InContent );
+	/** Content that appears in the TabWell to the right of the tabs */
 	void SetRightContent( TSharedRef<SWidget> InContent );
+	/** Content that appears in the TabWell behind the tabs */
+	void SetBackgroundContent( TSharedRef<SWidget> InContent );
 
 	/** @return True if this tab is currently focused */
 	bool IsActive() const;
@@ -126,6 +132,7 @@ public:
 	TSharedRef<SWidget> GetContent();
 	TSharedRef<SWidget> GetLeftContent();
 	TSharedRef<SWidget> GetRightContent();
+	TSharedRef<SWidget> GetBackgrounfContent();
 
 	/** Padding around the content when it is presented by the SDockingTabStack */
 	FMargin GetContentPadding() const;
@@ -304,6 +311,8 @@ protected:
 	TSharedRef<SWidget> Content;
 	TSharedRef<SWidget> TabWellContentLeft;
 	TSharedRef<SWidget> TabWellContentRight;
+	TSharedRef<SWidget> TabWellContentBackground;
+
 	
 	/** The tab's layout identifier */
 	FTabId LayoutIdentifier;
