Editor: Switched styling on 'snap to layer' menu command to match normal menu items, made it bindable, and gave it a default key binding of Ctrl+Space

diff --git a/Engine/Source/Editor/LevelEditor/Private/LevelEditor.cpp b/Engine/Source/Editor/LevelEditor/Private/LevelEditor.cpp
--- a/Engine/Source/Editor/LevelEditor/Private/LevelEditor.cpp
+++ b/Engine/Source/Editor/LevelEditor/Private/LevelEditor.cpp
@@ -743,6 +743,12 @@ void FLevelEditorModule::BindGlobalLevelEditorCommands()
 		FCanExecuteAction::CreateStatic(&FLevelEditorActionCallbacks::ActorsSelected_CanExecute)
 		);
 
+	ActionList.MapAction(
+		Commands.SnapToLayer2D,
+		FExecuteAction::CreateStatic(&FLevelEditorActionCallbacks::SnapToLayer2D_Clicked),
+		FCanExecuteAction::CreateStatic(&FLevelEditorActionCallbacks::CanSnapToLayer2D)
+		);
+
 	bAlign = false;
 	bool bUseLineTrace = false;
 	bool bUseBounds = false;
diff --git a/Engine/Source/Editor/LevelEditor/Private/LevelEditorActions.cpp b/Engine/Source/Editor/LevelEditor/Private/LevelEditorActions.cpp
--- a/Engine/Source/Editor/LevelEditor/Private/LevelEditorActions.cpp
+++ b/Engine/Source/Editor/LevelEditor/Private/LevelEditorActions.cpp
@@ -63,6 +63,7 @@
 #include "Animation/SkeletalMeshActor.h"
 #include "Editor/Persona/Public/AnimationRecorder.h"
 #include "Editor/KismetWidgets/Public/CreateBlueprintFromActorDialog.h"
+#include "EditorProjectSettings.h"
 
 DEFINE_LOG_CATEGORY_STATIC(LevelEditorActions, Log, All);
 
@@ -2528,6 +2529,59 @@ void FLevelEditorActionCallbacks::MoveActorTo_Clicked( const bool InAlign, const
 	GEditor->RebuildAlteredBSP(); // Update the Bsp of any levels containing a modified brush
 }
 
+void FLevelEditorActionCallbacks::SnapToLayer2D_Clicked()
+{
+	// Fires ULevel::LevelDirtiedEvent when falling out of scope.
+	FScopedLevelDirtied LevelDirtyCallback;
+
+	const ULevelEditorViewportSettings* ViewportSettings = GetDefault<ULevelEditorViewportSettings>();
+	const ULevelEditor2DSettings* Settings2D = GetDefault<ULevelEditor2DSettings>();
+	if (Settings2D->SnapLayers.IsValidIndex(ViewportSettings->ActiveSnapLayerIndex))
+	{
+		const FScopedTransaction Transaction(NSLOCTEXT("UnrealEd", "SnapSelection2D", "Snap Selection to 2D Layer"));
+
+		float SnapDepth = Settings2D->SnapLayers[ViewportSettings->ActiveSnapLayerIndex].Depth;
+		USelection* SelectedActors = GEditor->GetSelectedActors();
+		for (FSelectionIterator Iter(*SelectedActors); Iter; ++Iter)
+		{
+			AActor* Actor = CastChecked<AActor>(*Iter);
+
+			// Only snap actors that are not attached to something else
+			if (Actor->GetAttachParentActor() == nullptr)
+			{
+				FTransform Transform = Actor->GetTransform();
+				FVector CurrentLocation = Transform.GetLocation();
+
+				switch (Settings2D->SnapAxis)
+				{
+				case ELevelEditor2DAxis::X: CurrentLocation.X = SnapDepth; break;
+				case ELevelEditor2DAxis::Y: CurrentLocation.Y = SnapDepth; break;
+				case ELevelEditor2DAxis::Z: CurrentLocation.Z = SnapDepth; break;
+				}
+
+				Transform.SetLocation(CurrentLocation);
+				Actor->Modify();
+				Actor->SetActorTransform(Transform);
+
+				Actor->InvalidateLightingCache();
+				Actor->UpdateComponentTransforms();
+
+				Actor->MarkPackageDirty();
+				LevelDirtyCallback.Request();
+			}
+		}
+
+		GEditor->RedrawLevelEditingViewports(true);
+		GEditor->RebuildAlteredBSP();
+	}
+}
+
+bool FLevelEditorActionCallbacks::CanSnapToLayer2D()
+{
+	const ULevelEditor2DSettings* Settings = GetDefault<ULevelEditor2DSettings>();
+	return Settings->SnapLayers.IsValidIndex(GetDefault<ULevelEditorViewportSettings>()->ActiveSnapLayerIndex) && (GEditor->GetSelectedActorCount() > 0);
+}
+
 void FLevelEditorActionCallbacks::SnapToFloor_Clicked( bool InAlign, bool InUseLineTrace, bool InUseBounds, bool InUsePivot )
 {
 	const FScopedTransaction Transaction( NSLOCTEXT("UnrealEd", "SnapActorsToFloor", "Snap Actors To Floor") );
@@ -2829,6 +2883,7 @@ void FLevelEditorCommands::RegisterCommands()
 	UI_COMMAND( SnapOriginToGrid, "Snap Origin to Grid", "Snaps the actor to the nearest grid location at its origin", EUserInterfaceActionType::Button, FInputChord( EModifierKey::Control, EKeys::End ) );
 	UI_COMMAND( SnapOriginToGridPerActor, "Snap Origin to Grid Per Actor", "Snaps each selected actor separately to the nearest grid location at its origin", EUserInterfaceActionType::Button, FInputChord() );
 	UI_COMMAND( AlignOriginToGrid, "Align Origin to Grid", "Aligns the actor to the nearest grid location at its origin", EUserInterfaceActionType::Button, FInputChord() );
+	UI_COMMAND( SnapToLayer2D, "Snap to 2D Layer", "Snaps the actor to the current 2D snap layer", EUserInterfaceActionType::Button, FInputChord( EModifierKey::Control, EKeys::SpaceBar ) );
 	UI_COMMAND( SnapToFloor, "Snap to Floor", "Snaps the actor or component to the floor below it", EUserInterfaceActionType::Button, FInputChord( EKeys::End ) );
 	UI_COMMAND( AlignToFloor, "Align to Floor", "Aligns the actor or component with the floor", EUserInterfaceActionType::Button, FInputChord() );
 	UI_COMMAND( SnapPivotToFloor, "Snap Pivot to Floor", "Snaps the actor to the floor at its pivot point", EUserInterfaceActionType::Button, FInputChord( EModifierKey::Alt, EKeys::End ) );	
diff --git a/Engine/Source/Editor/LevelEditor/Private/LevelEditorContextMenu.cpp b/Engine/Source/Editor/LevelEditor/Private/LevelEditorContextMenu.cpp
--- a/Engine/Source/Editor/LevelEditor/Private/LevelEditorContextMenu.cpp
+++ b/Engine/Source/Editor/LevelEditor/Private/LevelEditorContextMenu.cpp
@@ -864,6 +864,7 @@ void FLevelEditorContextMenuImpl::FillSnapAlignMenu( FMenuBuilder& MenuBuilder )
 	MenuBuilder.AddMenuEntry( FLevelEditorCommands::Get().SnapOriginToGrid );
 	MenuBuilder.AddMenuEntry( FLevelEditorCommands::Get().SnapOriginToGridPerActor );
 	MenuBuilder.AddMenuEntry( FLevelEditorCommands::Get().AlignOriginToGrid );
+	MenuBuilder.AddMenuEntry( FLevelEditorCommands::Get().SnapToLayer2D );
 	MenuBuilder.AddMenuEntry( FLevelEditorCommands::Get().SnapToFloor );
 	MenuBuilder.AddMenuEntry( FLevelEditorCommands::Get().AlignToFloor );
 	MenuBuilder.AddMenuEntry( FLevelEditorCommands::Get().SnapPivotToFloor );
diff --git a/Engine/Source/Editor/LevelEditor/Public/LevelEditorActions.h b/Engine/Source/Editor/LevelEditor/Public/LevelEditorActions.h
--- a/Engine/Source/Editor/LevelEditor/Public/LevelEditorActions.h
+++ b/Engine/Source/Editor/LevelEditor/Public/LevelEditorActions.h
@@ -136,6 +136,9 @@ public:
 	/** Aligns the actor to the grid at its pivot*/
 	TSharedPtr< FUICommandInfo > AlignOriginToGrid;
 
+	/** Snaps the actor to the 2D layer */
+	TSharedPtr< FUICommandInfo > SnapToLayer2D;
+
 	/** Snaps the actor to the floor*/
 	TSharedPtr< FUICommandInfo > SnapToFloor;
 
@@ -1145,6 +1148,17 @@ public:
 	 */
 	static void MoveActorToActor_Clicked( bool InAlign );
 
+	/** 
+	 * Snaps an actor to the currently selected 2D snap layer
+	 */
+	static void SnapToLayer2D_Clicked();
+
+	/**
+	 * Checks to see if at least a single actor is selected and the 2D editor mode is enabled
+	 *	@return true if it can execute.
+	 */
+	static bool CanSnapToLayer2D();
+
 	/** 
 	 * Snaps an actor to the floor.  Optionally will align with the trace normal.
 	 * @param InAlign			Whether or not to rotate the actor to align with the trace normal.
diff --git a/Engine/Source/Editor/UnrealEd/Private/STransformViewportToolbar.cpp b/Engine/Source/Editor/UnrealEd/Private/STransformViewportToolbar.cpp
--- a/Engine/Source/Editor/UnrealEd/Private/STransformViewportToolbar.cpp
+++ b/Engine/Source/Editor/UnrealEd/Private/STransformViewportToolbar.cpp
@@ -11,6 +11,8 @@
 #include "SNumericEntryBox.h"
 #include "ScopedTransaction.h"
 #include "Settings/EditorProjectSettings.h"
+#include "LevelEditor.h"
+#include "LevelEditorActions.h"
 
 #define LOCTEXT_NAMESPACE "TransformToolBar"
 
@@ -656,51 +658,6 @@ TSharedRef<SWidget> STransformViewportToolBar::FillLayer2DSnapMenu()
 				SettingsModule->ShowViewer("Project", "Editor", "LevelEditor2DSettings");
 			}
 		}
-
-		static FReply SnapSelectionToSelectedOffset()
-		{
-			const ULevelEditorViewportSettings* ViewportSettings = GetDefault<ULevelEditorViewportSettings>();
-			const ULevelEditor2DSettings* Settings2D = GetDefault<ULevelEditor2DSettings>();
-			if (Settings2D->SnapLayers.IsValidIndex(ViewportSettings->ActiveSnapLayerIndex))
-			{
-				const FScopedTransaction Transaction(NSLOCTEXT("UnrealEd", "SnapSelection2D", "Snap Selection to 2D Layer"));
-
-				float SnapDepth = Settings2D->SnapLayers[ViewportSettings->ActiveSnapLayerIndex].Depth;
-				USelection* SelectedActors = GEditor->GetSelectedActors();
-				for (FSelectionIterator Iter(*SelectedActors); Iter; ++Iter)
-				{
-					AActor* Actor = CastChecked<AActor>(*Iter);
-
-					// Only snap actors that are not attached to something else
-					if (Actor->GetAttachParentActor() == nullptr)
-					{
-						FTransform Transform = Actor->GetTransform();
-						FVector CurrentLocation = Transform.GetLocation();
-
-						switch (Settings2D->SnapAxis)
-						{
-						case ELevelEditor2DAxis::X: CurrentLocation.X = SnapDepth; break;
-						case ELevelEditor2DAxis::Y: CurrentLocation.Y = SnapDepth; break;
-						case ELevelEditor2DAxis::Z: CurrentLocation.Z = SnapDepth; break;
-						}
-						
-						Transform.SetLocation(CurrentLocation);
-						Actor->Modify();
-						Actor->SetActorTransform(Transform);
-					}
-				}
-
-				GEditor->RedrawLevelEditingViewports(true);
-			}
-
-			return FReply::Handled();
-		}
-
-		static bool IsSnapSelectedActorsEnabled()
-		{
-			const ULevelEditor2DSettings* Settings = GetDefault<ULevelEditor2DSettings>();
-			return Settings->SnapLayers.IsValidIndex(GetDefault<ULevelEditorViewportSettings>()->ActiveSnapLayerIndex) && (GEditor->GetSelectedActorCount() > 0);
-		}
 	};
 
 	FUIAction ShowSettingsAction(FExecuteAction::CreateStatic(&FLocalFunctions::ShowSettingsViewer));
@@ -709,17 +666,9 @@ TSharedRef<SWidget> STransformViewportToolBar::FillLayer2DSnapMenu()
 	// -------------------------------------------------------
 	ShowMenuBuilder.AddMenuSeparator();
 
-	ShowMenuBuilder.AddWidget(
-		SNew(SButton)
-		.Text(LOCTEXT("2DSnap_SnapSelection", "Snap Selected Actors"))
-		.IsEnabled(TAttribute<bool>::Create(TAttribute<bool>::FGetter::CreateStatic(FLocalFunctions::IsSnapSelectedActorsEnabled)))
-		.OnClicked(FOnClicked::CreateStatic(FLocalFunctions::SnapSelectionToSelectedOffset))
-		.VAlign(VAlign_Center)
-		.HAlign(HAlign_Fill),
 
-		FText::GetEmpty(),
-		true
-		);
+	FLevelEditorModule& LevelEditor = FModuleManager::GetModuleChecked<FLevelEditorModule>(TEXT("LevelEditor"));
+	ShowMenuBuilder.AddMenuEntry(LevelEditor.GetLevelEditorCommands().SnapToLayer2D);
 
 	return ShowMenuBuilder.MakeWidget();
 }
