This change registers a delegate to notify when component blueprints are recompiled and forces a refresh on the actor details component section when neccesary.
#UE-11573 - Overriden property values do not survive property renames in Actor Components
#Branch UE4
#Proj LevelEditor

diff --git a/Engine/Source/Editor/LevelEditor/Private/SActorDetails.cpp b/Engine/Source/Editor/LevelEditor/Private/SActorDetails.cpp
--- a/Engine/Source/Editor/LevelEditor/Private/SActorDetails.cpp
+++ b/Engine/Source/Editor/LevelEditor/Private/SActorDetails.cpp
@@ -186,7 +186,8 @@ SActorDetails::~SActorDetails()
 {
 	GEditor->UnregisterForUndo(this);
 	USelection::SelectionChangedEvent.RemoveAll(this);
-	
+	ClearNotificationDelegates();
+
 	auto LevelEditor = FModuleManager::GetModulePtr<FLevelEditorModule>("LevelEditor");
 	if (LevelEditor != nullptr)
 	{
@@ -198,6 +199,7 @@ void SActorDetails::SetObjects(const TArray<UObject*>& InObjects, bool bForceRef
 {
 	if(!DetailsView->IsLocked())
 	{
+		ClearNotificationDelegates();
 		DetailsView->SetObjects(InObjects, bForceRefresh);
 
 		bool bShowingComponents = false;
@@ -210,6 +212,19 @@ void SActorDetails::SetObjects(const TArray<UObject*>& InObjects, bool bForceRef
 				LockedActorSelection = Actor;
 				bShowingComponents = true;
 
+				// Register for component blueprint compiled events
+				for( auto Component : Actor->GetComponents() )
+				{
+					if(UBlueprintGeneratedClass* ComponentBPGC = Cast<UBlueprintGeneratedClass>(Component->GetClass()))
+					{
+						UBlueprint* ComponentBlueprint = Cast<UBlueprint>(ComponentBPGC->ClassGeneratedBy);
+						if(!ComponentBlueprint->OnCompiled().IsBoundToObject( this ))
+						{
+							ComponentBlueprint->OnCompiled().AddSP(this, &SActorDetails::OnBlueprintRecompiled);
+						}
+					}
+				}
+
 				// Update the tree if a new actor is selected
 				if(GEditor->GetSelectedComponentCount() == 0)
 				{
@@ -496,6 +511,7 @@ void SActorDetails::UpdateComponentTreeFromEditorSelection()
 
 		auto& SCSTreeWidget = SCSEditor->SCSTreeWidget;
 		TArray<UObject*> DetailsObjects;
+		bool bForceRefresh = false;
 
 		// Update the tree selection to match the level editor component selection
 		SCSTreeWidget->ClearSelection();
@@ -503,6 +519,18 @@ void SActorDetails::UpdateComponentTreeFromEditorSelection()
 		{
 			UActorComponent* Component = CastChecked<UActorComponent>(*It);
 
+			if(!bForceRefresh)
+			{
+				if(UBlueprintGeneratedClass* ComponentBPGC = Cast<UBlueprintGeneratedClass>(Component->GetClass()))
+				{
+					UBlueprint* ComponentBlueprint = Cast<UBlueprint>(ComponentBPGC->ClassGeneratedBy);
+					if(ComponentBlueprint->CrcLastCompiledSignature == CrcLastCompiledSignature )
+					{
+						bForceRefresh = true;
+					}
+				}
+			}
+
 			auto SCSTreeNode = SCSEditor->GetNodeFromActorComponent(Component, false);
 			if (SCSTreeNode.IsValid() && SCSTreeNode->GetComponentTemplate())
 			{
@@ -514,10 +542,11 @@ void SActorDetails::UpdateComponentTreeFromEditorSelection()
 				DetailsObjects.Add(Component);
 			}
 		}
+		CrcLastCompiledSignature = 0;
 
 		if (DetailsObjects.Num() > 0)
 		{
-			DetailsView->SetObjects(DetailsObjects);
+			DetailsView->SetObjects(DetailsObjects,bForceRefresh);
 		}
 		else
 		{
@@ -682,4 +711,27 @@ EVisibility SActorDetails::GetNativeComponentWarningVisibility() const
 	}
 	
 	return bIsUneditableNative ? EVisibility::Visible : EVisibility::Collapsed;
-}
\ No newline at end of file
+}
+
+void SActorDetails::ClearNotificationDelegates()
+{
+	if(LockedActorSelection.IsValid())
+	{
+		AActor* SelectedActor = LockedActorSelection.Get();
+		// Unregister for component blueprint compiled events
+		for( auto Component : SelectedActor->GetComponents() )
+		{
+			if(UBlueprintGeneratedClass* ComponentBPGC = Cast<UBlueprintGeneratedClass>(Component->GetClass()))
+			{
+				UBlueprint* ComponentBlueprint = Cast<UBlueprint>(ComponentBPGC->ClassGeneratedBy);
+				ComponentBlueprint->OnCompiled().RemoveAll(this);
+			}
+		}
+	}
+}
+
+void SActorDetails::OnBlueprintRecompiled(UBlueprint* CompiledBlueprint)
+{
+	CrcLastCompiledSignature = CompiledBlueprint->CrcLastCompiledSignature;
+	UpdateComponentTreeFromEditorSelection();
+}
diff --git a/Engine/Source/Editor/LevelEditor/Private/SActorDetails.h b/Engine/Source/Editor/LevelEditor/Private/SActorDetails.h
--- a/Engine/Source/Editor/LevelEditor/Private/SActorDetails.h
+++ b/Engine/Source/Editor/LevelEditor/Private/SActorDetails.h
@@ -44,6 +44,8 @@ private:
 	EVisibility GetNativeComponentWarningVisibility() const;
 	void OnBlueprintedComponentWarningHyperlinkClicked(const FSlateHyperlinkRun::FMetadata& Metadata);
 	void OnNativeComponentWarningHyperlinkClicked(const FSlateHyperlinkRun::FMetadata& Metadata);
+	void OnBlueprintRecompiled(UBlueprint* CompiledBlueprint);
+	void ClearNotificationDelegates();
 
 private:
 	TSharedPtr<SSplitter> DetailsSplitter;
@@ -53,6 +55,7 @@ private:
 
 	// The actor selected when the details panel was locked
 	TWeakObjectPtr<AActor> LockedActorSelection;
+	uint32 CrcLastCompiledSignature;
 
 	// Used to prevent reentrant changes
 	bool bSelectionGuard;
