LevelEditor - Register the "build and submit" as a tab spawner so the menu entry in the toolbar works.

Jira: UE-15402

diff --git a/Engine/Source/Editor/LevelEditor/Private/SLevelEditor.cpp b/Engine/Source/Editor/LevelEditor/Private/SLevelEditor.cpp
--- a/Engine/Source/Editor/LevelEditor/Private/SLevelEditor.cpp
+++ b/Engine/Source/Editor/LevelEditor/Private/SLevelEditor.cpp
@@ -1045,6 +1045,9 @@ TSharedRef<SWidget> SLevelEditor::RestoreContentArea( const TSharedRef<SDockTab>
 				.SetGroup( MenuStructure.GetLevelEditorCategory() )
 				.SetIcon( WorldPropertiesIcon );
 		}
+
+		FTabSpawnerEntry& BuildAndSubmitEntry = LevelEditorTabManager->RegisterTabSpawner(LevelEditorBuildAndSubmitTab, FOnSpawnTab::CreateSP<SLevelEditor, FName, FString>(this, &SLevelEditor::SpawnLevelEditorTab, LevelEditorBuildAndSubmitTab, FString()));
+		BuildAndSubmitEntry.SetAutoGenerateMenuEntry(false);
 	}
 
 	// Rebuild the editor mode commands and their tab spawners before we restore the layout,
