Editor: Various 2D snap layer improvements

Added several additional layer action commands:
- Move selection up a layer (Ctrl+PageUp, also changes active layer)
- Move selection down a layer (Ctrl+PageDown, also changes active layer)
- Move selection to top layer (Shift+Ctrl+PageUp, also changes active layer)
- Move selection to bottom layer (Shift+Ctrl+PageDown, also changes active layer)
- Cycle active snap layer forwards (Alt+PageUp)
- Cycle active snap layer backwards (Alt+PageDown)

Changed the snap layer viewport toolbar button to be a reasonable minimum width, reducing horizontal bounciness when switching layers

Snap layers in the project settings are now automatically sorted by depth when they are edited

diff --git a/Engine/Source/Editor/LevelEditor/Private/LevelEditor.cpp b/Engine/Source/Editor/LevelEditor/Private/LevelEditor.cpp
--- a/Engine/Source/Editor/LevelEditor/Private/LevelEditor.cpp
+++ b/Engine/Source/Editor/LevelEditor/Private/LevelEditor.cpp
@@ -744,9 +744,40 @@ void FLevelEditorModule::BindGlobalLevelEditorCommands()
 		);
 
 	ActionList.MapAction(
-		Commands.SnapToLayer2D,
-		FExecuteAction::CreateStatic(&FLevelEditorActionCallbacks::SnapToLayer2D_Clicked),
-		FCanExecuteAction::CreateStatic(&FLevelEditorActionCallbacks::CanSnapToLayer2D)
+		Commands.SnapTo2DLayer,
+		FExecuteAction::CreateStatic(&FLevelEditorActionCallbacks::SnapTo2DLayer_Clicked),
+		FCanExecuteAction::CreateStatic(&FLevelEditorActionCallbacks::CanSnapTo2DLayer)
+		);
+
+	ActionList.MapAction(
+		Commands.MoveSelectionUpIn2DLayers,
+		FExecuteAction::CreateStatic(&FLevelEditorActionCallbacks::MoveSelectionToDifferent2DLayer_Clicked, /*bGoingUp=*/ true, /*bForceToTopOrBottom=*/ false),
+		FCanExecuteAction::CreateStatic(&FLevelEditorActionCallbacks::CanMoveSelectionToDifferent2DLayer, /*bGoingUp=*/ true)
+		);
+	ActionList.MapAction(
+		Commands.MoveSelectionDownIn2DLayers,
+		FExecuteAction::CreateStatic(&FLevelEditorActionCallbacks::MoveSelectionToDifferent2DLayer_Clicked, /*bGoingUp=*/ false, /*bForceToTopOrBottom=*/ false),
+		FCanExecuteAction::CreateStatic(&FLevelEditorActionCallbacks::CanMoveSelectionToDifferent2DLayer, /*bGoingUp=*/ false)
+		);
+	ActionList.MapAction(
+		Commands.MoveSelectionToTop2DLayer,
+		FExecuteAction::CreateStatic(&FLevelEditorActionCallbacks::MoveSelectionToDifferent2DLayer_Clicked, /*bGoingUp=*/ true, /*bForceToTopOrBottom=*/ true),
+		FCanExecuteAction::CreateStatic(&FLevelEditorActionCallbacks::CanMoveSelectionToDifferent2DLayer, /*bGoingUp=*/ true)
+		);
+	ActionList.MapAction(
+		Commands.MoveSelectionToBottom2DLayer,
+		FExecuteAction::CreateStatic(&FLevelEditorActionCallbacks::MoveSelectionToDifferent2DLayer_Clicked, /*bGoingUp=*/ false, /*bForceToTopOrBottom=*/ true),
+		FCanExecuteAction::CreateStatic(&FLevelEditorActionCallbacks::CanMoveSelectionToDifferent2DLayer, /*bGoingUp=*/ false)
+		);
+
+
+	ActionList.MapAction(
+		Commands.Select2DLayerAbove,
+		FExecuteAction::CreateStatic(&FLevelEditorActionCallbacks::Select2DLayerDeltaAway_Clicked, -1)
+		);
+	ActionList.MapAction(
+		Commands.Select2DLayerBelow,
+		FExecuteAction::CreateStatic(&FLevelEditorActionCallbacks::Select2DLayerDeltaAway_Clicked, 1)
 		);
 
 	bAlign = false;
diff --git a/Engine/Source/Editor/LevelEditor/Private/LevelEditorActions.cpp b/Engine/Source/Editor/LevelEditor/Private/LevelEditorActions.cpp
--- a/Engine/Source/Editor/LevelEditor/Private/LevelEditorActions.cpp
+++ b/Engine/Source/Editor/LevelEditor/Private/LevelEditorActions.cpp
@@ -2524,7 +2524,7 @@ void FLevelEditorActionCallbacks::MoveActorTo_Clicked( const bool InAlign, const
 	GEditor->RebuildAlteredBSP(); // Update the Bsp of any levels containing a modified brush
 }
 
-void FLevelEditorActionCallbacks::SnapToLayer2D_Clicked()
+void FLevelEditorActionCallbacks::SnapTo2DLayer_Clicked()
 {
 	// Fires ULevel::LevelDirtiedEvent when falling out of scope.
 	FScopedLevelDirtied LevelDirtyCallback;
@@ -2571,12 +2571,67 @@ void FLevelEditorActionCallbacks::SnapToLayer2D_Clicked()
 	}
 }
 
-bool FLevelEditorActionCallbacks::CanSnapToLayer2D()
+bool FLevelEditorActionCallbacks::CanSnapTo2DLayer()
 {
 	const ULevelEditor2DSettings* Settings = GetDefault<ULevelEditor2DSettings>();
 	return Settings->SnapLayers.IsValidIndex(GetDefault<ULevelEditorViewportSettings>()->ActiveSnapLayerIndex) && (GEditor->GetSelectedActorCount() > 0);
 }
 
+void FLevelEditorActionCallbacks::MoveSelectionToDifferent2DLayer_Clicked(bool bGoingUp, bool bForceToTopOrBottom)
+{
+	// Change the active layer first
+	const ULevelEditor2DSettings* Settings2D = GetDefault<ULevelEditor2DSettings>();
+	ULevelEditorViewportSettings* SettingsVP = GetMutableDefault<ULevelEditorViewportSettings>();
+
+	const int32 NumLayers = Settings2D->SnapLayers.Num();
+
+	if (NumLayers > 0)
+	{
+		if (bGoingUp)
+		{
+			SettingsVP->ActiveSnapLayerIndex = bForceToTopOrBottom ? 0 : (SettingsVP->ActiveSnapLayerIndex - 1);
+		}
+		else
+		{
+			SettingsVP->ActiveSnapLayerIndex = bForceToTopOrBottom ? (NumLayers - 1) : (SettingsVP->ActiveSnapLayerIndex + 1);
+		}
+		SettingsVP->ActiveSnapLayerIndex = FMath::Clamp<int32>(SettingsVP->ActiveSnapLayerIndex, 0, NumLayers - 1);
+
+		SettingsVP->PostEditChange();
+	}
+
+	// Snap the selection to the new active layer
+	SnapTo2DLayer_Clicked();
+}
+
+bool FLevelEditorActionCallbacks::CanMoveSelectionToDifferent2DLayer(bool bGoingUp)
+{
+	const ULevelEditor2DSettings* Settings2D = GetDefault<ULevelEditor2DSettings>();
+	const ULevelEditorViewportSettings* SettingsVP = GetMutableDefault<ULevelEditorViewportSettings>();
+
+	const int32 SelectedIndex = SettingsVP->ActiveSnapLayerIndex;
+	const int32 NumLayers = Settings2D->SnapLayers.Num();
+
+	return bGoingUp ? (SelectedIndex > 0) : (SelectedIndex + 1 < NumLayers);
+}
+
+void FLevelEditorActionCallbacks::Select2DLayerDeltaAway_Clicked(int32 Delta)
+{
+	const ULevelEditor2DSettings* Settings2D = GetDefault<ULevelEditor2DSettings>();
+	ULevelEditorViewportSettings* SettingsVP = GetMutableDefault<ULevelEditorViewportSettings>();
+
+	const int32 SelectedIndex = SettingsVP->ActiveSnapLayerIndex;
+	const int32 NumLayers = Settings2D->SnapLayers.Num();
+
+	if (NumLayers > 0)
+	{
+		const int32 NewIndex = ((NumLayers + SelectedIndex + Delta) % NumLayers);
+
+		SettingsVP->ActiveSnapLayerIndex = NewIndex;
+		SettingsVP->PostEditChange();
+	}
+}
+
 void FLevelEditorActionCallbacks::SnapToFloor_Clicked( bool InAlign, bool InUseLineTrace, bool InUseBounds, bool InUsePivot )
 {
 	const FScopedTransaction Transaction( NSLOCTEXT("UnrealEd", "SnapActorsToFloor", "Snap Actors To Floor") );
@@ -2878,7 +2933,16 @@ void FLevelEditorCommands::RegisterCommands()
 	UI_COMMAND( SnapOriginToGrid, "Snap Origin to Grid", "Snaps the actor to the nearest grid location at its origin", EUserInterfaceActionType::Button, FInputChord( EModifierKey::Control, EKeys::End ) );
 	UI_COMMAND( SnapOriginToGridPerActor, "Snap Origin to Grid Per Actor", "Snaps each selected actor separately to the nearest grid location at its origin", EUserInterfaceActionType::Button, FInputChord() );
 	UI_COMMAND( AlignOriginToGrid, "Align Origin to Grid", "Aligns the actor to the nearest grid location at its origin", EUserInterfaceActionType::Button, FInputChord() );
-	UI_COMMAND( SnapToLayer2D, "Snap to 2D Layer", "Snaps the actor to the current 2D snap layer", EUserInterfaceActionType::Button, FInputChord( EModifierKey::Control, EKeys::SpaceBar ) );
+
+	UI_COMMAND( SnapTo2DLayer, "Snap to 2D Layer", "Snaps the actor to the current 2D snap layer", EUserInterfaceActionType::Button, FInputChord( EModifierKey::Control, EKeys::SpaceBar ) );
+	UI_COMMAND( MoveSelectionUpIn2DLayers, "Bring selection forward a snap layer", "Bring selection forward a snap layer", EUserInterfaceActionType::Button, FInputChord(EKeys::PageUp, EModifierKey::Control) );
+	UI_COMMAND( MoveSelectionDownIn2DLayers, "Send selection backward a snap layer", "Send selection backward a snap layer", EUserInterfaceActionType::Button, FInputChord(EKeys::PageDown, EModifierKey::Control) );
+	UI_COMMAND( MoveSelectionToTop2DLayer, "Bring selection to the front snap layer", "Bring selection to the front snap layer", EUserInterfaceActionType::Button, FInputChord(EKeys::PageUp, EModifierKey::Shift | EModifierKey::Control) );
+	UI_COMMAND( MoveSelectionToBottom2DLayer, "Send selection to the back snap layer", "Send selection to the back snap layer", EUserInterfaceActionType::Button, FInputChord(EKeys::PageDown, EModifierKey::Shift | EModifierKey::Control) );
+	UI_COMMAND( Select2DLayerAbove, "Select next 2D layer", "Changes the active layer to the next 2D layer", EUserInterfaceActionType::Button, FInputChord(EKeys::PageUp, EModifierKey::Alt) );
+	UI_COMMAND( Select2DLayerBelow, "Select previous 2D layer", "Changes the active layer to the previous 2D layer", EUserInterfaceActionType::Button, FInputChord(EKeys::PageDown, EModifierKey::Alt) );
+
+
 	UI_COMMAND( SnapToFloor, "Snap to Floor", "Snaps the actor or component to the floor below it", EUserInterfaceActionType::Button, FInputChord( EKeys::End ) );
 	UI_COMMAND( AlignToFloor, "Align to Floor", "Aligns the actor or component with the floor", EUserInterfaceActionType::Button, FInputChord() );
 	UI_COMMAND( SnapPivotToFloor, "Snap Pivot to Floor", "Snaps the actor to the floor at its pivot point", EUserInterfaceActionType::Button, FInputChord( EModifierKey::Alt, EKeys::End ) );	
diff --git a/Engine/Source/Editor/LevelEditor/Private/LevelEditorContextMenu.cpp b/Engine/Source/Editor/LevelEditor/Private/LevelEditorContextMenu.cpp
--- a/Engine/Source/Editor/LevelEditor/Private/LevelEditorContextMenu.cpp
+++ b/Engine/Source/Editor/LevelEditor/Private/LevelEditorContextMenu.cpp
@@ -864,7 +864,7 @@ void FLevelEditorContextMenuImpl::FillSnapAlignMenu( FMenuBuilder& MenuBuilder )
 	MenuBuilder.AddMenuEntry( FLevelEditorCommands::Get().SnapOriginToGrid );
 	MenuBuilder.AddMenuEntry( FLevelEditorCommands::Get().SnapOriginToGridPerActor );
 	MenuBuilder.AddMenuEntry( FLevelEditorCommands::Get().AlignOriginToGrid );
-	MenuBuilder.AddMenuEntry( FLevelEditorCommands::Get().SnapToLayer2D );
+	MenuBuilder.AddMenuEntry( FLevelEditorCommands::Get().SnapTo2DLayer );
 	MenuBuilder.AddMenuEntry( FLevelEditorCommands::Get().SnapToFloor );
 	MenuBuilder.AddMenuEntry( FLevelEditorCommands::Get().AlignToFloor );
 	MenuBuilder.AddMenuEntry( FLevelEditorCommands::Get().SnapPivotToFloor );
diff --git a/Engine/Source/Editor/LevelEditor/Public/LevelEditorActions.h b/Engine/Source/Editor/LevelEditor/Public/LevelEditorActions.h
--- a/Engine/Source/Editor/LevelEditor/Public/LevelEditorActions.h
+++ b/Engine/Source/Editor/LevelEditor/Public/LevelEditorActions.h
@@ -137,7 +137,25 @@ public:
 	TSharedPtr< FUICommandInfo > AlignOriginToGrid;
 
 	/** Snaps the actor to the 2D layer */
-	TSharedPtr< FUICommandInfo > SnapToLayer2D;
+	TSharedPtr< FUICommandInfo > SnapTo2DLayer;
+
+	/** Moves the selected actors up one 2D layer (changing the active layer at the same time) */
+	TSharedPtr< FUICommandInfo > MoveSelectionUpIn2DLayers;
+
+	/** Moves the selected actors down one 2D layer (changing the active layer at the same time) */
+	TSharedPtr< FUICommandInfo > MoveSelectionDownIn2DLayers;
+
+	/** Moves the selected actors to the top 2D layer (changing the active layer at the same time) */
+	TSharedPtr< FUICommandInfo > MoveSelectionToTop2DLayer;
+
+	/** Moves the selected actors to the bottom 2D layer (changing the active layer at the same time) */
+	TSharedPtr< FUICommandInfo > MoveSelectionToBottom2DLayer;
+
+	/** Changes the active 2D layer to one above the current one */
+	TSharedPtr< FUICommandInfo > Select2DLayerAbove;
+
+	/** Changes the active 2D layer to one below the current one */
+	TSharedPtr< FUICommandInfo > Select2DLayerBelow;
 
 	/** Snaps the actor to the floor*/
 	TSharedPtr< FUICommandInfo > SnapToFloor;
@@ -1151,13 +1169,29 @@ public:
 	/** 
 	 * Snaps an actor to the currently selected 2D snap layer
 	 */
-	static void SnapToLayer2D_Clicked();
+	static void SnapTo2DLayer_Clicked();
 
 	/**
 	 * Checks to see if at least a single actor is selected and the 2D editor mode is enabled
 	 *	@return true if it can execute.
 	 */
-	static bool CanSnapToLayer2D();
+	static bool CanSnapTo2DLayer();
+
+	/** 
+	 * Snaps an actor to the currently selected 2D snap layer
+	 */
+	static void MoveSelectionToDifferent2DLayer_Clicked(bool bGoingUp, bool bForceToTopOrBottom);
+
+	/**
+	 * Checks to see if at least a single actor is selected and the 2D editor mode is enabled and there is a layer above/below the current setting
+	 *	@return true if it can execute.
+	 */
+	static bool CanMoveSelectionToDifferent2DLayer(bool bGoingUp);
+
+	/**
+	 * Changes the active 2D snap layer to one a delta above or below the current layer
+	 */
+	static void Select2DLayerDeltaAway_Clicked(int32 Delta);
 
 	/** 
 	 * Snaps an actor to the floor.  Optionally will align with the trace normal.
diff --git a/Engine/Source/Editor/UnrealEd/Private/STransformViewportToolbar.cpp b/Engine/Source/Editor/UnrealEd/Private/STransformViewportToolbar.cpp
--- a/Engine/Source/Editor/UnrealEd/Private/STransformViewportToolbar.cpp
+++ b/Engine/Source/Editor/UnrealEd/Private/STransformViewportToolbar.cpp
@@ -290,7 +290,8 @@ TSharedRef< SWidget > STransformViewportToolBar::MakeTransformToolBar( const TSh
 
 		static FName Layer2DSnapName = FName(TEXT("Layer2DSnap"));
 
-		ToolbarBuilder.AddWidget(SNew(SViewportToolBarComboMenu)
+		TSharedRef<SWidget> SnapLayerPickerWidget =
+			SNew(SViewportToolBarComboMenu)
 			.Cursor(EMouseCursor::Default)
 			.Style(ToolBarStyle)
 			.Visibility(this, &STransformViewportToolBar::IsLayer2DSnapVisible)
@@ -302,7 +303,9 @@ TSharedRef< SWidget > STransformViewportToolBar::MakeTransformToolBar( const TSh
 			.MenuButtonToolTip(LOCTEXT("Layer2DSnap_ToolTip", "Set the 2d layer snap value"))
 			.Icon(Command->GetIcon())
 			.ParentToolBar(SharedThis(this))
-			, Layer2DSnapName);
+			.MinDesiredButtonWidth(88.0f);
+
+		ToolbarBuilder.AddWidget(SnapLayerPickerWidget, Layer2DSnapName);
 	}
 	ToolbarBuilder.EndSection();
 
@@ -668,7 +671,17 @@ TSharedRef<SWidget> STransformViewportToolBar::FillLayer2DSnapMenu()
 
 
 	FLevelEditorModule& LevelEditor = FModuleManager::GetModuleChecked<FLevelEditorModule>(TEXT("LevelEditor"));
-	ShowMenuBuilder.AddMenuEntry(LevelEditor.GetLevelEditorCommands().SnapToLayer2D);
+	ShowMenuBuilder.AddMenuEntry(LevelEditor.GetLevelEditorCommands().SnapTo2DLayer);
+
+	ShowMenuBuilder.AddMenuSeparator();
+	ShowMenuBuilder.AddMenuEntry(LevelEditor.GetLevelEditorCommands().MoveSelectionToTop2DLayer);
+	ShowMenuBuilder.AddMenuEntry(LevelEditor.GetLevelEditorCommands().MoveSelectionUpIn2DLayers);
+	ShowMenuBuilder.AddMenuEntry(LevelEditor.GetLevelEditorCommands().MoveSelectionDownIn2DLayers);
+	ShowMenuBuilder.AddMenuEntry(LevelEditor.GetLevelEditorCommands().MoveSelectionToBottom2DLayer);
+	
+	ShowMenuBuilder.AddMenuSeparator();
+	ShowMenuBuilder.AddMenuEntry(LevelEditor.GetLevelEditorCommands().Select2DLayerAbove);
+	ShowMenuBuilder.AddMenuEntry(LevelEditor.GetLevelEditorCommands().Select2DLayerBelow);
 
 	return ShowMenuBuilder.MakeWidget();
 }
diff --git a/Engine/Source/Editor/UnrealEd/Private/Settings/EditorProjectSettings.cpp b/Engine/Source/Editor/UnrealEd/Private/Settings/EditorProjectSettings.cpp
--- a/Engine/Source/Editor/UnrealEd/Private/Settings/EditorProjectSettings.cpp
+++ b/Engine/Source/Editor/UnrealEd/Private/Settings/EditorProjectSettings.cpp
@@ -149,3 +149,11 @@ ULevelEditor2DSettings::ULevelEditor2DSettings(const FObjectInitializer& ObjectI
 	SnapLayers.Emplace(TEXT("Default"), 0.0f);
 	SnapLayers.Emplace(TEXT("Background"), -100.0f);
 }
+
+void ULevelEditor2DSettings::PostEditChangeProperty(struct FPropertyChangedEvent& PropertyChangedEvent)
+{
+	// Sort the snap layers
+	SnapLayers.Sort([](const FMode2DLayer& LHS, const FMode2DLayer& RHS){ return LHS.Depth > RHS.Depth; });
+
+	Super::PostEditChangeProperty(PropertyChangedEvent);
+}
diff --git a/Engine/Source/Editor/UnrealEd/Public/Settings/EditorProjectSettings.h b/Engine/Source/Editor/UnrealEd/Public/Settings/EditorProjectSettings.h
--- a/Engine/Source/Editor/UnrealEd/Public/Settings/EditorProjectSettings.h
+++ b/Engine/Source/Editor/UnrealEd/Public/Settings/EditorProjectSettings.h
@@ -126,5 +126,11 @@ public:
 	/** Snap layers that are displayed in the viewport toolbar */
 	UPROPERTY(EditAnywhere, config, Category = LayerSnapping)
 	TArray<FMode2DLayer> SnapLayers;
+
+public:
+	// UObject interface
+	virtual void PostEditChangeProperty(struct FPropertyChangedEvent& PropertyChangedEvent) override;
+	// End of UObject interface
+
 };
 
