Editor: Allowed shifting snap layers to still work even if at the top/bottom layer, functioning as a shortcut for just snapping the selection to the current layer

diff --git a/Engine/Source/Editor/LevelEditor/Private/LevelEditorActions.cpp b/Engine/Source/Editor/LevelEditor/Private/LevelEditorActions.cpp
--- a/Engine/Source/Editor/LevelEditor/Private/LevelEditorActions.cpp
+++ b/Engine/Source/Editor/LevelEditor/Private/LevelEditorActions.cpp
@@ -2587,17 +2587,16 @@ void FLevelEditorActionCallbacks::MoveSelectionToDifferent2DLayer_Clicked(bool b
 
 	if (NumLayers > 0)
 	{
-		if (bGoingUp)
+		if (bGoingUp && (SettingsVP->ActiveSnapLayerIndex > 0))
 		{
 			SettingsVP->ActiveSnapLayerIndex = bForceToTopOrBottom ? 0 : (SettingsVP->ActiveSnapLayerIndex - 1);
+			SettingsVP->PostEditChange();
 		}
-		else
+		else if (!bGoingUp && ((SettingsVP->ActiveSnapLayerIndex + 1) < NumLayers))
 		{
 			SettingsVP->ActiveSnapLayerIndex = bForceToTopOrBottom ? (NumLayers - 1) : (SettingsVP->ActiveSnapLayerIndex + 1);
+			SettingsVP->PostEditChange();
 		}
-		SettingsVP->ActiveSnapLayerIndex = FMath::Clamp<int32>(SettingsVP->ActiveSnapLayerIndex, 0, NumLayers - 1);
-
-		SettingsVP->PostEditChange();
 	}
 
 	// Snap the selection to the new active layer
@@ -2612,7 +2611,11 @@ bool FLevelEditorActionCallbacks::CanMoveSelectionToDifferent2DLayer(bool bGoing
 	const int32 SelectedIndex = SettingsVP->ActiveSnapLayerIndex;
 	const int32 NumLayers = Settings2D->SnapLayers.Num();
 
-	return bGoingUp ? (SelectedIndex > 0) : (SelectedIndex + 1 < NumLayers);
+	const bool bHasLayerInDirection = bGoingUp ? (SelectedIndex > 0) : (SelectedIndex + 1 < NumLayers);
+	const bool bHasActorSelected = GEditor->GetSelectedActorCount() > 0;
+
+	// Allow it even if there is no layer in the corresponding direction, to let it double as a snap operation shortcut even when at the end stops
+	return bHasLayerInDirection || bHasActorSelected;
 }
 
 void FLevelEditorActionCallbacks::Select2DLayerDeltaAway_Clicked(int32 Delta)
