Deprecated PreUnloadCallback and PostLoadCallback for modules
- PreUnloadCallback() was not called in all cases a module could be unloaded (Buggy!)
- PostLoadCallback() was never called for modules unless the module happened to be loaded with LoadModuleWithCallback().  (Very confusing!)
- These functions were really just called when dynamically unloading and loading modules (not hot reload, and not when starting up or shutting down.)  This was very confusing.  For the most part you can achieve the same result by overriding StartupModule() and ShutdownModule().
- Additional context: CL 1607584 and https://udn.unrealengine.com/questions/246629/why-is-preunloadcallback-not-called-in-unloadmodul.html

diff --git a/Engine/Source/Editor/LevelEditor/Private/LevelEditor.cpp b/Engine/Source/Editor/LevelEditor/Private/LevelEditor.cpp
--- a/Engine/Source/Editor/LevelEditor/Private/LevelEditor.cpp
+++ b/Engine/Source/Editor/LevelEditor/Private/LevelEditor.cpp
@@ -270,22 +270,6 @@ void FLevelEditorModule::ShutdownModule()
 	FLevelViewportCommands::Unregister();
 }
 
-void FLevelEditorModule::PreUnloadCallback()
-{
-	// Disable the "tab closed" delegate that closes the editor if the level editor tab is closed
-	IMainFrameModule& MainFrameModule = FModuleManager::LoadModuleChecked<IMainFrameModule>( MainFrame );
-	MainFrameModule.DisableTabClosedDelegate();
-}
-
-void FLevelEditorModule::PostLoadCallback()
-{
-	// Re-open the level editor tab and re-enable the "tab closed" delegate
-	IMainFrameModule& MainFrameModule = FModuleManager::LoadModuleChecked<IMainFrameModule>( MainFrame );
-	TSharedRef<SDockTab> LevelEditorTab = FGlobalTabmanager::Get()->InvokeTab(FTabId("LevelEditor"));
-	MainFrameModule.SetMainTab(LevelEditorTab);
-	MainFrameModule.EnableTabClosedDelegate();
-}
-
 /**
  * Spawns a new property viewer
  * @todo This only works with the first level editor. Fix it.
diff --git a/Engine/Source/Editor/LevelEditor/Public/LevelEditor.h b/Engine/Source/Editor/LevelEditor/Public/LevelEditor.h
--- a/Engine/Source/Editor/LevelEditor/Public/LevelEditor.h
+++ b/Engine/Source/Editor/LevelEditor/Public/LevelEditor.h
@@ -31,16 +31,6 @@ public:
 	 */
 	virtual void ShutdownModule() override;
 
-	/**
-	 * Called before the module has been unloaded
-	 */
-	virtual void PreUnloadCallback() override;
-
-	/**
-	 * Called after the module has been reloaded
-	 */
-	virtual void PostLoadCallback() override;
-
 	/**
 	 * Override this to set whether your module is allowed to be unloaded on the fly
 	 *
diff --git a/Engine/Source/Runtime/Core/Private/Modules/ModuleManager.cpp b/Engine/Source/Runtime/Core/Private/Modules/ModuleManager.cpp
--- a/Engine/Source/Runtime/Core/Private/Modules/ModuleManager.cpp
+++ b/Engine/Source/Runtime/Core/Private/Modules/ModuleManager.cpp
@@ -939,7 +939,9 @@ void FModuleManager::UnloadOrAbandonModuleWithCallback(const FName InModuleName,
 {
 	auto Module = FindModuleChecked(InModuleName);
 	
+PRAGMA_DISABLE_DEPRECATION_WARNINGS
 	Module->Module->PreUnloadCallback();
+PRAGMA_ENABLE_DEPRECATION_WARNINGS
 
 	const bool bIsHotReloadable = DoesLoadedModuleHaveUObjects( InModuleName );
 	if (!bAbandonOnly && bIsHotReloadable && Module->Module->SupportsDynamicReloading())
@@ -968,7 +970,9 @@ bool FModuleManager::LoadModuleWithCallback( const FName InModuleName, FOutputDe
 
 	if (bWasSuccessful && LoadedModule.IsValid())
 	{
+PRAGMA_DISABLE_DEPRECATION_WARNINGS
 		LoadedModule->PostLoadCallback();
+PRAGMA_ENABLE_DEPRECATION_WARNINGS
 	}
 	else
 	{
diff --git a/Engine/Source/Runtime/Core/Public/Modules/ModuleInterface.h b/Engine/Source/Runtime/Core/Public/Modules/ModuleInterface.h
--- a/Engine/Source/Runtime/Core/Public/Modules/ModuleInterface.h
+++ b/Engine/Source/Runtime/Core/Public/Modules/ModuleInterface.h
@@ -32,6 +32,7 @@ public:
 	/**
 	 * Called before the module has been unloaded
 	 */
+	DEPRECATED(4.9, "PreUnloadCallback() was deprecated because it was not called in all expected cases and generally was not useful.  Override the ShutdownModule() function instead.")
 	virtual void PreUnloadCallback()
 	{
 	}
@@ -39,6 +40,7 @@ public:
 	/**
 	 * Called after the module has been reloaded
 	 */
+	DEPRECATED(4.9, "PostLoadCallback() was deprecated because it was not called in all expected cases and generally was not useful.  Override the StartupModule() function instead.")
 	virtual void PostLoadCallback()
 	{
 	}
