Added (optional) display of Engine Version Number in the ProjectBadge

diff --git a/Engine/Source/Editor/LevelEditor/Private/LevelEditor.cpp b/Engine/Source/Editor/LevelEditor/Private/LevelEditor.cpp
--- a/Engine/Source/Editor/LevelEditor/Private/LevelEditor.cpp
+++ b/Engine/Source/Editor/LevelEditor/Private/LevelEditor.cpp
@@ -59,9 +59,12 @@ public:
 		FColor BadgeTextColor = FColor(128,128,128,255);
 		GConfig->GetColor(TEXT("LevelEditor"), TEXT("ProjectBadgeTextColor"), /*out*/ BadgeTextColor, GEditorPerProjectIni);
 
+		const FString EngineVersionString = GEngineVersion.ToString(GEngineVersion.IsPromotedBuild() ? EVersionComponent::Changelist : EVersionComponent::Patch);
+		
 		FFormatNamedArguments Args;
 		Args.Add(TEXT("Branch"), FText::FromString(OptionalBranchPrefix));
 		Args.Add(TEXT("GameName"), FText::FromString(FString(FApp::GetGameName())));
+		Args.Add(TEXT("EngineVersion"), (GetDefault<UEditorPerProjectUserSettings>()->bDisplayEngineVersionInBadge) ? FText::FromString("(" + EngineVersionString + ")") : FText());
 
 		FText RightContentText;
 		FText RightContentTooltip;
@@ -70,17 +73,16 @@ public:
 		if (BuildConfig != EBuildConfigurations::Shipping && BuildConfig != EBuildConfigurations::Development && BuildConfig != EBuildConfigurations::Unknown)
 		{
 			Args.Add(TEXT("Config"), EBuildConfigurations::ToText(BuildConfig));
-			RightContentText = FText::Format(NSLOCTEXT("UnrealEditor", "TitleBarRightContentAndConfig", "{Branch}{GameName} [{Config}]"), Args);
+			RightContentText = FText::Format(NSLOCTEXT("UnrealEditor", "TitleBarRightContentAndConfig", "{Branch}{GameName} [{Config}] {EngineVersion}"), Args);
 		}
 		else
 		{
-			RightContentText = FText::Format(NSLOCTEXT("UnrealEditor", "TitleBarRightContent", "{Branch}{GameName}"), Args);
+			RightContentText = FText::Format(NSLOCTEXT("UnrealEditor", "TitleBarRightContent", "{Branch}{GameName} {EngineVersion}"), Args);
 		}
 
 		// Create the tooltip showing more detailed information
 		FFormatNamedArguments TooltipArgs;
-		const FString EngineVerisonString = GEngineVersion.ToString(GEngineVersion.IsPromotedBuild() ? EVersionComponent::Changelist : EVersionComponent::Patch);
-		TooltipArgs.Add(TEXT("Version"), FText::FromString(EngineVerisonString));
+		TooltipArgs.Add(TEXT("Version"), FText::FromString(EngineVersionString));
 		TooltipArgs.Add(TEXT("Branch"), FText::FromString(FApp::GetBranchName()));
 		TooltipArgs.Add(TEXT("BuildConfiguration"), EBuildConfigurations::ToText(BuildConfig));
 		TooltipArgs.Add(TEXT("BuildDate"), FText::FromString(FApp::GetBuildDate()));
diff --git a/Engine/Source/Editor/UnrealEd/Classes/Editor/EditorPerProjectUserSettings.h b/Engine/Source/Editor/UnrealEd/Classes/Editor/EditorPerProjectUserSettings.h
--- a/Engine/Source/Editor/UnrealEd/Classes/Editor/EditorPerProjectUserSettings.h
+++ b/Engine/Source/Editor/UnrealEd/Classes/Editor/EditorPerProjectUserSettings.h
@@ -30,6 +30,10 @@ class UEditorPerProjectUserSettings : public UObject
 	/** If enabled, behavior tree debugger will collect its data even when all behavior tree editor windows are closed */
 	UPROPERTY(EditAnywhere, config, Category = AI)
 	uint32 bAlwaysGatherBehaviorTreeDebuggerData : 1;
+
+	/** When enabled, Engine Version Number is displayed in the ProjectBadge */
+	UPROPERTY(EditAnywhere, config, Category = DeveloperTools, meta = (DisplayName = "Display Engine Version Number in Project Badge", ConfigRestartRequired = true))
+	bool bDisplayEngineVersionInBadge;
 	
 	/** When enabled, the application frame rate, memory and Unreal object count will be displayed in the main editor UI */
 	UPROPERTY(EditAnywhere, config, Category=Performance)
