Fixed editor viewports not remembering maximized state after shutting down in immersive mode (UE-17108)

diff --git a/Engine/Source/Editor/LevelEditor/Private/LevelViewportLayout.cpp b/Engine/Source/Editor/LevelEditor/Private/LevelViewportLayout.cpp
--- a/Engine/Source/Editor/LevelEditor/Private/LevelViewportLayout.cpp
+++ b/Engine/Source/Editor/LevelEditor/Private/LevelViewportLayout.cpp
@@ -236,8 +236,9 @@ void FLevelViewportLayout::InitCommonLayoutFromString( const FString& SpecificLa
 	{
 		const FString& IniSection = FLayoutSaveRestore::GetAdditionalLayoutConfigIni();
 
+		// NOTE: We don't support starting back up in immersive mode, even if the user shut down with a window that way.  See the
+		// comment below in SaveCommonLayoutString() for more info.
 		GConfig->GetBool(*IniSection, *(SpecificLayoutString + TEXT(".bIsMaximized")), bShouldBeMaximized, GEditorPerProjectIni);
-		GConfig->GetBool(*IniSection, *(SpecificLayoutString + TEXT(".bIsImmersive")), bShouldBeImmersive, GEditorPerProjectIni);
 		GConfig->GetInt(*IniSection, *(SpecificLayoutString + TEXT(".MaximizedViewportID")), MaximizedViewportID, GEditorPerProjectIni);
 	}
 	// Replacement layouts (those selected by the user via a command) don't start maximized so the layout can be seen clearly.
@@ -268,9 +269,18 @@ void FLevelViewportLayout::SaveCommonLayoutString( const FString& SpecificLayout
 		}
 	}
 
-	GConfig->SetBool(*IniSection, *(SpecificLayoutString + TEXT(".bIsMaximized")), bIsMaximizeSupported && bIsMaximized, GEditorPerProjectIni);
-	GConfig->SetBool(*IniSection, *(SpecificLayoutString + TEXT(".bIsImmersive")), bIsImmersive, GEditorPerProjectIni);
-	GConfig->SetInt(*IniSection, *(SpecificLayoutString + TEXT(".MaximizedViewportID")), MaximizedViewportID, GEditorPerProjectIni);
+	// We don't bother saving that we were in immersive mode, because we never want to start back up directly in immersive mode
+	// unless the user asks for that on the command-line.  The reason is it can be disorientating to not see any editor UI when
+	// to restart the editor.  In this case, we'll store the mode they were previously in before they switched to immersive mode.
+	if( bIsImmersive )
+	{
+		GConfig->SetBool( *IniSection, *( SpecificLayoutString + TEXT( ".bIsMaximized" ) ), bIsMaximizeSupported && bWasMaximized, GEditorPerProjectIni );
+	}
+	else
+	{
+		GConfig->SetBool(*IniSection, *(SpecificLayoutString + TEXT(".bIsMaximized")), bIsMaximizeSupported && bIsMaximized, GEditorPerProjectIni);
+	}
+	GConfig->SetInt( *IniSection, *( SpecificLayoutString + TEXT( ".MaximizedViewportID" ) ), MaximizedViewportID, GEditorPerProjectIni );
 }
 
 void FLevelViewportLayout::RequestMaximizeViewport( TSharedRef<class SLevelViewport> ViewportToMaximize, const bool bWantMaximize, const bool bWantImmersive, const bool bAllowAnimation )
