This change adds a delegate to monitor if the currently selected blueprint component is recompiled whilst currently displaying its details in the in world component editor and forces a details refresh in that instance. re-worked change from CL 2547405 due to some remaining issues.
#UE-11573 - Overriden property values do not survive property renames in Actor Components
#Branch UE4
#Proj LevelEditor

#info this change is more scoped and more robust than CL 2547405.

diff --git a/Engine/Source/Editor/LevelEditor/Private/SActorDetails.cpp b/Engine/Source/Editor/LevelEditor/Private/SActorDetails.cpp
--- a/Engine/Source/Editor/LevelEditor/Private/SActorDetails.cpp
+++ b/Engine/Source/Editor/LevelEditor/Private/SActorDetails.cpp
@@ -66,6 +66,7 @@ void SActorDetails::Construct(const FArguments& InArgs, const FName TabIdentifie
 {
 	bSelectionGuard = false;
 	bShowingRootActorNodeSelected = false;
+	bSelectedComponentRecompiled = false;
 
 	USelection::SelectionChangedEvent.AddRaw(this, &SActorDetails::OnEditorSelectionChanged);
 	
@@ -186,7 +187,7 @@ SActorDetails::~SActorDetails()
 {
 	GEditor->UnregisterForUndo(this);
 	USelection::SelectionChangedEvent.RemoveAll(this);
-	ClearNotificationDelegates();
+	RemoveBPComponentCompileEventDelegate();
 
 	auto LevelEditor = FModuleManager::GetModulePtr<FLevelEditorModule>("LevelEditor");
 	if (LevelEditor != nullptr)
@@ -199,7 +200,6 @@ void SActorDetails::SetObjects(const TArray<UObject*>& InObjects, bool bForceRef
 {
 	if(!DetailsView->IsLocked())
 	{
-		ClearNotificationDelegates();
 		DetailsView->SetObjects(InObjects, bForceRefresh);
 
 		bool bShowingComponents = false;
@@ -212,19 +212,6 @@ void SActorDetails::SetObjects(const TArray<UObject*>& InObjects, bool bForceRef
 				LockedActorSelection = Actor;
 				bShowingComponents = true;
 
-				// Register for component blueprint compiled events
-				for( auto Component : Actor->GetComponents() )
-				{
-					if(UBlueprintGeneratedClass* ComponentBPGC = Cast<UBlueprintGeneratedClass>(Component->GetClass()))
-					{
-						UBlueprint* ComponentBlueprint = Cast<UBlueprint>(ComponentBPGC->ClassGeneratedBy);
-						if(!ComponentBlueprint->OnCompiled().IsBoundToObject( this ))
-						{
-							ComponentBlueprint->OnCompiled().AddSP(this, &SActorDetails::OnBlueprintRecompiled);
-						}
-					}
-				}
-
 				// Update the tree if a new actor is selected
 				if(GEditor->GetSelectedComponentCount() == 0)
 				{
@@ -346,6 +333,12 @@ void SActorDetails::OnSCSEditorTreeViewSelectionChanged(const TArray<FSCSEditorT
 {
 	if (!bSelectionGuard && SelectedNodes.Num() > 0)
 	{
+		if( SelectedNodes.Num() > 1 && SelectedBPComponentBlueprint.IsValid() )
+		{
+			// Remove the compilation delegate if we are no longer displaying the full details for a single blueprint component.
+			RemoveBPComponentCompileEventDelegate();
+		}
+
 		AActor* Actor = GetActorContext();
 		if (Actor)
 		{
@@ -381,12 +374,26 @@ void SActorDetails::OnSCSEditorTreeViewSelectionChanged(const TArray<FSCSEditorT
 				}
 				else
 				{
+					const bool bSingleComponentSelection = SelectedNodes.Num() == 1;
+
 					for (auto& SelectedNode : SelectedNodes)
 					{
 						UActorComponent* ComponentInstance = SelectedNode->FindComponentInstanceInActor(Actor);
 						if (ComponentInstance)
 						{
 							DetailsObjects.Add(ComponentInstance);
+
+							if(bSingleComponentSelection)
+							{
+								// Add delegate to monitor blueprint component compilation if we have a full details view ( i.e. single selection )
+								if(UBlueprintGeneratedClass* ComponentBPGC = Cast<UBlueprintGeneratedClass>(ComponentInstance->GetClass()))
+								{
+									if(UBlueprint* ComponentBlueprint = Cast<UBlueprint>(ComponentBPGC->ClassGeneratedBy))
+									{
+										AddBPComponentCompileEventDelegate(ComponentBlueprint);
+									}
+								}
+							}
 						}
 					}
 				}
@@ -452,6 +459,8 @@ void SActorDetails::OnSCSEditorTreeViewSelectionChanged(const TArray<FSCSEditorT
 					}
 					else
 					{
+						const bool bSingleComponentSelection = SelectedNodes.Num() == 1;
+
 						for (auto& SelectedNode : SelectedNodes)
 						{
 							if (SelectedNode.IsValid())
@@ -462,6 +471,17 @@ void SActorDetails::OnSCSEditorTreeViewSelectionChanged(const TArray<FSCSEditorT
 									DetailsObjects.Add(ComponentInstance);
 									SelectedComponents->Select(ComponentInstance);
 
+									if(bSingleComponentSelection)
+									{
+										// Add delegate to monitor blueprint component compilation if we have a full details view ( i.e. single selection )
+										if(UBlueprintGeneratedClass* ComponentBPGC = Cast<UBlueprintGeneratedClass>(ComponentInstance->GetClass()))
+										{
+											if(UBlueprint* ComponentBlueprint = Cast<UBlueprint>(ComponentBPGC->ClassGeneratedBy))
+											{
+												AddBPComponentCompileEventDelegate(ComponentBlueprint);
+											}
+										}
+									}
 									// Ensure the selection override is bound for this component (including any attached editor-only children)
 									auto SceneComponent = Cast<USceneComponent>(ComponentInstance);
 									if (SceneComponent)
@@ -511,7 +531,6 @@ void SActorDetails::UpdateComponentTreeFromEditorSelection()
 
 		auto& SCSTreeWidget = SCSEditor->SCSTreeWidget;
 		TArray<UObject*> DetailsObjects;
-		bool bForceRefresh = false;
 
 		// Update the tree selection to match the level editor component selection
 		SCSTreeWidget->ClearSelection();
@@ -519,18 +538,6 @@ void SActorDetails::UpdateComponentTreeFromEditorSelection()
 		{
 			UActorComponent* Component = CastChecked<UActorComponent>(*It);
 
-			if(!bForceRefresh)
-			{
-				if(UBlueprintGeneratedClass* ComponentBPGC = Cast<UBlueprintGeneratedClass>(Component->GetClass()))
-				{
-					UBlueprint* ComponentBlueprint = Cast<UBlueprint>(ComponentBPGC->ClassGeneratedBy);
-					if(ComponentBlueprint->CrcLastCompiledSignature == CrcLastCompiledSignature )
-					{
-						bForceRefresh = true;
-					}
-				}
-			}
-
 			auto SCSTreeNode = SCSEditor->GetNodeFromActorComponent(Component, false);
 			if (SCSTreeNode.IsValid() && SCSTreeNode->GetComponentTemplate())
 			{
@@ -542,11 +549,10 @@ void SActorDetails::UpdateComponentTreeFromEditorSelection()
 				DetailsObjects.Add(Component);
 			}
 		}
-		CrcLastCompiledSignature = 0;
 
 		if (DetailsObjects.Num() > 0)
 		{
-			DetailsView->SetObjects(DetailsObjects,bForceRefresh);
+			DetailsView->SetObjects(DetailsObjects, bSelectedComponentRecompiled);
 		}
 		else
 		{
@@ -713,25 +719,34 @@ EVisibility SActorDetails::GetNativeComponentWarningVisibility() const
 	return bIsUneditableNative ? EVisibility::Visible : EVisibility::Collapsed;
 }
 
-void SActorDetails::ClearNotificationDelegates()
+void SActorDetails::AddBPComponentCompileEventDelegate(UBlueprint* ComponentBlueprint)
 {
-	if(LockedActorSelection.IsValid())
+	if(SelectedBPComponentBlueprint.Get() != ComponentBlueprint)
 	{
-		AActor* SelectedActor = LockedActorSelection.Get();
-		// Unregister for component blueprint compiled events
-		for( auto Component : SelectedActor->GetComponents() )
+		RemoveBPComponentCompileEventDelegate();
+		SelectedBPComponentBlueprint = ComponentBlueprint;
+		// Add blueprint component compilation event delegate
+		if(!ComponentBlueprint->OnCompiled().IsBoundToObject(this))
 		{
-			if(UBlueprintGeneratedClass* ComponentBPGC = Cast<UBlueprintGeneratedClass>(Component->GetClass()))
-			{
-				UBlueprint* ComponentBlueprint = Cast<UBlueprint>(ComponentBPGC->ClassGeneratedBy);
-				ComponentBlueprint->OnCompiled().RemoveAll(this);
-			}
+			ComponentBlueprint->OnCompiled().AddSP(this, &SActorDetails::OnBlueprintComponentCompiled);
 		}
 	}
 }
 
-void SActorDetails::OnBlueprintRecompiled(UBlueprint* CompiledBlueprint)
+void SActorDetails::RemoveBPComponentCompileEventDelegate()
+{
+	// Remove blueprint component compilation event delegate
+	if(SelectedBPComponentBlueprint.IsValid())
+	{
+		SelectedBPComponentBlueprint.Get()->OnCompiled().RemoveAll(this);
+		SelectedBPComponentBlueprint.Reset();
+		bSelectedComponentRecompiled = false;
+	}
+}
+
+void SActorDetails::OnBlueprintComponentCompiled(UBlueprint* ComponentBlueprint)
 {
-	CrcLastCompiledSignature = CompiledBlueprint->CrcLastCompiledSignature;
+	bSelectedComponentRecompiled = true;
 	UpdateComponentTreeFromEditorSelection();
+	bSelectedComponentRecompiled = false;
 }
diff --git a/Engine/Source/Editor/LevelEditor/Private/SActorDetails.h b/Engine/Source/Editor/LevelEditor/Private/SActorDetails.h
--- a/Engine/Source/Editor/LevelEditor/Private/SActorDetails.h
+++ b/Engine/Source/Editor/LevelEditor/Private/SActorDetails.h
@@ -44,8 +44,10 @@ private:
 	EVisibility GetNativeComponentWarningVisibility() const;
 	void OnBlueprintedComponentWarningHyperlinkClicked(const FSlateHyperlinkRun::FMetadata& Metadata);
 	void OnNativeComponentWarningHyperlinkClicked(const FSlateHyperlinkRun::FMetadata& Metadata);
-	void OnBlueprintRecompiled(UBlueprint* CompiledBlueprint);
-	void ClearNotificationDelegates();
+
+	void AddBPComponentCompileEventDelegate(UBlueprint* ComponentBlueprint);
+	void RemoveBPComponentCompileEventDelegate();
+	void OnBlueprintComponentCompiled(UBlueprint* ComponentBlueprint);
 
 private:
 	TSharedPtr<SSplitter> DetailsSplitter;
@@ -55,7 +57,10 @@ private:
 
 	// The actor selected when the details panel was locked
 	TWeakObjectPtr<AActor> LockedActorSelection;
-	uint32 CrcLastCompiledSignature;
+
+	// The current component blueprint selection
+	TWeakObjectPtr<UBlueprint> SelectedBPComponentBlueprint;
+	bool bSelectedComponentRecompiled;
 
 	// Used to prevent reentrant changes
 	bool bSelectionGuard;
