Fix issues with sequencer not opening properly when clicking on a new movie scene from the content browser when one is already open.

diff --git a/Engine/Source/Editor/LevelEditor/Private/LevelEditor.cpp b/Engine/Source/Editor/LevelEditor/Private/LevelEditor.cpp
--- a/Engine/Source/Editor/LevelEditor/Private/LevelEditor.cpp
+++ b/Engine/Source/Editor/LevelEditor/Private/LevelEditor.cpp
@@ -313,32 +313,13 @@ void FLevelEditorModule::SummonWorldBrowserComposition()
 }
 
 // @todo remove when world-centric mode is added
-void FLevelEditorModule::AttachSequencer(TSharedPtr<SWidget> SequencerWidget, TSharedPtr<IAssetEditorInstance> SequencerAssetEditor )
+void FLevelEditorModule::AttachSequencer( TSharedPtr<SWidget> SequencerWidget, TSharedPtr<IAssetEditorInstance> SequencerAssetEditor )
 {
 	if( FParse::Param( FCommandLine::Get(), TEXT( "Sequencer" ) ) )
 	{
-		struct Local
-		{
-			static void OnSequencerClosed( TSharedRef<SDockTab> DockTab, TWeakPtr<IAssetEditorInstance> InSequencerAssetEditor )
-			{
-				InSequencerAssetEditor.Pin()->CloseWindow();
-			}
-		};
-		
 		TSharedPtr<SLevelEditor> LevelEditorInstance = LevelEditorInstancePtr.Pin();
 
-		if( SequencerWidget.IsValid() && SequencerAssetEditor.IsValid() )
-		{
-			LevelEditorInstance->InvokeTab("Sequencer");
-			LevelEditorInstance->SequencerTab->SetOnTabClosed( SDockTab::FOnTabClosedCallback::CreateStatic( &Local::OnSequencerClosed, TWeakPtr<IAssetEditorInstance>( SequencerAssetEditor ) ) );
-			LevelEditorInstance->SequencerTab->SetContent(SequencerWidget.ToSharedRef());
-		}
-		else
-		{
-			LevelEditorInstance->SequencerTab.Reset();
-		}
-
-	
+		LevelEditorInstance->AttachSequencer( SequencerWidget, SequencerAssetEditor );
 	}
 }
 
diff --git a/Engine/Source/Editor/LevelEditor/Private/SLevelEditor.cpp b/Engine/Source/Editor/LevelEditor/Private/SLevelEditor.cpp
--- a/Engine/Source/Editor/LevelEditor/Private/SLevelEditor.cpp
+++ b/Engine/Source/Editor/LevelEditor/Private/SLevelEditor.cpp
@@ -509,7 +509,43 @@ TSharedRef<FTabManager> SLevelEditor::GetTabManager() const
 	return LevelEditorTabManager.ToSharedRef();
 }
 
+void SLevelEditor::AttachSequencer( TSharedPtr<SWidget> SequencerWidget, TSharedPtr<IAssetEditorInstance> NewSequencerAssetEditor )
+{
+	struct Local
+	{
+		static void OnSequencerClosed( TSharedRef<SDockTab> DockTab, TWeakPtr<IAssetEditorInstance> InSequencerAssetEditor )
+		{
+			InSequencerAssetEditor.Pin()->CloseWindow();
+		}
+	};
+
+	static bool bIsReentrant = false;
+
+	if( !bIsReentrant )
+	{
+		if(SequencerAssetEditor.IsValid())
+		{
+			// Closing the window will invoke this method again but we are handling reopening with a new movie scene ourselves
+			TGuardValue<bool> ReentrantGuard(bIsReentrant, true);
+			// Shutdown cleanly
+			SequencerAssetEditor.Pin()->CloseWindow();
+		}
 
+		if(SequencerWidget.IsValid() && NewSequencerAssetEditor.IsValid())
+		{
+			InvokeTab("Sequencer");
+			SequencerTab->SetOnTabClosed(SDockTab::FOnTabClosedCallback::CreateStatic(&Local::OnSequencerClosed, TWeakPtr<IAssetEditorInstance>(NewSequencerAssetEditor)));
+			SequencerTab->SetContent(SequencerWidget.ToSharedRef());
+
+			SequencerAssetEditor = NewSequencerAssetEditor;
+		}
+		else
+		{
+			SequencerTab.Reset();
+			SequencerAssetEditor.Reset();
+		}
+	}
+}
 
 TSharedRef<SDockTab> SLevelEditor::SummonDetailsPanel( FName TabIdentifier )
 {
diff --git a/Engine/Source/Editor/LevelEditor/Private/SLevelEditor.h b/Engine/Source/Editor/LevelEditor/Private/SLevelEditor.h
--- a/Engine/Source/Editor/LevelEditor/Private/SLevelEditor.h
+++ b/Engine/Source/Editor/LevelEditor/Private/SLevelEditor.h
@@ -121,9 +121,8 @@ public:
 	// Tab Management
 	TSharedRef<FTabManager> GetTabManager() const;
 	
-	// @todo remove when world-centric mode is added
-	TSharedPtr<SDockTab> SequencerTab;
-
+	/** Attaches a sequencer asset editor used to animate objects in the level to this level editor */
+	void AttachSequencer( TSharedPtr<SWidget> SequencerWidget, TSharedPtr<IAssetEditorInstance> NewSequencerAssetEditor );
 private:
 	
 	TSharedRef<SDockTab> SpawnLevelEditorTab(const FSpawnTabArgs& Args, FName TabIdentifier, FString InitializationPayload);
@@ -229,6 +228,10 @@ private:
 
 	/** List of all actor details panels to update when selection changes */
 	TArray< TWeakPtr<class SActorDetails> > AllActorDetailPanels;
+
+	/** Attached sequencer info */
+	TSharedPtr<SDockTab> SequencerTab;
+	TWeakPtr<IAssetEditorInstance> SequencerAssetEditor;
 };
 
 
