Far View Plane for level editor viewports is now saved between sessions

diff --git a/Engine/Source/Editor/LevelEditor/Private/SLevelViewport.cpp b/Engine/Source/Editor/LevelEditor/Private/SLevelViewport.cpp
--- a/Engine/Source/Editor/LevelEditor/Private/SLevelViewport.cpp
+++ b/Engine/Source/Editor/LevelEditor/Private/SLevelViewport.cpp
@@ -372,6 +372,7 @@ void SLevelViewport::ConstructLevelEditorViewportClient( const FArguments& InArg
 	LevelViewportClient->bDrawBaseInfo = true;
 	LevelViewportClient->bDrawVertices = true;
 	LevelViewportClient->ViewFOV = LevelViewportClient->FOVAngle = ViewportInstanceSettings.FOVAngle;
+	LevelViewportClient->OverrideFarClipPlane( ViewportInstanceSettings.FarViewPlane );
 	
 	// Set the selection outline flag based on preferences
 	LevelViewportClient->EngineShowFlags.SelectionOutline = GetDefault<ULevelEditorViewportSettings>()->bUseSelectionOutline;
@@ -1921,6 +1922,8 @@ void SLevelViewport::SaveConfig(const FString& ConfigName)
 	ViewportInstanceSettings.FOVAngle = LevelViewportClient->FOVAngle;
 	ViewportInstanceSettings.bIsRealtime = LevelViewportClient->IsRealtime();
 	ViewportInstanceSettings.bShowStats = LevelViewportClient->ShouldShowStats();
+	ViewportInstanceSettings.FarViewPlane = LevelViewportClient->GetFarClipPlaneOverride();
+
 	if (GetDefault<ULevelEditorViewportSettings>()->bSaveEngineStats)
 	{
 		const TArray<FString>* EnabledStats = NULL;
diff --git a/Engine/Source/Editor/UnrealEd/Classes/Settings/LevelEditorViewportSettings.h b/Engine/Source/Editor/UnrealEd/Classes/Settings/LevelEditorViewportSettings.h
--- a/Engine/Source/Editor/UnrealEd/Classes/Settings/LevelEditorViewportSettings.h
+++ b/Engine/Source/Editor/UnrealEd/Classes/Settings/LevelEditorViewportSettings.h
@@ -74,6 +74,7 @@ struct UNREALED_API FLevelEditorViewportInstanceSettings
 		, BufferVisualizationMode()
 		, ExposureSettings()
 		, FOVAngle(EditorViewportDefs::DefaultPerspectiveFOVAngle)
+		, FarViewPlane(0)
 		, bIsRealtime(false)
 		, bShowFPS_DEPRECATED(false)
 		, bShowStats(false)
@@ -117,6 +118,10 @@ struct UNREALED_API FLevelEditorViewportInstanceSettings
 	UPROPERTY(config)
 	float FOVAngle;
 
+	/* Position of the var plane in the editor viewport */
+	UPROPERTY(config)
+	float FarViewPlane;
+
 	/* Whether this viewport is updating in real-time. */
 	UPROPERTY(config)
 	bool bIsRealtime;
