Renamed AllowMatineePreview to AllowCinematicPreview since it will be used for Sequencer and Matinee for the time being.

diff --git a/Engine/Source/Editor/LevelEditor/Private/LevelViewportActions.cpp b/Engine/Source/Editor/LevelEditor/Private/LevelViewportActions.cpp
--- a/Engine/Source/Editor/LevelEditor/Private/LevelViewportActions.cpp
+++ b/Engine/Source/Editor/LevelEditor/Private/LevelViewportActions.cpp
@@ -52,7 +52,7 @@ void FLevelViewportCommands::RegisterCommands()
 
 	UI_COMMAND( ApplyMaterialToActor, "Apply Material", "Attempts to apply a dropped material to this object", EUserInterfaceActionType::Button, FInputChord() );
 
-	UI_COMMAND( AllowMatineePreview, "Allow Matinee Preview", "If enabled, allows matinee to be previewed in this viewport", EUserInterfaceActionType::ToggleButton, FInputChord() );
+	UI_COMMAND( ToggleCinematicPreview, "Toggles Cinemaic Preview", "If enabled, allows Matinee or Sequencer previews to play in this viewport", EUserInterfaceActionType::ToggleButton, FInputChord() );
 
 	UI_COMMAND( FindInLevelScriptBlueprint, "Find In Level Script", "Finds references of a selected actor in the level script blueprint", EUserInterfaceActionType::Button, FInputChord(EModifierKey::Control, EKeys::K) );
 	UI_COMMAND( AdvancedSettings, "Advanced Settings...", "Opens the advanced viewport settings", EUserInterfaceActionType::Button, FInputChord());
diff --git a/Engine/Source/Editor/LevelEditor/Private/SLevelViewport.cpp b/Engine/Source/Editor/LevelEditor/Private/SLevelViewport.cpp
--- a/Engine/Source/Editor/LevelEditor/Private/SLevelViewport.cpp
+++ b/Engine/Source/Editor/LevelEditor/Private/SLevelViewport.cpp
@@ -354,7 +354,7 @@ void SLevelViewport::ConstructLevelEditorViewportClient( const FArguments& InArg
 	{
 		LevelViewportClient->SetViewLocation( EditorViewportDefs::DefaultPerspectiveViewLocation );
 		LevelViewportClient->SetViewRotation( EditorViewportDefs::DefaultPerspectiveViewRotation );
-		LevelViewportClient->SetAllowMatineePreview(true);
+		LevelViewportClient->SetAllowCinematicPreview(true);
 	}
 	LevelViewportClient->SetRealtime(ViewportInstanceSettings.bIsRealtime);
 	LevelViewportClient->SetShowStats(ViewportInstanceSettings.bShowStats);
@@ -1135,10 +1135,10 @@ void SLevelViewport::BindOptionCommands( FUICommandList& CommandList )
 
 
 	CommandList.MapAction(
-		ViewportActions.AllowMatineePreview,
-		FExecuteAction::CreateSP( this, &SLevelViewport::OnToggleAllowMatineePreview ),
+		ViewportActions.ToggleCinematicPreview,
+		FExecuteAction::CreateSP( this, &SLevelViewport::OnToggleAllowCinematicPreview ),
 		FCanExecuteAction(),
-		FIsActionChecked::CreateSP( this, &SLevelViewport::DoesAllowMatineePreview )
+		FIsActionChecked::CreateSP( this, &SLevelViewport::AllowsCinematicPreview )
 		);
 
 
@@ -2040,18 +2040,18 @@ void SLevelViewport::OnClearAllBookMarks()
 	GLevelEditorModeTools().ClearAllBookmarks( LevelViewportClient.Get() );
 }
 
-void SLevelViewport::OnToggleAllowMatineePreview()
+void SLevelViewport::OnToggleAllowCinematicPreview()
 {
 	// Reset the FOV of Viewport for cases where we have been previewing the matinee with a changing FOV
-	LevelViewportClient->ViewFOV = LevelViewportClient->AllowMatineePreview() ? LevelViewportClient->ViewFOV : LevelViewportClient->FOVAngle;
+	LevelViewportClient->ViewFOV = LevelViewportClient->AllowsCinematicPreview() ? LevelViewportClient->ViewFOV : LevelViewportClient->FOVAngle;
 
-	LevelViewportClient->SetAllowMatineePreview( !LevelViewportClient->AllowMatineePreview() );
+	LevelViewportClient->SetAllowCinematicPreview( !LevelViewportClient->AllowsCinematicPreview() );
 	LevelViewportClient->Invalidate( false );
 }
 
-bool SLevelViewport::DoesAllowMatineePreview() const
+bool SLevelViewport::AllowsCinematicPreview() const
 {
-	return LevelViewportClient->AllowMatineePreview();
+	return LevelViewportClient->AllowsCinematicPreview();
 }
 
 void SLevelViewport::OnIncrementPositionGridSize()
@@ -2978,7 +2978,7 @@ void SLevelViewport::PreviewActors( const TArray< AActor* >& ActorsToPreview )
 				ActorPreviewLevelViewportClient->bDisableInput = true;
 
 				// Never allow Matinee to possess these views
-				ActorPreviewLevelViewportClient->SetAllowMatineePreview( false );
+				ActorPreviewLevelViewportClient->SetAllowCinematicPreview( false );
 
 				// Our preview viewport is always visible if our owning SLevelViewport is visible, so we hook up
 				// to the same IsVisible method
diff --git a/Engine/Source/Editor/LevelEditor/Private/SLevelViewportToolBar.cpp b/Engine/Source/Editor/LevelEditor/Private/SLevelViewportToolBar.cpp
--- a/Engine/Source/Editor/LevelEditor/Private/SLevelViewportToolBar.cpp
+++ b/Engine/Source/Editor/LevelEditor/Private/SLevelViewportToolBar.cpp
@@ -526,8 +526,8 @@ TSharedRef<SWidget> SLevelViewportToolBar::GenerateOptionsMenu() const
 		{
 			if( bIsPerspective )
 			{
-				// Allow matinee preview only applies to perspective
-				OptionsMenuBuilder.AddMenuEntry( LevelViewportActions.AllowMatineePreview );
+				// Cinematic preview only applies to perspective
+				OptionsMenuBuilder.AddMenuEntry( LevelViewportActions.ToggleCinematicPreview );
 			}
 
 			OptionsMenuBuilder.AddMenuEntry( LevelViewportActions.ToggleGameView );
diff --git a/Engine/Source/Editor/LevelEditor/Public/LevelViewportActions.h b/Engine/Source/Editor/LevelEditor/Public/LevelViewportActions.h
--- a/Engine/Source/Editor/LevelEditor/Public/LevelViewportActions.h
+++ b/Engine/Source/Editor/LevelEditor/Public/LevelViewportActions.h
@@ -80,8 +80,8 @@ public:
 	/** Reset all show flags to default */
 	TSharedPtr< FUICommandInfo > UseDefaultShowFlags;
 
-	/** Allows this viewport to be used by matinee for previewing */ 
-	TSharedPtr< FUICommandInfo > AllowMatineePreview;
+	/** Allows this viewport to be used for previewing cinematic animations */ 
+	TSharedPtr< FUICommandInfo > ToggleCinematicPreview;
 
 	/** Finds instances of selected object in level script. */ 
 	TSharedPtr< FUICommandInfo > FindInLevelScriptBlueprint;
diff --git a/Engine/Source/Editor/LevelEditor/Public/SLevelViewport.h b/Engine/Source/Editor/LevelEditor/Public/SLevelViewport.h
--- a/Engine/Source/Editor/LevelEditor/Public/SLevelViewport.h
+++ b/Engine/Source/Editor/LevelEditor/Public/SLevelViewport.h
@@ -438,12 +438,12 @@ private:
 	/**
 	 * Called to toggle allowing matinee to use this viewport to preview in
 	 */
-	void OnToggleAllowMatineePreview();
+	void OnToggleAllowCinematicPreview();
 
 	/**
 	 * @return true if this viewport allows matinee to be previewed in it                   
 	 */
-	bool DoesAllowMatineePreview() const;
+	bool AllowsCinematicPreview() const;
 
 	/** Find currently selected actor in the level script.  */
 	void FindSelectedInLevelScript();
diff --git a/Engine/Source/Editor/Matinee/Private/Matinee.cpp b/Engine/Source/Editor/Matinee/Private/Matinee.cpp
--- a/Engine/Source/Editor/Matinee/Private/Matinee.cpp
+++ b/Engine/Source/Editor/Matinee/Private/Matinee.cpp
@@ -339,7 +339,7 @@ void FMatinee::SetAudioRealtimeOverride( bool bAudioIsRealtime ) const
 		FEditorViewportClient* const LevelVC = GEditor->LevelViewportClients[i];
 		if (LevelVC)
 		{
-			if(LevelVC->IsPerspective() && LevelVC->AllowMatineePreview() )
+			if(LevelVC->IsPerspective() && LevelVC->AllowsCinematicPreview() )
 			{
 				LevelVC->SetForcedAudioRealtime(bAudioIsRealtime);
 			}
@@ -351,7 +351,7 @@ void FMatinee::OnToggleAspectRatioBars()
 {
 	if (GCurrentLevelEditingViewportClient)
 	{
-		if(GCurrentLevelEditingViewportClient->IsPerspective() && GCurrentLevelEditingViewportClient->AllowMatineePreview() )
+		if(GCurrentLevelEditingViewportClient->IsPerspective() && GCurrentLevelEditingViewportClient->AllowsCinematicPreview() )
 		{
 			bool bEnabled = !AreAspectRatioBarsEnabled();
 			GCurrentLevelEditingViewportClient->SetShowAspectRatioBarDisplay(bEnabled);
@@ -365,7 +365,7 @@ void FMatinee::OnToggleSafeFrames()
 {
 	if (GCurrentLevelEditingViewportClient)
 	{
-		if(GCurrentLevelEditingViewportClient->IsPerspective() && GCurrentLevelEditingViewportClient->AllowMatineePreview() )
+		if(GCurrentLevelEditingViewportClient->IsPerspective() && GCurrentLevelEditingViewportClient->AllowsCinematicPreview() )
 		{
 			bool bEnabled = !IsSafeFrameDisplayEnabled();
 			GCurrentLevelEditingViewportClient->SetShowSafeFrameBoxDisplay(bEnabled);
@@ -937,7 +937,7 @@ void FMatinee::InitMatinee(const EToolkitMode::Type Mode, const TSharedPtr< clas
 		if(LevelVC)
 		{
 			// If there is a director group, set the perspective viewports to realtime automatically.
-			if(LevelVC->IsPerspective() && LevelVC->AllowMatineePreview())
+			if(LevelVC->IsPerspective() && LevelVC->AllowsCinematicPreview())
 			{				
 				//Ensure Realtime is turned on and store the original setting so we can restore it later.
 				LevelVC->SetRealtime(true, true);
@@ -2060,7 +2060,7 @@ void FMatinee::OnClose()
 		if(LevelVC)
 		{
 			// Turn off realtime when exiting.
-			if( LevelVC->IsPerspective() && LevelVC->AllowMatineePreview() )
+			if( LevelVC->IsPerspective() && LevelVC->AllowsCinematicPreview() )
 			{				
 				//Specify true so RestoreRealtime will allow us to disable Realtime if it was original disabled
 				LevelVC->RestoreRealtime(true);
@@ -2176,7 +2176,7 @@ static void DrawTransparentLine( FCanvas* Canvas, const FVector2D& Start, const
 
 void FMatinee::DrawModeHUD(FEditorViewportClient* ViewportClient,FViewport* Viewport,const FSceneView* View,FCanvas* Canvas)
 {
-	if( ViewportClient->AllowMatineePreview() )
+	if( ViewportClient->AllowsCinematicPreview() )
 	{
 		// Get the size of the viewport
 		const int32 SizeX = Viewport->GetSizeXY().X;
diff --git a/Engine/Source/Editor/Matinee/Private/MatineeTools.cpp b/Engine/Source/Editor/Matinee/Private/MatineeTools.cpp
--- a/Engine/Source/Editor/Matinee/Private/MatineeTools.cpp
+++ b/Engine/Source/Editor/Matinee/Private/MatineeTools.cpp
@@ -108,7 +108,7 @@ void FMatinee::UpdateViewportSettings()
 {
 	if ( GCurrentLevelEditingViewportClient )
 	{
-		if ( GCurrentLevelEditingViewportClient->IsPerspective() && GCurrentLevelEditingViewportClient->AllowMatineePreview() )
+		if ( GCurrentLevelEditingViewportClient->IsPerspective() && GCurrentLevelEditingViewportClient->AllowsCinematicPreview() )
 		{
 			bool bSafeFrames = IsSafeFrameDisplayEnabled();
 			bool bAspectRatioBars = AreAspectRatioBarsEnabled();
@@ -3760,7 +3760,7 @@ void FMatinee::LockCamToGroup(class UInterpGroup* InGroup, const bool bResetView
 			for(int32 i=0; i<GEditor->LevelViewportClients.Num(); i++)
 			{
 				FLevelEditorViewportClient* LevelVC =GEditor->LevelViewportClients[i];
-				if(LevelVC && LevelVC->IsPerspective() && LevelVC->AllowMatineePreview() )
+				if(LevelVC && LevelVC->IsPerspective() && LevelVC->AllowsCinematicPreview() )
 				{
 					LevelVC->RemoveCameraRoll();
 					LevelVC->ViewFOV = LevelVC->FOVAngle;
@@ -3845,7 +3845,7 @@ void FMatinee::UpdateCameraToGroup(const bool bInUpdateStandardViewports, bool b
 		for(int32 i=0; i<GEditor->LevelViewportClients.Num(); i++)
 		{
 			FLevelEditorViewportClient* LevelVC =GEditor->LevelViewportClients[i];
-			if(LevelVC && LevelVC->IsPerspective() && LevelVC->AllowMatineePreview() )
+			if(LevelVC && LevelVC->IsPerspective() && LevelVC->AllowsCinematicPreview() )
 			{
 				UpdateLevelViewport(DefaultViewedActor, LevelVC, FadeAmount, ColorScale, bEnableColorScaling, bUpdateViewportTransform );
 			}
@@ -3955,7 +3955,7 @@ void FMatinee::SaveLevelViewports()
 	for( int32 ViewIndex = 0; ViewIndex < GEditor->LevelViewportClients.Num(); ++ViewIndex )
 	{
 		FLevelEditorViewportClient* LevelVC = GEditor->LevelViewportClients[ ViewIndex ];
-		if( LevelVC && LevelVC->IsPerspective() && LevelVC->AllowMatineePreview() )
+		if( LevelVC && LevelVC->IsPerspective() && LevelVC->AllowsCinematicPreview() )
 		{
 			FMatineeViewSaveData SaveData;
 			SaveData.ViewIndex = ViewIndex;
@@ -3976,7 +3976,7 @@ void FMatinee::RestoreLevelViewports()
 		if( SavedData.ViewIndex < GEditor->LevelViewportClients.Num() )
 		{
 			FLevelEditorViewportClient* LevelVC = GEditor->LevelViewportClients[ SavedData.ViewIndex ];
-			if ( LevelVC && LevelVC->IsPerspective() && LevelVC->AllowMatineePreview() )
+			if ( LevelVC && LevelVC->IsPerspective() && LevelVC->AllowsCinematicPreview() )
 			{
 				LevelVC->SetMatineeActorLock( nullptr );
 				LevelVC->SetViewRotation( SavedData.ViewRotation );
diff --git a/Engine/Source/Editor/Matinee/Private/SMatineeRecorder.cpp b/Engine/Source/Editor/Matinee/Private/SMatineeRecorder.cpp
--- a/Engine/Source/Editor/Matinee/Private/SMatineeRecorder.cpp
+++ b/Engine/Source/Editor/Matinee/Private/SMatineeRecorder.cpp
@@ -121,7 +121,7 @@ void SMatineeRecorder::Construct(const FArguments& InArgs)
 	LevelViewportClient->SetViewRotation( EditorViewportDefs::DefaultPerspectiveViewRotation );
 
 	LevelViewportClient->SetRealtime( true );
-	LevelViewportClient->SetAllowMatineePreview( true );
+	LevelViewportClient->SetAllowCinematicPreview( true );
 	LevelViewportClient->Viewport->SetUserFocus(true);
 	LevelViewportClient->SetMatineeRecordingWindow( ParentMatineeWindow.Pin().Get() );
 
diff --git a/Engine/Source/Editor/UnrealEd/Private/EditorModeInterpolation.cpp b/Engine/Source/Editor/UnrealEd/Private/EditorModeInterpolation.cpp
--- a/Engine/Source/Editor/UnrealEd/Private/EditorModeInterpolation.cpp
+++ b/Engine/Source/Editor/UnrealEd/Private/EditorModeInterpolation.cpp
@@ -102,7 +102,7 @@ void FEdModeInterpEdit::CamMoveNotify(FEditorViewportClient* ViewportClient)
 	if(!InterpEd)
 		return;
 
-	if( ViewportClient->AllowMatineePreview() )
+	if( ViewportClient->AllowsCinematicPreview() )
 	{
 		InterpEd->CamMoved( ViewportClient->GetViewLocation(), ViewportClient->GetViewRotation() );
 	}
diff --git a/Engine/Source/Editor/UnrealEd/Private/EditorViewportClient.cpp b/Engine/Source/Editor/UnrealEd/Private/EditorViewportClient.cpp
--- a/Engine/Source/Editor/UnrealEd/Private/EditorViewportClient.cpp
+++ b/Engine/Source/Editor/UnrealEd/Private/EditorViewportClient.cpp
@@ -214,7 +214,7 @@ int32 FEditorViewportClient::GetCameraSpeedSetting() const
 float const FEditorViewportClient::SafePadding = 0.075f;
 
 FEditorViewportClient::FEditorViewportClient(FEditorModeTools* InModeTools, FPreviewScene* InPreviewScene, const TWeakPtr<SEditorViewport>& InEditorViewportWidget)
-	: bAllowMatineePreview(false)
+	: bAllowCinematicPreview(false)
 	, CameraSpeedSetting(4)
 	, ImmersiveDelegate()
 	, VisibilityDelegate()
diff --git a/Engine/Source/Editor/UnrealEd/Private/LevelEditorViewport.cpp b/Engine/Source/Editor/UnrealEd/Private/LevelEditorViewport.cpp
--- a/Engine/Source/Editor/UnrealEd/Private/LevelEditorViewport.cpp
+++ b/Engine/Source/Editor/UnrealEd/Private/LevelEditorViewport.cpp
@@ -2058,7 +2058,7 @@ void FLevelEditorViewportClient::Tick(float DeltaTime)
 void FLevelEditorViewportClient::UpdateViewForLockedActor()
 {
 	// We can't be locked to a matinee actor if this viewport doesn't allow matinee control
-	if ( !bAllowMatineePreview && ActorLockedByMatinee.IsValid() )
+	if ( !bAllowCinematicPreview && ActorLockedByMatinee.IsValid() )
 	{
 		ActorLockedByMatinee = nullptr;
 	}
@@ -3984,7 +3984,7 @@ void FLevelEditorViewportClient::SetupViewForRendering( FSceneViewFamily& ViewFa
 		}
 	}
 
-	if (ModeTools->GetActiveMode(FBuiltinEditorModes::EM_InterpEdit) == 0 || !AllowMatineePreview())
+	if (ModeTools->GetActiveMode(FBuiltinEditorModes::EM_InterpEdit) == 0 || !AllowsCinematicPreview())
 	{
 		// in the editor, disable camera motion blur and other rendering features that rely on the former frame
 		// unless the view port is Matinee controlled
@@ -4098,7 +4098,7 @@ void FLevelEditorViewportClient::CopyLayoutFromViewport( const FLevelEditorViewp
 	ViewportType = InViewport.ViewportType;
 	SetOrthoZoom( InViewport.GetOrthoZoom() );
 	ActorLockedToCamera = InViewport.ActorLockedToCamera;
-	bAllowMatineePreview = InViewport.bAllowMatineePreview;
+	bAllowCinematicPreview = InViewport.bAllowCinematicPreview;
 }
 
 
diff --git a/Engine/Source/Editor/UnrealEd/Public/EditorViewportClient.h b/Engine/Source/Editor/UnrealEd/Public/EditorViewportClient.h
--- a/Engine/Source/Editor/UnrealEd/Public/EditorViewportClient.h
+++ b/Engine/Source/Editor/UnrealEd/Public/EditorViewportClient.h
@@ -683,18 +683,25 @@ public:
 	virtual bool DropObjectsAtCoordinates(int32 MouseX, int32 MouseY, const TArray<UObject*>& DroppedObjects, TArray<AActor*>& OutNewActors, bool bOnlyDropOnTarget = false, bool bCreateDropPreview = false, bool bSelectActors = true, UActorFactory* FactoryToUse = NULL ) { return false; }
 
 	/** Returns true if the viewport is allowed to be possessed by Matinee for previewing sequences */
-	bool AllowMatineePreview() const { return bAllowMatineePreview; }
+	DEPRECATED(4.9, "AllowMatineePreview is deprecated.  Use AllowsCinematicPreview instead")
+	bool AllowMatineePreview() const { return AllowsCinematicPreview(); }
 
-	/** Sets whether or not this viewport is allowed to be possessed by Matinee */
-	void SetAllowMatineePreview(const bool bInAllowMatineePreview)
+	
+	DEPRECATED(4.9, "SetAllowMatineePreview is deprecated.  Use SetAllowCinematicPreview instead")
+	void SetAllowMatineePreview(const bool bInAllowCinematicPreview)
 	{
-		bAllowMatineePreview = bInAllowMatineePreview;
+		SetAllowCinematicPreview( bInAllowCinematicPreview );
 	}
 
+	/** Returns true if the viewport is allowed to be possessed for previewing cinematic sequences or keyframe animations*/
+	bool AllowsCinematicPreview() const { return bAllowCinematicPreview; }
+
+	/** Sets whether or not this viewport is allowed to be possessed by cinematic/scrubbing tools */
+	void SetAllowCinematicPreview( bool bInAllowCinematicPreview ) { bAllowCinematicPreview = bInAllowCinematicPreview; }
 protected:
 
-	/** true if this window is allowed to be possessed by Matinee for previewing sequences in real-time */
-	bool bAllowMatineePreview;
+	/** true if this window is allowed to be possessed by cinematic tools for previewing sequences in real-time */
+	bool bAllowCinematicPreview;
 
 public:
 	/** True if the window is maximized or floating */
