Allowed mac viewport to display engine controls overlay. Added linux as well.

UE-6657

diff --git a/Engine/Source/Editor/LevelEditor/Private/SLevelViewportControlsPopup.cpp b/Engine/Source/Editor/LevelEditor/Private/SLevelViewportControlsPopup.cpp
--- a/Engine/Source/Editor/LevelEditor/Private/SLevelViewportControlsPopup.cpp
+++ b/Engine/Source/Editor/LevelEditor/Private/SLevelViewportControlsPopup.cpp
@@ -7,9 +7,13 @@
 
 void SLevelViewportControlsPopup::Construct(const FArguments& InArgs)
 {
-#if PLATFORM_WINDOWS
 	PopupSize.Set(252, 558);
+#if PLATFORM_WINDOWS
+	PopupPath = FPaths::ConvertRelativePathToFull(FPaths::Combine(*FPaths::EngineDir(), TEXT("Documentation/Source/Shared/Editor/Overlays/ViewportControlsWindows.html")));
+#elif PLATFORM_LINUX
 	PopupPath = FPaths::ConvertRelativePathToFull(FPaths::Combine(*FPaths::EngineDir(), TEXT("Documentation/Source/Shared/Editor/Overlays/ViewportControlsWindows.html")));
+#elif PLATFORM_MAC
+	PopupPath = FPaths::ConvertRelativePathToFull(FPaths::Combine(*FPaths::EngineDir(), TEXT("Documentation/Source/Shared/Editor/Overlays/ViewportControlsMac.html")));
 #endif
 
 	if (!PopupPath.IsEmpty())
