Fixed UE-18232 SProjectBadge does not consistently show tooltip.

diff --git a/Engine/Source/Editor/LevelEditor/Private/LevelEditor.cpp b/Engine/Source/Editor/LevelEditor/Private/LevelEditor.cpp
--- a/Engine/Source/Editor/LevelEditor/Private/LevelEditor.cpp
+++ b/Engine/Source/Editor/LevelEditor/Private/LevelEditor.cpp
@@ -93,7 +93,7 @@ public:
 		SBox::Construct(SBox::FArguments()
 			.HAlign(HAlign_Right)
 			.VAlign(VAlign_Top)
-			.Padding(FMargin(0.0f, 0.0f, 15.0f, 0.0f))
+			.Padding(FMargin(0.0f, 0.0f, 2.0f, 0.0f))
 			[
 				SNew(SBorder)
 				.BorderImage(FEditorStyle::GetBrush("SProjectBadge.BadgeShape"))
@@ -110,6 +110,10 @@ public:
 			]);
 	}
 
+	FVector2D GetSizeLastFrame() const
+	{
+		return GetDesiredSize();
+	}
 
 	// SWidget interface
 	virtual EWindowZone::Type GetWindowZoneOverride() const override
@@ -117,8 +121,16 @@ public:
 		return EWindowZone::TitleBar;
 	}
 	// End of SWidget interface
+
+private:
+	FGeometry CachedGeometry;
 };
 
+static FMargin GetRoomForBadge(TWeakPtr<SProjectBadge> ProjBadge)
+{
+	return FMargin(8.0f, 0.0f, ProjBadge.Pin()->GetSizeLastFrame().X + 8.0f, 0.0f);
+}
+
 TSharedRef<SDockTab> FLevelEditorModule::SpawnLevelEditor( const FSpawnTabArgs& InArgs )
 {
 	TSharedRef<SDockTab> LevelEditorTab = SNew(SDockTab)
@@ -151,6 +163,9 @@ TSharedRef<SDockTab> FLevelEditorModule::SpawnLevelEditor( const FSpawnTabArgs&
 	IIntroTutorials& IntroTutorials = FModuleManager::LoadModuleChecked<IIntroTutorials>(TEXT("IntroTutorials"));
 	TSharedRef<SWidget> TutorialWidget = IntroTutorials.CreateTutorialsWidget(TEXT("LevelEditor"), OwnerWindow);
 
+	TSharedRef<SProjectBadge> ProjectBadge = SNew(SProjectBadge);
+	TAttribute<FMargin> BadgeSizeGetter = TAttribute<FMargin>::Create(TAttribute<FMargin>::FGetter::CreateStatic(&GetRoomForBadge, TWeakPtr<SProjectBadge>(ProjectBadge)));
+		
 	TSharedPtr< SWidget > RightContent=
 		SNew( SHorizontalBox )
 
@@ -174,7 +189,7 @@ TSharedRef<SDockTab> FLevelEditorModule::SpawnLevelEditor( const FSpawnTabArgs&
 		]
 		+ SHorizontalBox::Slot()
 		.AutoWidth()
-		.Padding(8.0f, 0.0f, 8.0f, 0.0f)
+		.Padding(BadgeSizeGetter)
 		.VAlign(VAlign_Center)
 		[
 			TutorialWidget
@@ -183,13 +198,14 @@ TSharedRef<SDockTab> FLevelEditorModule::SpawnLevelEditor( const FSpawnTabArgs&
 	LevelEditorTab->SetRightContent( RightContent.ToSharedRef() );
 
 	LevelEditorTab->SetBackgroundContent(
-			SNew(SProjectBadge)
+		ProjectBadge
 	);
 	
 	
 	return LevelEditorTab;
 }
 
+
 /**
  * Called right after the module's DLL has been loaded and the module object has been created
  */
diff --git a/Engine/Source/Runtime/Slate/Private/Framework/Docking/SDockingTabStack.cpp b/Engine/Source/Runtime/Slate/Private/Framework/Docking/SDockingTabStack.cpp
--- a/Engine/Source/Runtime/Slate/Private/Framework/Docking/SDockingTabStack.cpp
+++ b/Engine/Source/Runtime/Slate/Private/Framework/Docking/SDockingTabStack.cpp
@@ -79,48 +79,51 @@ void SDockingTabStack::Construct( const FArguments& InArgs, const TSharedRef<FTa
 	const FButtonStyle* const UnhideTabWellButtonStyle = &FCoreStyle::Get().GetWidgetStyle< FButtonStyle >( "Docking.UnhideTabwellButton" );
 
 	// create inline title bar content
-	TitleBarContent = SNew(SHorizontalBox)
+	TitleBarContent = 
+	SNew(SOverlay)
+	+ SOverlay::Slot().Expose(BackgroundContentArea)
+	+ SOverlay::Slot()
+	[
+	
+		
+		SNew(SHorizontalBox)
 		.Visibility(EVisibility::SelfHitTestInvisible)
 
-	+ SHorizontalBox::Slot()
+		+ SHorizontalBox::Slot()
 		.AutoWidth()
 		.Expose(InlineContentAreaLeft)
 
-	+ SHorizontalBox::Slot() 
+		+ SHorizontalBox::Slot() 
 		.FillWidth(1.0f)
 		.VAlign(VAlign_Bottom)
 		.Padding(4.0f, 0.0f, 0.0f, 0.0f)
 		[
-			SNew(SOverlay)
-			+ SOverlay::Slot().Expose(BackgroundContentArea)
-			+ SOverlay::Slot()
-			[
-				SNew(SVerticalBox)
-				.Visibility(EVisibility::SelfHitTestInvisible)
+			SNew(SVerticalBox)
+			.Visibility(EVisibility::SelfHitTestInvisible)
 
-				+ SVerticalBox::Slot()
-				.AutoHeight()
-				[
-					SNew(SSpacer)
-						.Visibility(this, &SDockingTabStack::GetMaximizeSpacerVisibility)
-						.Size(FVector2D(0.0f, 10.0f))
-				]
+			+ SVerticalBox::Slot()
+			.AutoHeight()
+			[
+				SNew(SSpacer)
+				.Visibility(this, &SDockingTabStack::GetMaximizeSpacerVisibility)
+				.Size(FVector2D(0.0f, 10.0f))
+			]
 
-				+ SVerticalBox::Slot()
-				.AutoHeight()
-				[
-					// TabWell
-					SAssignNew(TabWell, SDockingTabWell)
-						.ParentStackNode(SharedThis(this))
-				]
+			+ SVerticalBox::Slot()
+			.AutoHeight()
+			[
+				// TabWell
+				SAssignNew(TabWell, SDockingTabWell)
+				.ParentStackNode(SharedThis(this))
 			]
 		]
 
-	+ SHorizontalBox::Slot()
+		+ SHorizontalBox::Slot()
 		.AutoWidth()
 		.Expose(InlineContentAreaRight)
 		.Padding(5.0f, 0.0f, 0.0f, 0.0f)
-		.VAlign((VAlign_Center));
+		.VAlign((VAlign_Center))
+	];
 
 	ChildSlot
 	[
