Fixed: Entering Foliage mode should warn before closing Matinee UE-10082

diff --git a/Engine/Source/Editor/LevelEditor/Private/SLevelEditor.cpp b/Engine/Source/Editor/LevelEditor/Private/SLevelEditor.cpp
--- a/Engine/Source/Editor/LevelEditor/Private/SLevelEditor.cpp
+++ b/Engine/Source/Editor/LevelEditor/Private/SLevelEditor.cpp
@@ -1182,6 +1182,24 @@ FName SLevelEditor::GetEditorModeTabId( FEditorModeID ModeID )
 
 void SLevelEditor::ToggleEditorMode( FEditorModeID ModeID )
 {
+	// Prompt the user if Matinee must be closed before activating new mode
+	if (ModeID != FBuiltinEditorModes::EM_InterpEdit)
+	{
+		FEdMode* MatineeMode = GLevelEditorModeTools().GetActiveMode(FBuiltinEditorModes::EM_InterpEdit);
+		if (MatineeMode && !MatineeMode->IsCompatibleWith(ModeID))
+		{
+			FEditorModeInfo MatineeModeInfo = FEditorModeRegistry::Get().GetModeInfo(ModeID);
+			FFormatNamedArguments Args;
+			Args.Add(TEXT("ModeName"), MatineeModeInfo.Name);
+			FText Msg = FText::Format(NSLOCTEXT("LevelEditor", "ModeSwitchCloseMatineeQ", "Activating '{ModeName}' editor mode will close UnrealMatinee.  Continue?"), Args);
+			
+			if (EAppReturnType::Yes != FMessageDialog::Open(EAppMsgType::YesNo, Msg))
+			{
+				return;
+			}
+		}
+	}
+		
 	// *Important* - activate the mode first since FEditorModeTools::DeactivateMode will
 	// activate the default mode when the stack becomes empty, resulting in multiple active visible modes.
 	GLevelEditorModeTools().ActivateMode( ModeID );
diff --git a/Engine/Source/Editor/UnrealEd/Private/EditorModeInterpolation.cpp b/Engine/Source/Editor/UnrealEd/Private/EditorModeInterpolation.cpp
--- a/Engine/Source/Editor/UnrealEd/Private/EditorModeInterpolation.cpp
+++ b/Engine/Source/Editor/UnrealEd/Private/EditorModeInterpolation.cpp
@@ -185,6 +185,7 @@ void FEdModeInterpEdit::ActorSelectionChangeNotify()
 bool FEdModeInterpEdit::IsCompatibleWith(FEditorModeID OtherModeID) const
 {
 	return
+		OtherModeID == FBuiltinEditorModes::EM_Placement	||
 		OtherModeID == FBuiltinEditorModes::EM_MeshPaint	||
 		OtherModeID == FBuiltinEditorModes::EM_Geometry		||
 		OtherModeID == FBuiltinEditorModes::EM_Bsp;
diff --git a/Engine/Source/Editor/UnrealEd/Public/EditorModeRegistry.cpp b/Engine/Source/Editor/UnrealEd/Public/EditorModeRegistry.cpp
--- a/Engine/Source/Editor/UnrealEd/Public/EditorModeRegistry.cpp
+++ b/Engine/Source/Editor/UnrealEd/Public/EditorModeRegistry.cpp
@@ -91,6 +91,18 @@ TArray<FEditorModeInfo> FEditorModeRegistry::GetSortedModeInfo() const
 	return ModeInfoArray;
 }
 
+FEditorModeInfo FEditorModeRegistry::GetModeInfo(FEditorModeID ModeID) const
+{
+	FEditorModeInfo Result;
+	const TSharedRef<IEditorModeFactory>* ModeFactory = ModeFactories.Find(ModeID);
+	if (ModeFactory)
+	{
+		Result = (*ModeFactory)->GetModeInfo();
+	}
+	
+	return Result;
+}
+
 TSharedPtr<FEdMode> FEditorModeRegistry::CreateMode(FEditorModeID ModeID, FEditorModeTools& Owner)
 {
 	const TSharedRef<IEditorModeFactory>* ModeFactory = ModeFactories.Find(ModeID);
diff --git a/Engine/Source/Editor/UnrealEd/Public/EditorModeRegistry.h b/Engine/Source/Editor/UnrealEd/Public/EditorModeRegistry.h
--- a/Engine/Source/Editor/UnrealEd/Public/EditorModeRegistry.h
+++ b/Engine/Source/Editor/UnrealEd/Public/EditorModeRegistry.h
@@ -135,6 +135,11 @@ public:
 	 */
 	UNREALED_API TArray<FEditorModeInfo> GetSortedModeInfo() const;
 
+	/**
+	 * Get a currently registered mode information for specified ID
+	 */
+	UNREALED_API FEditorModeInfo GetModeInfo(FEditorModeID ModeID) const;
+
 	/**
 	 * Registers an editor mode. Typically called from a module's StartupModule() routine.
 	 *
