Removed the SaveBrushAsCollision command and its bindings.

diff --git a/Engine/Source/Editor/DetailCustomizations/Private/StaticMeshActorDetails.cpp b/Engine/Source/Editor/DetailCustomizations/Private/StaticMeshActorDetails.cpp
--- a/Engine/Source/Editor/DetailCustomizations/Private/StaticMeshActorDetails.cpp
+++ b/Engine/Source/Editor/DetailCustomizations/Private/StaticMeshActorDetails.cpp
@@ -69,16 +69,4 @@ void FStaticMeshActorDetails::CustomizeDetails( IDetailLayoutBuilder& DetailBuil
 	];
 }
 
-FReply FStaticMeshActorDetails::OnSetCollisionFromBuilder()
-{
-	FLevelEditorModule& LevelEditor = FModuleManager::GetModuleChecked<FLevelEditorModule>( TEXT("LevelEditor") );
-
-	const FLevelEditorCommands& Commands = LevelEditor.GetLevelEditorCommands();
-	TSharedRef<const FUICommandList> CommandBindings = LevelEditor.GetGlobalLevelEditorActions();
-
-	CommandBindings->ExecuteAction( Commands.SaveBrushAsCollision.ToSharedRef() );
-
-	return FReply::Handled();
-}
-
 #undef LOCTEXT_NAMESPACE
diff --git a/Engine/Source/Editor/DetailCustomizations/Private/StaticMeshActorDetails.h b/Engine/Source/Editor/DetailCustomizations/Private/StaticMeshActorDetails.h
--- a/Engine/Source/Editor/DetailCustomizations/Private/StaticMeshActorDetails.h
+++ b/Engine/Source/Editor/DetailCustomizations/Private/StaticMeshActorDetails.h
@@ -10,7 +10,4 @@ public:
 
 	/** IDetailCustomization interface */
 	virtual void CustomizeDetails( IDetailLayoutBuilder& DetailBuilder ) override;
-private:
-	/** Called when the set collision from builder brush command should be executed */
-	FReply OnSetCollisionFromBuilder();
 };
\ No newline at end of file
diff --git a/Engine/Source/Editor/LevelEditor/Private/LevelEditor.cpp b/Engine/Source/Editor/LevelEditor/Private/LevelEditor.cpp
--- a/Engine/Source/Editor/LevelEditor/Private/LevelEditor.cpp
+++ b/Engine/Source/Editor/LevelEditor/Private/LevelEditor.cpp
@@ -1301,12 +1301,6 @@ void FLevelEditorModule::BindGlobalLevelEditorCommands()
 		FExecuteAction::CreateStatic( &FLevelEditorActionCallbacks::ExecuteExecCommand, FString( TEXT("ACTOR CREATE_BV_CONVEXVOLUME NORMALTOLERANCE=0.75 SnapToGrid=1") ) )
 		);
 
-	ActionList.MapAction(
-		Commands.SaveBrushAsCollision,
-		FExecuteAction::CreateStatic( &FLevelEditorActionCallbacks::OnSaveBrushAsCollision )
-		);
-
-
 	ActionList.MapAction(
 		Commands.KeepSimulationChanges,
 		FExecuteAction::CreateStatic( &FLevelEditorActionCallbacks::OnKeepSimulationChanges ),
diff --git a/Engine/Source/Editor/LevelEditor/Private/LevelEditorActions.cpp b/Engine/Source/Editor/LevelEditor/Private/LevelEditorActions.cpp
--- a/Engine/Source/Editor/LevelEditor/Private/LevelEditorActions.cpp
+++ b/Engine/Source/Editor/LevelEditor/Private/LevelEditorActions.cpp
@@ -2728,128 +2728,6 @@ void FLevelEditorActionCallbacks::SnapTo_Clicked( const bool InAlign, const bool
 	GEditor->RedrawLevelEditingViewports();
 }
 
-void FLevelEditorActionCallbacks::OnSaveBrushAsCollision()
-{
-	// First, find the currently selected actor with a static mesh.
-	// Fail if more than one actor with staticmesh is selected.
-	UStaticMesh* StaticMesh = NULL;
-	FMatrix MeshToWorld;
-
-	// Pointer to the world
-	UWorld* World = NULL;
-	for ( FSelectionIterator It( GEditor->GetSelectedActorIterator() ) ; It ; ++It )
-	{
-		AActor* Actor = Cast<AActor>( *It );
-		checkSlow( Actor->IsA(AActor::StaticClass()) );
-
-		UStaticMeshComponent* FoundStaticMeshComponent = NULL;
-		if( Actor->IsA(AStaticMeshActor::StaticClass()) )
-		{
-			FoundStaticMeshComponent = CastChecked<AStaticMeshActor>(Actor)->GetStaticMeshComponent();
-		}
-
-		UStaticMesh* FoundMesh = FoundStaticMeshComponent ? FoundStaticMeshComponent->StaticMesh : NULL;
-		if( FoundMesh )
-		{
-			// If we find multiple actors with static meshes, warn and do nothing.
-			if( StaticMesh )
-			{
-				FMessageDialog::Open( EAppMsgType::Ok, NSLOCTEXT("UnrealEd", "Error_SelectOneActor", "Please select just one Actor with a StaticMesh.") );
-				return;
-			}
-			StaticMesh = FoundMesh;
-			MeshToWorld = FoundStaticMeshComponent->ComponentToWorld.ToMatrixWithScale();
-			// Store the pointer to the world
-			World = Actor->GetWorld();
-		}
-	}
-
-	// If no static-mesh-toting actor found, warn and do nothing.
-	if(!StaticMesh)
-	{
-		FMessageDialog::Open( EAppMsgType::Ok, NSLOCTEXT("UnrealEd", "Error_NoActorWithStaticMesh", "No Actor found with a StaticMesh.") );
-		return;
-	}
-
-	// If we already have a collision model for this staticmesh, ask if we want to replace it.
-	if( StaticMesh->BodySetup )
-	{
-		const bool bDoReplace = EAppReturnType::Yes == FMessageDialog::Open( EAppMsgType::YesNo, NSLOCTEXT("UnrealEd", "Prompt_24", "Static Mesh already has a collision model. \nDo you want to replace it with Builder Brush?") );
-		if( !bDoReplace )
-		{
-			return;
-		}
-	}
-	
-	check(World)
-	ABrush* BuildBrush = World->GetDefaultBrush();
-	if(BuildBrush != nullptr)
-	{
-		// Now get the builder brush.
-		UModel* BuilderModel = BuildBrush->Brush;
-
-		// Need the transform between builder brush space and static mesh actor space.
-		const FMatrix BrushL2W = BuildBrush->ActorToWorld().ToMatrixWithScale();
-		const FMatrix MeshW2L = MeshToWorld.InverseFast();
-		const FMatrix SMToBB = BrushL2W * MeshW2L;
-		const FMatrix SMToBB_AT = SMToBB.TransposeAdjoint();
-
-		// Copy the current builder brush into a temp model.
-		// We keep no reference to this, so it will be GC'd at some point.
-		UModel* TempModel = NewObject<UModel>();
-		TempModel->Initialize(nullptr, 1);
-		TempModel->Polys->Element = BuilderModel->Polys->Element;
-
-		// Now transform each poly into local space for the selected static mesh.
-		for (int32 i = 0; i < TempModel->Polys->Element.Num(); i++)
-		{
-			FPoly* Poly = &TempModel->Polys->Element[i];
-
-			for (int32 j = 0; j < Poly->Vertices.Num(); j++)
-			{
-				Poly->Vertices[j] = SMToBB.TransformPosition(Poly->Vertices[j]);
-			}
-
-			Poly->Normal = SMToBB_AT.TransformVector(Poly->Normal);
-			Poly->Normal.Normalize(); // SmToBB might have scaling in it.
-		}
-
-		// Build bounding box.
-		TempModel->BuildBound();
-
-		// Build BSP for the brush.
-		FBSPOps::bspBuild(TempModel, FBSPOps::BSP_Good, 15, 70, 1, 0);
-		FBSPOps::bspRefresh(TempModel, 1);
-		FBSPOps::bspBuildBounds(TempModel);
-
-
-		// Now - use this as the Rigid Body collision for this static mesh as well.
-
-		// Make sure rendering is done - so we are not changing data being used by collision drawing.
-		FlushRenderingCommands();
-
-		// If we already have a BodySetup - clear it.
-		if (StaticMesh->BodySetup)
-		{
-			StaticMesh->BodySetup->RemoveSimpleCollision();
-		}
-		// If we don't already have physics props, construct them here.
-		else
-		{
-			StaticMesh->CreateBodySetup();
-		}
-
-		// Convert collision model into a collection of convex hulls.
-		// NB: This removes any convex hulls that were already part of the collision data.
-		StaticMesh->BodySetup->CreateFromModel(TempModel, true);
-	}
-	// refresh collision change back to staticmesh components
-	RefreshCollisionChange(StaticMesh);
-
-	// Finally mark the parent package as 'dirty', so user will be prompted if they want to save it etc.
-	StaticMesh->MarkPackageDirty();
-}
-
 bool FLevelEditorActionCallbacks::ActorSelected_CanExecute()
 {
 	// Had to have something selected
@@ -3070,8 +2948,6 @@ void FLevelEditorCommands::RegisterCommands()
 	UI_COMMAND( CreateLightConvexVolume, "Light Convex Blocking Volume From Mesh", "Creates a light convex blocking volume from the static mesh", EUserInterfaceActionType::Button, FInputChord() );
 	UI_COMMAND( CreateRoughConvexVolume, "Rought Convex Blocking Volume From Mesh", "Creates a rough convex blocking volume from the static mesh", EUserInterfaceActionType::Button, FInputChord() );
 
-	UI_COMMAND( SaveBrushAsCollision, "Save Collision from Builder Brush", "Creates a collision primitive on the selected static meshes based on the shape of the builder brush", EUserInterfaceActionType::Button, FInputChord() );
-
 	UI_COMMAND( KeepSimulationChanges, "Keep Simulation Changes", "Saves the changes made to this actor in Simulate mode to the actor's default state.", EUserInterfaceActionType::Button, FInputChord( EKeys::K ) );
 
 	UI_COMMAND( MakeActorLevelCurrent, "Make Selected Actor's Level Current", "Makes the selected actors level the current level", EUserInterfaceActionType::Button, FInputChord( EKeys::M ) );
diff --git a/Engine/Source/Editor/LevelEditor/Public/LevelEditorActions.h b/Engine/Source/Editor/LevelEditor/Public/LevelEditorActions.h
--- a/Engine/Source/Editor/LevelEditor/Public/LevelEditorActions.h
+++ b/Engine/Source/Editor/LevelEditor/Public/LevelEditorActions.h
@@ -461,9 +461,6 @@ public:
 	/** Create a blocking volume from the meshes using a rough convex shape */
 	TSharedPtr< FUICommandInfo > CreateRoughConvexVolume;
 
-	/** Set the collision model on the static meshes to be the same shape as the builder brush */
-	TSharedPtr< FUICommandInfo > SaveBrushAsCollision;
-
 	/** Set the actors collision to block all */
 	TSharedPtr< FUICommandInfo > SetCollisionBlockAll;
 
@@ -1043,11 +1040,6 @@ public:
 	 */
 	static void OnFindActorInLevelScript();
 
-	/**
-	* Take the currently selected static mesh, and save the builder brush as its low poly collision model.
-	*/
-	static void OnSaveBrushAsCollision();
-
 	/** Select the world info actor and show the properties */
 	static void OnShowWorldProperties( TWeakPtr< SLevelEditor > LevelEditor );
 
