Fixed camera roll being applied to level viewport cameras after ejecting from the player in a play in viewport

diff --git a/Engine/Source/Editor/LevelEditor/Private/SLevelViewport.cpp b/Engine/Source/Editor/LevelEditor/Private/SLevelViewport.cpp
--- a/Engine/Source/Editor/LevelEditor/Private/SLevelViewport.cpp
+++ b/Engine/Source/Editor/LevelEditor/Private/SLevelViewport.cpp
@@ -3465,7 +3465,7 @@ void SLevelViewport::EndPlayInEditorSession()
 		ActiveViewport->OnPlayWorldViewportSwapped(*InactiveViewport);
 
 		// Play in editor viewport was active, swap back to our level editor viewport
-		ActiveViewport->SetViewportClient( NULL );
+		ActiveViewport->SetViewportClient( nullptr );
 
 		// We should be the only thing holding on to viewports
 		check( ActiveViewport.IsUnique() );
@@ -3481,9 +3481,12 @@ void SLevelViewport::EndPlayInEditorSession()
 	}
 	else
 	{
-		InactiveViewport->SetViewportClient( NULL );
+		InactiveViewport->SetViewportClient( nullptr );
 	}
 
+	// Remove camera roll from any PIE camera applied in this viewport. A rolled camera is hard to use for editing
+	LevelViewportClient->RemoveCameraRoll();
+
 	// Reset the inactive viewport
 	InactiveViewport.Reset();
 
