Foliage types grouped by sub-level in editor viewport  menu

diff --git a/Engine/Source/Editor/FoliageEdit/Private/FoliageEdMode.cpp b/Engine/Source/Editor/FoliageEdit/Private/FoliageEdMode.cpp
--- a/Engine/Source/Editor/FoliageEdit/Private/FoliageEdMode.cpp
+++ b/Engine/Source/Editor/FoliageEdit/Private/FoliageEdMode.cpp
@@ -55,27 +55,8 @@ FFoliageMeshUIInfo::FFoliageMeshUIInfo(UFoliageType* InSettings)
 FText FFoliageMeshUIInfo::GetNameText() const
 {
 	//@todo: this is redundant with FFoliagePaletteItem::DisplayFName, should probably move sorting implementation over to SFoliagePalette
-	const UFoliageType* FoliageType = Settings;
-	if (FoliageType->IsAsset())
-	{
-		return FText::FromName(FoliageType->GetFName());
-	}
-	else if (UBlueprint* Blueprint = Cast<UBlueprint>(FoliageType->GetClass()->ClassGeneratedBy))
-	{
-		return FText::FromName(Blueprint->GetFName());
-	}
-	else
-	{
-		UStaticMesh* Mesh = FoliageType->GetStaticMesh();
-		if (Mesh)
-		{
-			return FText::FromName(Mesh->GetFName());
-		}
-		else
-		{
-			return LOCTEXT("EmptyMeshName", "[Empty Mesh]");
-		}
-	}
+	FName DisplayFName = Settings->GetDisplayFName();
+	return FText::FromName(DisplayFName);
 }
 
 //
diff --git a/Engine/Source/Editor/FoliageEdit/Private/FoliagePaletteItem.cpp b/Engine/Source/Editor/FoliageEdit/Private/FoliagePaletteItem.cpp
--- a/Engine/Source/Editor/FoliageEdit/Private/FoliagePaletteItem.cpp
+++ b/Engine/Source/Editor/FoliageEdit/Private/FoliagePaletteItem.cpp
@@ -32,18 +32,7 @@ FFoliagePaletteItemModel::FFoliagePaletteItemModel(FFoliageMeshUIInfoPtr InTypeI
 	, FoliageEditMode(InFoliageEditMode)
 {
 	// Determine the display FName
-	if (TypeInfo->Settings->IsAsset())
-	{
-		DisplayFName = TypeInfo->Settings->GetFName();
-	}
-	else if (UBlueprint* FoliageTypeBP = Cast<UBlueprint>(TypeInfo->Settings->GetClass()->ClassGeneratedBy))
-	{
-		DisplayFName = FoliageTypeBP->GetFName();
-	}
-	else
-	{
-		DisplayFName = TypeInfo->Settings->GetStaticMesh()->GetFName();
-	}
+	DisplayFName = TypeInfo->Settings->GetDisplayFName();
 	
 	FAssetData AssetData;
 	if (!IsBlueprint())
diff --git a/Engine/Source/Editor/LevelEditor/Private/SLevelViewportToolBar.cpp b/Engine/Source/Editor/LevelEditor/Private/SLevelViewportToolBar.cpp
--- a/Engine/Source/Editor/LevelEditor/Private/SLevelViewportToolBar.cpp
+++ b/Engine/Source/Editor/LevelEditor/Private/SLevelViewportToolBar.cpp
@@ -1070,8 +1070,7 @@ void SLevelViewportToolBar::FillShowLayersMenu( FMenuBuilder& MenuBuilder, TWeak
 		{
 			// Get all the layers and create an entry for each of them
 			TArray<FName> AllLayerNames;
-
-				GEditor->Layers->AddAllLayerNamesTo(AllLayerNames);
+			GEditor->Layers->AddAllLayerNamesTo(AllLayerNames);
 
 			for( int32 LayerIndex = 0; LayerIndex < AllLayerNames.Num(); ++LayerIndex )
 			{
@@ -1089,6 +1088,27 @@ void SLevelViewportToolBar::FillShowLayersMenu( FMenuBuilder& MenuBuilder, TWeak
 	}
 }
 
+static TMap<FName, TArray<UFoliageType*>> GroupFoliageByOuter(const TArray<UFoliageType*> FoliageList)
+{
+	TMap<FName, TArray<UFoliageType*>> Result;
+
+	for (UFoliageType* FoliageType : FoliageList)
+	{
+		if (FoliageType->IsAsset())
+		{
+			Result.FindOrAdd(NAME_None).Add(FoliageType);
+		}
+		else
+		{
+			FName LevelName = FoliageType->GetOutermost()->GetFName();
+			Result.FindOrAdd(LevelName).Add(FoliageType);
+		}
+	}
+
+	Result.KeySort([](const FName& A, const FName& B) { return (A < B && B != NAME_None); });
+	return Result;
+}
+
 void SLevelViewportToolBar::FillShowFoliageTypesMenu(class FMenuBuilder& MenuBuilder, TWeakPtr<class SLevelViewport> Viewport)
 {
 	auto ViewportPtr = Viewport.Pin();
@@ -1107,24 +1127,28 @@ void SLevelViewportToolBar::FillShowFoliageTypesMenu(class FMenuBuilder& MenuBui
 		MenuBuilder.AddMenuEntry(LOCTEXT("HideAllLabel", "Hide All"), FText::GetEmpty(), FSlateIcon(), HideAllFoliage);
 	}
 	MenuBuilder.EndSection();
-		
-	TArray<UFoliageType*> AllFoliageTypes = GEditor->GetFoliageTypesInWorld(ViewportPtr->GetWorld());
+	
+	// Gather all foliage types used in this world and group them by sub-levels
+	auto AllFoliageMap = GroupFoliageByOuter(GEditor->GetFoliageTypesInWorld(ViewportPtr->GetWorld()));
 
-	if (AllFoliageTypes.Num())
+	for (auto& FoliagePair : AllFoliageMap)
 	{
-		MenuBuilder.BeginSection("LevelViewportFoliageTypes");
-			
-		for (UFoliageType* FoliageType : AllFoliageTypes)
+		// Name foliage group by an outer sub-level name, or empty if foliage type is an asset
+		FText EntryName = (FoliagePair.Key == NAME_None ? FText::GetEmpty() : FText::FromName(FPackageName::GetShortFName(FoliagePair.Key)));
+		MenuBuilder.BeginSection(NAME_None, EntryName);
+
+		TArray<UFoliageType*>& FoliageList = FoliagePair.Value;
+		for (UFoliageType* FoliageType : FoliageList)
 		{
-			FString MeshName = FoliageType->GetName();
+			FName MeshName = FoliageType->GetDisplayFName();
 			TWeakObjectPtr<UFoliageType> FoliageTypePtr = FoliageType;
-				
+
 			FUIAction Action(
-					FExecuteAction::CreateSP(ViewportPtr.ToSharedRef(), &SLevelViewport::ToggleShowFoliageType, FoliageTypePtr),
-					FCanExecuteAction(),
-					FIsActionChecked::CreateSP(ViewportPtr.ToSharedRef(), &SLevelViewport::IsFoliageTypeVisible, FoliageTypePtr));
+				FExecuteAction::CreateSP(ViewportPtr.ToSharedRef(), &SLevelViewport::ToggleShowFoliageType, FoliageTypePtr),
+				FCanExecuteAction(),
+				FIsActionChecked::CreateSP(ViewportPtr.ToSharedRef(), &SLevelViewport::IsFoliageTypeVisible, FoliageTypePtr));
 
-			MenuBuilder.AddMenuEntry(FText::FromString(MeshName), FText::GetEmpty(), FSlateIcon(), Action, NAME_None, EUserInterfaceActionType::ToggleButton);
+			MenuBuilder.AddMenuEntry(FText::FromName(MeshName), FText::GetEmpty(), FSlateIcon(), Action, NAME_None, EUserInterfaceActionType::ToggleButton);
 		}
 
 		MenuBuilder.EndSection();
diff --git a/Engine/Source/Runtime/Foliage/Private/InstancedFoliage.cpp b/Engine/Source/Runtime/Foliage/Private/InstancedFoliage.cpp
--- a/Engine/Source/Runtime/Foliage/Private/InstancedFoliage.cpp
+++ b/Engine/Source/Runtime/Foliage/Private/InstancedFoliage.cpp
@@ -554,6 +554,26 @@ void UFoliageType::OnHiddenEditorViewMaskChanged(UWorld* InWorld)
 	}
 }
 
+FName UFoliageType::GetDisplayFName() const
+{
+	FName DisplayFName;
+
+	if (IsAsset())
+	{
+		DisplayFName = GetFName();
+	}
+	else if (UBlueprint* FoliageTypeBP = Cast<UBlueprint>(GetClass()->ClassGeneratedBy))
+	{
+		DisplayFName = FoliageTypeBP->GetFName();
+	}
+	else if (UStaticMesh* StaticMesh = GetStaticMesh())
+	{
+		DisplayFName = StaticMesh->GetFName();
+	}
+
+	return DisplayFName;
+}
+
 #endif
 
 //
diff --git a/Engine/Source/Runtime/Foliage/Public/FoliageType.h b/Engine/Source/Runtime/Foliage/Public/FoliageType.h
--- a/Engine/Source/Runtime/Foliage/Public/FoliageType.h
+++ b/Engine/Source/Runtime/Foliage/Public/FoliageType.h
@@ -61,6 +61,9 @@ class UFoliageType : public UObject
 
 	/* Notifies all relevant foliage actors that HiddenEditorView mask has been changed */
 	FOLIAGE_API void OnHiddenEditorViewMaskChanged(UWorld* InWorld);
+
+	/* Nice and clean name for user interface */
+	FOLIAGE_API FName GetDisplayFName() const;
 #endif
 
 	/* A GUID that is updated every time the foliage type is modified, 
