ActorAnimationEditor: Fixed tab icon for level editor tab

diff --git a/Engine/Source/Editor/LevelEditor/Private/SLevelEditor.cpp b/Engine/Source/Editor/LevelEditor/Private/SLevelEditor.cpp
--- a/Engine/Source/Editor/LevelEditor/Private/SLevelEditor.cpp
+++ b/Engine/Source/Editor/LevelEditor/Private/SLevelEditor.cpp
@@ -737,7 +737,7 @@ TSharedRef<SDockTab> SLevelEditor::SpawnLevelEditorTab( const FSpawnTabArgs& Arg
 	{
 		// @todo remove when world-centric mode is added
 		SequencerTab = SNew(SDockTab)
-			.Icon( FEditorStyle::GetBrush("Sequencer.Tabs.SequencerMain") )
+			.Icon( FSlateStyleRegistry::FindSlateStyle("ActorAnimationEditorStyle")->GetBrush("ActorAnimationEditor.Tabs.Sequencer") )
 			.Label( NSLOCTEXT("Sequencer", "SequencerMainTitle", "Sequencer") )
 			[
 				SNullWidget::NullWidget
