Fixed URL for the LevelViewport controls help popup to work on Mac
[UE-18355] - Clicking on show viewport controls on mac comes up with missing page error on mac

#fix Prefixed URLs correctly with file:// in SLevelViewportControlsPopup

diff --git a/Engine/Source/Editor/LevelEditor/Private/SLevelViewportControlsPopup.cpp b/Engine/Source/Editor/LevelEditor/Private/SLevelViewportControlsPopup.cpp
--- a/Engine/Source/Editor/LevelEditor/Private/SLevelViewportControlsPopup.cpp
+++ b/Engine/Source/Editor/LevelEditor/Private/SLevelViewportControlsPopup.cpp
@@ -9,11 +9,11 @@ void SLevelViewportControlsPopup::Construct(const FArguments& InArgs)
 {
 	PopupSize.Set(252, 558);
 #if PLATFORM_WINDOWS
-	PopupPath = FPaths::ConvertRelativePathToFull(FPaths::Combine(*FPaths::EngineDir(), TEXT("Documentation/Source/Shared/Editor/Overlays/ViewportControlsWindows.html")));
+	PopupPath = FString(TEXT("file:///")) + FPaths::ConvertRelativePathToFull(FPaths::Combine(*FPaths::EngineDir(), TEXT("Documentation/Source/Shared/Editor/Overlays/ViewportControlsWindows.html")));
 #elif PLATFORM_LINUX
-	PopupPath = FPaths::ConvertRelativePathToFull(FPaths::Combine(*FPaths::EngineDir(), TEXT("Documentation/Source/Shared/Editor/Overlays/ViewportControlsWindows.html")));
+	PopupPath = FString(TEXT("file://")) + FPaths::ConvertRelativePathToFull(FPaths::Combine(*FPaths::EngineDir(), TEXT("Documentation/Source/Shared/Editor/Overlays/ViewportControlsWindows.html")));
 #elif PLATFORM_MAC
-	PopupPath = FPaths::ConvertRelativePathToFull(FPaths::Combine(*FPaths::EngineDir(), TEXT("Documentation/Source/Shared/Editor/Overlays/ViewportControlsMac.html")));
+	PopupPath = FString(TEXT("file://")) + FPaths::ConvertRelativePathToFull(FPaths::Combine(*FPaths::EngineDir(), TEXT("Documentation/Source/Shared/Editor/Overlays/ViewportControlsMac.html")));
 #endif
 
 	if (!PopupPath.IsEmpty())
