Integrated the application portal service into the level editor open marketplace button.

Merging using UE4-To-UE4-LauncherDev

diff --git a/Engine/Source/Editor/LevelEditor/LevelEditor.Build.cs b/Engine/Source/Editor/LevelEditor/LevelEditor.Build.cs
--- a/Engine/Source/Editor/LevelEditor/LevelEditor.Build.cs
+++ b/Engine/Source/Editor/LevelEditor/LevelEditor.Build.cs
@@ -16,7 +16,8 @@ public class LevelEditor : ModuleRules
 				"SizeMap",
 				"SlateReflector",
                 "IntroTutorials",
-                "AppFramework"
+                "AppFramework",
+                "PortalServices",
 			}
 		);
 
diff --git a/Engine/Source/Editor/LevelEditor/Private/LevelEditorActions.cpp b/Engine/Source/Editor/LevelEditor/Private/LevelEditorActions.cpp
--- a/Engine/Source/Editor/LevelEditor/Private/LevelEditorActions.cpp
+++ b/Engine/Source/Editor/LevelEditor/Private/LevelEditorActions.cpp
@@ -65,6 +65,9 @@
 #include "Editor/KismetWidgets/Public/CreateBlueprintFromActorDialog.h"
 #include "EditorProjectSettings.h"
 #include "HierarchicalLODUtils.h"
+#include "AsyncResult.h"
+#include "IPortalApplicationWindow.h"
+#include "IPortalServiceLocator.h"
 
 DEFINE_LOG_CATEGORY_STATIC(LevelEditorActions, Log, All);
 
@@ -1608,7 +1611,7 @@ void FLevelEditorActionCallbacks::OnSelectOwningHLODCluster()
 	}
 }
 
-void FLevelEditorActionCallbacks::OnSelectMatineeActor(AMatineeActor * ActorToSelect)
+void FLevelEditorActionCallbacks::OnSelectMatineeActor( AMatineeActor * ActorToSelect )
 {
 	GEditor->SelectNone( false, true );
 	GEditor->SelectActor(ActorToSelect, true, false, true);
@@ -1971,6 +1974,17 @@ void FLevelEditorActionCallbacks::OpenContentBrowser()
 
 void FLevelEditorActionCallbacks::OpenMarketplace()
 {
+	auto Service = GEditor->GetServiceLocator()->GetServiceRef<IPortalApplicationWindow>();
+	if (Service->IsAvailable())
+	{
+		TAsyncResult<bool> Result = Service->NavigateTo(TEXT("/ue/marketplace"));
+		if (FEngineAnalytics::IsAvailable())
+		{
+			FEngineAnalytics::GetProvider().RecordEvent(TEXT("Editor.Usage.OpenMarketplace"));
+		}
+	}
+	else
+	{
 	IDesktopPlatform* DesktopPlatform = FDesktopPlatformModule::Get();
 
 	if (DesktopPlatform != nullptr)
@@ -2007,6 +2021,7 @@ void FLevelEditorActionCallbacks::OpenMarketplace()
 		}
 	}
 }
+}
 
 bool FLevelEditorActionCallbacks::CanSelectGameModeBlueprint()
 {
