Sequencer: Removed -sequencer command line switch

diff --git a/Engine/Source/Editor/LevelEditor/Private/SLevelEditor.cpp b/Engine/Source/Editor/LevelEditor/Private/SLevelEditor.cpp
--- a/Engine/Source/Editor/LevelEditor/Private/SLevelEditor.cpp
+++ b/Engine/Source/Editor/LevelEditor/Private/SLevelEditor.cpp
@@ -1081,19 +1081,18 @@ TSharedRef<SWidget> SLevelEditor::RestoreContentArea( const TSharedRef<SDockTab>
 				.SetGroup( WorkspaceMenu::GetMenuStructure().GetLevelEditorCategory() )
 				.SetIcon( FSlateIcon(FEditorStyle::GetStyleSetName(), "LevelEditor.Tabs.WorldBrowserComposition") );
 		}
-		
+
 		{
 			const FSlateIcon StatsViewerIcon(FEditorStyle::GetStyleSetName(), "LevelEditor.Tabs.StatsViewer");
-			LevelEditorTabManager->RegisterTabSpawner( LevelEditorStatsViewerTab, FOnSpawnTab::CreateSP<SLevelEditor, FName, FString>(this, &SLevelEditor::SpawnLevelEditorTab, LevelEditorStatsViewerTab, FString()) )
+			LevelEditorTabManager->RegisterTabSpawner(LevelEditorStatsViewerTab, FOnSpawnTab::CreateSP<SLevelEditor, FName, FString>(this, &SLevelEditor::SpawnLevelEditorTab, LevelEditorStatsViewerTab, FString()))
 				.SetDisplayName(NSLOCTEXT("LevelEditorTabs", "LevelEditorStatsViewer", "Statistics"))
 				.SetTooltipText(NSLOCTEXT("LevelEditorTabs", "LevelEditorStatsViewerTooltipText", "Open the Statistics tab, in order to see data pertaining to lighting, textures and primitives."))
-				.SetGroup( MenuStructure.GetLevelEditorCategory() )
-				.SetIcon( StatsViewerIcon );
+				.SetGroup(MenuStructure.GetLevelEditorCategory())
+				.SetIcon(StatsViewerIcon);
 		}
 
-		// @todo remove when world-centric mode is added
-		if (FParse::Param(FCommandLine::Get(), TEXT("sequencer")))
 		{
+			// @todo remove when world-centric mode is added
 			LevelEditorTabManager->RegisterTabSpawner( "Sequencer", FOnSpawnTab::CreateSP<SLevelEditor, FName, FString>(this, &SLevelEditor::SpawnLevelEditorTab, FName("Sequencer"), FString()) )
 				.SetDisplayName(NSLOCTEXT("LevelEditorTabs", "Sequencer", "Sequencer"))
 				.SetGroup( MenuStructure.GetLevelEditorCategory() );
