Sequencer: Removed sequencer commandline switch.

diff --git a/Engine/Source/Editor/LevelEditor/Private/LevelEditor.cpp b/Engine/Source/Editor/LevelEditor/Private/LevelEditor.cpp
--- a/Engine/Source/Editor/LevelEditor/Private/LevelEditor.cpp
+++ b/Engine/Source/Editor/LevelEditor/Private/LevelEditor.cpp
@@ -331,12 +331,9 @@ void FLevelEditorModule::SummonWorldBrowserComposition()
 // @todo remove when world-centric mode is added
 void FLevelEditorModule::AttachSequencer( TSharedPtr<SWidget> SequencerWidget, TSharedPtr<IAssetEditorInstance> SequencerAssetEditor )
 {
-	if( FParse::Param( FCommandLine::Get(), TEXT( "Sequencer" ) ) )
-	{
-		TSharedPtr<SLevelEditor> LevelEditorInstance = LevelEditorInstancePtr.Pin();
+	TSharedPtr<SLevelEditor> LevelEditorInstance = LevelEditorInstancePtr.Pin();
 
-		LevelEditorInstance->AttachSequencer( SequencerWidget, SequencerAssetEditor );
-	}
+	LevelEditorInstance->AttachSequencer( SequencerWidget, SequencerAssetEditor );
 }
 
 TSharedPtr<ILevelViewport> FLevelEditorModule::GetFirstActiveViewport()
