Sequencer: Add sequencer tab icon.

diff --git a/Engine/Source/Editor/LevelEditor/Private/SLevelEditor.cpp b/Engine/Source/Editor/LevelEditor/Private/SLevelEditor.cpp
--- a/Engine/Source/Editor/LevelEditor/Private/SLevelEditor.cpp
+++ b/Engine/Source/Editor/LevelEditor/Private/SLevelEditor.cpp
@@ -1093,9 +1093,11 @@ TSharedRef<SWidget> SLevelEditor::RestoreContentArea( const TSharedRef<SDockTab>
 
 		{
 			// @todo remove when world-centric mode is added
+			const FSlateIcon SequencerIcon("ActorAnimationEditorStyle", "ActorAnimationEditor.Tabs.Sequencer" );
 			LevelEditorTabManager->RegisterTabSpawner( "Sequencer", FOnSpawnTab::CreateSP<SLevelEditor, FName, FString>(this, &SLevelEditor::SpawnLevelEditorTab, FName("Sequencer"), FString()) )
 				.SetDisplayName(NSLOCTEXT("LevelEditorTabs", "Sequencer", "Sequencer"))
-				.SetGroup( MenuStructure.GetLevelEditorCategory() );
+				.SetGroup( MenuStructure.GetLevelEditorCategory() )
+				.SetIcon( SequencerIcon );
 		}
 
 		{
