Sequencer: Fix crash when the actor animation plugin isn't enabled and there's a sequencer tab in the user's window configuration.

diff --git a/Engine/Source/Editor/LevelEditor/Private/SLevelEditor.cpp b/Engine/Source/Editor/LevelEditor/Private/SLevelEditor.cpp
--- a/Engine/Source/Editor/LevelEditor/Private/SLevelEditor.cpp
+++ b/Engine/Source/Editor/LevelEditor/Private/SLevelEditor.cpp
@@ -752,14 +752,17 @@ TSharedRef<SDockTab> SLevelEditor::SpawnLevelEditorTab( const FSpawnTabArgs& Arg
 	}
 	else if( TabIdentifier == TEXT("Sequencer") )
 	{
-		// @todo sequencer: remove when world-centric mode is added
-		SequencerTab = SNew(SDockTab)
-			.Icon( FSlateStyleRegistry::FindSlateStyle("ActorAnimationEditorStyle")->GetBrush("ActorAnimationEditor.Tabs.Sequencer") )
-			.Label( NSLOCTEXT("Sequencer", "SequencerMainTitle", "Sequencer") )
-			[
-				SNullWidget::NullWidget
-			];
-		return SequencerTab.ToSharedRef();
+		if (FSlateStyleRegistry::FindSlateStyle("ActorAnimationEditorStyle"))
+		{
+			// @todo sequencer: remove when world-centric mode is added
+			SequencerTab = SNew(SDockTab)
+				.Icon( FSlateStyleRegistry::FindSlateStyle("ActorAnimationEditorStyle")->GetBrush("ActorAnimationEditor.Tabs.Sequencer") )
+				.Label( NSLOCTEXT("Sequencer", "SequencerMainTitle", "Sequencer") )
+				[
+					SNullWidget::NullWidget
+				];
+			return SequencerTab.ToSharedRef();
+		}
 	}
 	else if( TabIdentifier == LevelEditorStatsViewerTab )
 	{
