Camera preview window in the level viewport now has better support for actors with multiple camera components. It will ignore components with bIsActive=false, letting the user decide which component to look through.

(merge from branch)

diff --git a/Engine/Source/Editor/LevelEditor/Private/SLevelViewport.cpp b/Engine/Source/Editor/LevelEditor/Private/SLevelViewport.cpp
--- a/Engine/Source/Editor/LevelEditor/Private/SLevelViewport.cpp
+++ b/Engine/Source/Editor/LevelEditor/Private/SLevelViewport.cpp
@@ -3743,12 +3743,19 @@ bool SLevelViewport::GetCameraInformationFromActor(AActor* Actor, FMinimalViewIn
 {
 	// @todo camerapip: Could support actors other than cameras too!  (Character views?)
 	//@TODO: CAMERA: Support richer camera interactions in SIE; this may shake out naturally if everything uses camera components though
-	if (UCameraComponent* CameraComponent = Actor->FindComponentByClass<UCameraComponent>())
+	TArray<UCameraComponent*> CamComps;
+	Actor->GetComponents<UCameraComponent>(CamComps);
+	for (UCameraComponent* CamComp : CamComps)
 	{
-		CameraComponent->GetCameraView(0.0f, out_CameraInfo);
-		return true;
+		if (CamComp->bIsActive)
+		{
+			// first active camera, use it and be done
+			CamComp->GetCameraView(0.0f, out_CameraInfo);
+			return true;
+		}
 	}
 
+	// no active cameras
 	return false;
 }
 
diff --git a/Engine/Source/Editor/UnrealEd/Private/LevelEditorViewport.cpp b/Engine/Source/Editor/UnrealEd/Private/LevelEditorViewport.cpp
--- a/Engine/Source/Editor/UnrealEd/Private/LevelEditorViewport.cpp
+++ b/Engine/Source/Editor/UnrealEd/Private/LevelEditorViewport.cpp
@@ -2093,7 +2093,19 @@ void FLevelEditorViewportClient::UpdateViewForLockedActor()
 			if (bLockedCameraView)
 			{
 				// If this is a camera actor, then inherit some other settings
-				UCameraComponent* CameraComponent = Actor->FindComponentByClass<UCameraComponent>();
+				TArray<UCameraComponent*> CamComps;
+				Actor->GetComponents<UCameraComponent>(CamComps);
+
+				UCameraComponent* CameraComponent = nullptr;
+				for (UCameraComponent* Comp : CamComps)
+				{
+					if (Comp->bIsActive)
+					{
+						CameraComponent = Comp;
+						break;
+					}
+				}
+
 				if (CameraComponent != NULL)
 				{
 					bUseControllingActorViewInfo = true;
diff --git a/Engine/Source/Editor/UnrealEd/Public/LevelEditorViewport.h b/Engine/Source/Editor/UnrealEd/Public/LevelEditorViewport.h
--- a/Engine/Source/Editor/UnrealEd/Public/LevelEditorViewport.h
+++ b/Engine/Source/Editor/UnrealEd/Public/LevelEditorViewport.h
@@ -436,7 +436,23 @@ public:
 			LockedActor = ActorLockedToCamera.Get();
 		}
 
-		return LockedActor ? LockedActor->FindComponentByClass<UCameraComponent>() : nullptr;
+		UCameraComponent* CameraComponent = nullptr;
+		if (LockedActor)
+		{
+			TArray<UCameraComponent*> CamComps;
+			LockedActor->GetComponents<UCameraComponent>(CamComps);
+
+			for (UCameraComponent* Comp : CamComps)
+			{
+				if (Comp->bIsActive)
+				{
+					CameraComponent = Comp;
+					break;
+				}
+			}
+		}
+
+		return CameraComponent;
 	}
 
 	/** 
