Added ability to create new actor animation inside the level editor

diff --git a/Engine/Plugins/MovieScene/ActorAnimationEditor/Source/ActorAnimationEditor/ActorAnimationEditor.Build.cs b/Engine/Plugins/MovieScene/ActorAnimationEditor/Source/ActorAnimationEditor/ActorAnimationEditor.Build.cs
--- a/Engine/Plugins/MovieScene/ActorAnimationEditor/Source/ActorAnimationEditor/ActorAnimationEditor.Build.cs
+++ b/Engine/Plugins/MovieScene/ActorAnimationEditor/Source/ActorAnimationEditor/ActorAnimationEditor.Build.cs
@@ -20,6 +20,7 @@ public class ActorAnimationEditor : ModuleRules
 				"CoreUObject",
                 "EditorStyle",
                 "Engine",
+				"LevelEditor",
 				"MovieScene",
 				"MovieSceneTracks",
 				"Sequencer",
diff --git a/Engine/Plugins/MovieScene/ActorAnimationEditor/Source/ActorAnimationEditor/Private/ActorAnimationEditorModule.cpp b/Engine/Plugins/MovieScene/ActorAnimationEditor/Source/ActorAnimationEditor/Private/ActorAnimationEditorModule.cpp
--- a/Engine/Plugins/MovieScene/ActorAnimationEditor/Source/ActorAnimationEditor/Private/ActorAnimationEditorModule.cpp
+++ b/Engine/Plugins/MovieScene/ActorAnimationEditor/Source/ActorAnimationEditor/Private/ActorAnimationEditorModule.cpp
@@ -3,7 +3,26 @@
 #include "ActorAnimationEditorPrivatePCH.h"
 #include "ActorAnimationEditorStyle.h"
 #include "ModuleInterface.h"
+#include "LevelEditor.h"
+#include "MovieSceneActor.h"
 
+#define LOCTEXT_NAMESPACE "ActorAnimationEditor"
+
+class FActorAnimationExtensionCommands : public TCommands<FActorAnimationExtensionCommands>
+{
+public:
+	FActorAnimationExtensionCommands()
+		: TCommands("ActorAnimationEditor", LOCTEXT("ExtensionDescription", "Extension commands specific to the actor animation editor"), NAME_None, "ActorAnimationEditorStyle")
+	{}
+
+	TSharedPtr<FUICommandInfo> CreateNewActorAnimationInLevel;
+
+	/** Initialize commands */
+	virtual void RegisterCommands() override
+	{
+		UI_COMMAND(CreateNewActorAnimationInLevel, "Add New Actor Animation", "Create a new actor animation asset, and place an instance of it in this level", EUserInterfaceActionType::Button, FInputChord());
+	}
+};
 
 /**
  * Implements the ActorAnimationEditor module.
@@ -20,11 +39,13 @@ public:
 		Style = MakeShareable(new FActorAnimationEditorStyle());
 
 		RegisterAssetTools();
+		RegisterMenuExtensions();
 	}
 	
 	virtual void ShutdownModule() override
 	{
 		UnregisterAssetTools();
+		UnregisterMenuExtensions();
 	}
 
 protected:
@@ -65,6 +86,96 @@ protected:
 		}
 	}
 
+protected:
+
+	/**
+	 * Registers menu extensions for the level editor toolbar
+	 */
+	void RegisterMenuExtensions()
+	{
+		FActorAnimationExtensionCommands::Register();
+
+		CommandList = MakeShareable(new FUICommandList);
+
+		CommandList->MapAction(FActorAnimationExtensionCommands::Get().CreateNewActorAnimationInLevel,
+			FExecuteAction::CreateStatic(&FActorAnimationEditorModule::OnCreateActorInLevel)
+			);
+
+		// Create and register the level editor toolbar menu extension
+		CinematicsMenuExtender = MakeShareable(new FExtender);
+		CinematicsMenuExtender->AddMenuExtension("LevelEditorNewMatinee", EExtensionHook::First, CommandList, FMenuExtensionDelegate::CreateStatic([](FMenuBuilder& MenuBuilder){
+			MenuBuilder.AddMenuEntry(FActorAnimationExtensionCommands::Get().CreateNewActorAnimationInLevel);
+		}));
+
+		FLevelEditorModule& LevelEditorModule = FModuleManager::LoadModuleChecked<FLevelEditorModule>("LevelEditor");
+		LevelEditorModule.GetAllLevelEditorToolbarCinematicsMenuExtenders().Add(CinematicsMenuExtender);
+	}
+
+	/**
+	 * Unregisters menu extensions for the level editor toolbar
+	 */
+	void UnregisterMenuExtensions()
+	{
+		if (FLevelEditorModule* LevelEditorModule = FModuleManager::GetModulePtr<FLevelEditorModule>("LevelEditor"))
+		{
+			LevelEditorModule->GetAllLevelEditorToolbarCinematicsMenuExtenders().Remove(CinematicsMenuExtender);
+		}
+		CinematicsMenuExtender = nullptr;
+		CommandList = nullptr;
+
+		FActorAnimationExtensionCommands::Unregister();
+	}
+
+	/**
+	 * Callback for creating a new actor animation asset in the level
+	 */
+	static void OnCreateActorInLevel()
+	{
+		// Create a new actor animation
+		IAssetTools& AssetTools = FModuleManager::GetModuleChecked<FAssetToolsModule>("AssetTools").Get();
+
+		UObject* NewAsset = nullptr;
+
+		// Attempt to create a new asset
+		for (TObjectIterator<UClass> It ; It ; ++It)
+		{
+			UClass* CurrentClass = *It;
+			if (CurrentClass->IsChildOf(UFactory::StaticClass()) && !(CurrentClass->HasAnyClassFlags(CLASS_Abstract)))
+			{
+				UFactory* Factory = Cast<UFactory>(CurrentClass->GetDefaultObject());
+				if (Factory->CanCreateNew() && Factory->ImportPriority >= 0 && Factory->SupportedClass == UActorAnimation::StaticClass())
+				{
+					NewAsset = AssetTools.CreateAsset(UActorAnimation::StaticClass(), Factory);
+					break;
+				}
+			}
+		}
+
+		if (!NewAsset)
+		{
+			return;
+		}
+
+		// Spawn a  actor at the origin, and either move infront of the camera or focus camera on it (depending on the viewport) and open for edit
+		UActorFactory* ActorFactory = GEditor->FindActorFactoryForActorClass( AMovieSceneActor::StaticClass() );
+		if (!ensure(ActorFactory))
+		{
+			return;
+		}
+
+		AMovieSceneActor* NewActor = CastChecked<AMovieSceneActor>(GEditor->UseActorFactory(ActorFactory, FAssetData(NewAsset), &FTransform::Identity));
+		if( GCurrentLevelEditingViewportClient->IsPerspective() )
+		{
+			GEditor->MoveActorInFrontOfCamera( *NewActor, GCurrentLevelEditingViewportClient->GetViewLocation(), GCurrentLevelEditingViewportClient->GetViewRotation().Vector() );
+		}
+		else
+		{
+			GEditor->MoveViewportCamerasToActor( *NewActor, false );
+		}
+
+		FAssetEditorManager::Get().OpenEditorForAsset(NewAsset);
+	}
+
 private:
 
 	/** The collection of registered asset type actions. */
@@ -72,7 +183,14 @@ private:
 
 	/** Holds the plug-ins style set. */
 	TSharedPtr<ISlateStyle> Style;
+
+	/** Extender for the cinematics menu */
+	TSharedPtr<FExtender> CinematicsMenuExtender;
+
+	TSharedPtr<FUICommandList> CommandList;
 };
 
 
 IMPLEMENT_MODULE(FActorAnimationEditorModule, ActorAnimationEditor);
+
+#undef LOCTEXT_NAMESPACE
\ No newline at end of file
diff --git a/Engine/Plugins/MovieScene/ActorAnimationEditor/Source/ActorAnimationEditor/Private/Styles/ActorAnimationEditorStyle.h b/Engine/Plugins/MovieScene/ActorAnimationEditor/Source/ActorAnimationEditor/Private/Styles/ActorAnimationEditorStyle.h
--- a/Engine/Plugins/MovieScene/ActorAnimationEditor/Source/ActorAnimationEditor/Private/Styles/ActorAnimationEditorStyle.h
+++ b/Engine/Plugins/MovieScene/ActorAnimationEditor/Source/ActorAnimationEditor/Private/Styles/ActorAnimationEditorStyle.h
@@ -30,7 +30,9 @@ public:
 
 		// tab icons
 		Set("ActorAnimationEditor.Tabs.Sequencer", new IMAGE_BRUSH("icon_tab_sequencer_16x", Icon16x16));
-	
+		Set("ActorAnimationEditor.CreateNewActorAnimationInLevel", new IMAGE_BRUSH("CreateNewActorAnimationInLevel_16x", Icon16x16));
+		Set("ActorAnimationEditor.CreateNewActorAnimationInLevel.Small", new IMAGE_BRUSH("CreateNewActorAnimationInLevel_16x", Icon16x16));
+
 		FSlateStyleRegistry::RegisterSlateStyle(*this);
 	 }
 
diff --git a/Engine/Source/Editor/LevelEditor/LevelEditor.Build.cs b/Engine/Source/Editor/LevelEditor/LevelEditor.Build.cs
--- a/Engine/Source/Editor/LevelEditor/LevelEditor.Build.cs
+++ b/Engine/Source/Editor/LevelEditor/LevelEditor.Build.cs
@@ -31,6 +31,7 @@ public class LevelEditor : ModuleRules
 
 		PrivateDependencyModuleNames.AddRange(
 			new string[] {
+				"ActorAnimation",
 				"Analytics",
 				"Core",
 				"CoreUObject",
diff --git a/Engine/Source/Editor/LevelEditor/Private/LevelEditorToolBar.cpp b/Engine/Source/Editor/LevelEditor/Private/LevelEditorToolBar.cpp
--- a/Engine/Source/Editor/LevelEditor/Private/LevelEditorToolBar.cpp
+++ b/Engine/Source/Editor/LevelEditor/Private/LevelEditorToolBar.cpp
@@ -15,6 +15,7 @@
 #include "Editor/ContentBrowser/Public/ContentBrowserModule.h"
 #include "EngineBuildSettings.h"
 #include "Matinee/MatineeActor.h"
+#include "MovieSceneActor.h"
 #include "Engine/LevelScriptBlueprint.h"
 #include "ISettingsCategory.h"
 #include "ISettingsContainer.h"
@@ -1230,9 +1231,9 @@ TSharedRef< SWidget > FLevelEditorToolBar::MakeLevelEditorToolBar( const TShared
 
 		ToolbarBuilder.AddComboButton(
 			FUIAction(),
-			FOnGetContent::CreateStatic( &FLevelEditorToolBar::GenerateMatineeMenuContent, InCommandList, TWeakPtr<SLevelEditor>( InLevelEditor ) ),
-			LOCTEXT( "EditMatinee_Label", "Matinee" ),
-			LOCTEXT( "EditMatinee_Tooltip", "Displays a list of Matinee objects to open in the Matinee Editor"),
+			FOnGetContent::CreateStatic( &FLevelEditorToolBar::GenerateCinematicsMenuContent, InCommandList, TWeakPtr<SLevelEditor>( InLevelEditor ) ),
+			LOCTEXT( "EditCinematics_Label", "Cinematics" ),
+			LOCTEXT( "EditCinematics_Tooltip", "Displays a list of Matinee and Actor Animation objects to open in their respective editors"),
 			FSlateIcon(FEditorStyle::GetStyleSetName(), "LevelEditor.EditMatinee") 
 			);
 	}
@@ -2020,14 +2021,17 @@ void FLevelEditorToolBar::OnOpenSubLevelBlueprint( ULevel* InLevel )
 	}
 }
 
-TSharedRef< SWidget > FLevelEditorToolBar::GenerateMatineeMenuContent( TSharedRef<FUICommandList> InCommandList, TWeakPtr<SLevelEditor> LevelEditorWeakPtr )
+TSharedRef< SWidget > FLevelEditorToolBar::GenerateCinematicsMenuContent( TSharedRef<FUICommandList> InCommandList, TWeakPtr<SLevelEditor> LevelEditorWeakPtr )
 {
-#define LOCTEXT_NAMESPACE "LevelToolBarMatineeMenu"
+#define LOCTEXT_NAMESPACE "LevelToolBarCinematicsMenu"
+
+	FLevelEditorModule& LevelEditorModule = FModuleManager::LoadModuleChecked<FLevelEditorModule>("LevelEditor");
+	TSharedPtr<FExtender> Extender = FExtender::Combine(LevelEditorModule.GetAllLevelEditorToolbarCinematicsMenuExtenders());
 
 	const bool bShouldCloseWindowAfterMenuSelection = true;
-	FMenuBuilder MenuBuilder( bShouldCloseWindowAfterMenuSelection, InCommandList );
+	FMenuBuilder MenuBuilder( bShouldCloseWindowAfterMenuSelection, InCommandList, Extender );
 
-	// We can't build a list of Matinees while the current World is a PIE world.
+	// We can't build a list of Matinees and MovieSceneActors while the current World is a PIE world.
 	SceneOutliner::FInitializationOptions InitOptions;
 	{
 		InitOptions.Mode = ESceneOutlinerMode::ActorPicker;
@@ -2037,19 +2041,15 @@ TSharedRef< SWidget > FLevelEditorToolBar::GenerateMatineeMenuContent( TSharedRe
 		InitOptions.bShowHeaderRow = false;
 		InitOptions.bShowSearchBox = false;
 		InitOptions.bShowCreateNewFolder = false;
-		struct Local
-		{
-			static bool IsMatineeActor( const AActor* const Actor )
-			{
-				return Actor->IsA( AMatineeActor::StaticClass() );
-			}
-		};
 
-		// Only display Matinee actors
-		InitOptions.Filters->AddFilterPredicate( SceneOutliner::FActorFilterPredicate::CreateStatic( &Local::IsMatineeActor ) );
+		// Only display Matinee and MovieScene actors
+		auto ActorFilter = [](const AActor* Actor){
+			return Actor->IsA( AMatineeActor::StaticClass() ) || Actor->IsA( AMovieSceneActor::StaticClass() );
+		};
+		InitOptions.Filters->AddFilterPredicate( SceneOutliner::FActorFilterPredicate::CreateLambda( ActorFilter ) );
 	}
 
-	// actor selector to allow the user to choose a Matinee actor
+	// actor selector to allow the user to choose an actor
 	FSceneOutlinerModule& SceneOutlinerModule = FModuleManager::LoadModuleChecked<FSceneOutlinerModule>( "SceneOutliner" );
 	TSharedRef< SWidget > MiniSceneOutliner =
 		SNew( SVerticalBox )
@@ -2059,7 +2059,7 @@ TSharedRef< SWidget > FLevelEditorToolBar::GenerateMatineeMenuContent( TSharedRe
 		[
 			SceneOutlinerModule.CreateSceneOutliner(
 				InitOptions,
-				FOnActorPicked::CreateStatic( &FLevelEditorToolBar::OnMatineeActorPicked ) )
+				FOnActorPicked::CreateStatic( &FLevelEditorToolBar::OnCinematicsActorPicked ) )
 		];
 
 	static const FName DefaultForegroundName("DefaultForeground");
@@ -2081,7 +2081,7 @@ TSharedRef< SWidget > FLevelEditorToolBar::GenerateMatineeMenuContent( TSharedRe
 			.HAlign( HAlign_Center )
 			[
 				SNew( STextBlock )
-				.Text( LOCTEXT( "SelectMatineeActorToEdit", "Select a Matinee actor" ) )
+				.Text( LOCTEXT( "SelectCinematicsActorToEdit", "Select an actor" ) )
 			]
 
 			+SVerticalBox::Slot()
@@ -2099,10 +2099,8 @@ TSharedRef< SWidget > FLevelEditorToolBar::GenerateMatineeMenuContent( TSharedRe
 	}
 	MenuBuilder.EndSection();
 
-	bool bHasAnyMatineeActors = false;
-	TActorIterator<AMatineeActor> MatineeIt( LevelEditorWeakPtr.Pin()->GetWorld() );
-
-	bHasAnyMatineeActors = !!MatineeIt;
+	UWorld* World = LevelEditorWeakPtr.Pin()->GetWorld();
+	const bool bHasAnyMatineeActors = !!TActorIterator<AMatineeActor>(World) || !!TActorIterator<AMovieSceneActor>(World);
 
 	//Add a heading to separate the existing matinees from the 'Add New Matinee Actor' button
 	MenuBuilder.BeginSection("LevelEditorExistingMatinee", LOCTEXT( "MatineeMenuCombo_ExistingHeading", "Edit Existing Matinee" ) );
@@ -2118,17 +2116,24 @@ TSharedRef< SWidget > FLevelEditorToolBar::GenerateMatineeMenuContent( TSharedRe
 	return MenuBuilder.MakeWidget();
 }
 
-void FLevelEditorToolBar::OnMatineeActorPicked( AActor* Actor )
+void FLevelEditorToolBar::OnCinematicsActorPicked( AActor* Actor )
 {
 	//The matinee editor will not tick unless the editor viewport is in realtime mode.
 	//the scene outliner eats input, so we must close any popups manually.
 	FSlateApplication::Get().DismissAllMenus();
 
 	// Make sure we dismiss the menus before we open this
-	AMatineeActor* MatineeActor = Cast<AMatineeActor>( Actor );
-	if( MatineeActor != NULL )
+	if(AMatineeActor* MatineeActor = Cast<AMatineeActor>(Actor))
 	{
 		// Open Matinee for editing!
 		GEditor->OpenMatinee( MatineeActor );
 	}
+	else if (AMovieSceneActor* MovieSceneActor = Cast<AMovieSceneActor>(Actor))
+	{
+		UObject* Animation = MovieSceneActor->ActorAnimation.TryLoad();
+		if (Animation)
+		{
+			FAssetEditorManager::Get().OpenEditorForAsset(Animation);
+		}
+	}
 }
diff --git a/Engine/Source/Editor/LevelEditor/Private/LevelEditorToolBar.h b/Engine/Source/Editor/LevelEditor/Private/LevelEditorToolBar.h
--- a/Engine/Source/Editor/LevelEditor/Private/LevelEditorToolBar.h
+++ b/Engine/Source/Editor/LevelEditor/Private/LevelEditorToolBar.h
@@ -64,17 +64,17 @@ protected:
 	static TSharedRef< SWidget > GenerateOpenBlueprintMenuContent( TSharedRef<FUICommandList> InCommandList, TWeakPtr< SLevelEditor > InLevelEditor );
 
 	/**
-	 * Generates menu content for the matinee combo button drop down menu
+	 * Generates menu content for the Cinematics combo button drop down menu
 	 *
 	 * @return	Menu content widget
 	 */
-	static TSharedRef< SWidget > GenerateMatineeMenuContent( TSharedRef<FUICommandList> InCommandList, TWeakPtr<SLevelEditor> LevelEditorWeakPtr );
+	static TSharedRef< SWidget > GenerateCinematicsMenuContent( TSharedRef<FUICommandList> InCommandList, TWeakPtr<SLevelEditor> LevelEditorWeakPtr );
 
 	/**
-	 * Delegate for actor selection within the Matinee popup menu's SceneOutliner.
-	 * Opens the matinee editor for the selected actor and dismisses all popup menus.
+	 * Delegate for actor selection within the Cinematics popup menu's SceneOutliner.
+	 * Opens the editor for the selected actor and dismisses all popup menus.
 	 */
-	static void OnMatineeActorPicked( AActor* Actor );
+	static void OnCinematicsActorPicked( AActor* Actor );
 
 	/**
 	 * Callback to open a sub-level script Blueprint
diff --git a/Engine/Source/Editor/LevelEditor/Public/LevelEditor.h b/Engine/Source/Editor/LevelEditor/Public/LevelEditor.h
--- a/Engine/Source/Editor/LevelEditor/Public/LevelEditor.h
+++ b/Engine/Source/Editor/LevelEditor/Public/LevelEditor.h
@@ -204,6 +204,7 @@ public:
 	virtual TArray<FLevelEditorMenuExtender>& GetAllLevelEditorToolbarCompileMenuExtenders() {return LevelEditorToolbarCompileMenuExtenders;}
 	virtual TArray<FLevelEditorMenuExtender>& GetAllLevelEditorToolbarSourceControlMenuExtenders() { return LevelEditorToolbarSourceControlMenuExtenders; }
 	virtual TArray<FLevelEditorMenuExtender>& GetAllLevelEditorToolbarCreateMenuExtenders() { return LevelEditorToolbarCreateMenuExtenders; }
+	virtual TArray<TSharedPtr<FExtender>>& GetAllLevelEditorToolbarCinematicsMenuExtenders() {return LevelEditorToolbarCinematicsMenuExtenders;}
 	
 	/** Gets the extensibility managers for outside entities to extend static mesh editor's menus and toolbars */
 	virtual TSharedPtr<FExtensibilityManager> GetMenuExtensibilityManager() override {return MenuExtensibilityManager;}
@@ -279,6 +280,7 @@ private:
 	TArray<FLevelEditorMenuExtender> LevelEditorToolbarCompileMenuExtenders;
 	TArray<FLevelEditorMenuExtender> LevelEditorToolbarSourceControlMenuExtenders;
 	TArray<FLevelEditorMenuExtender> LevelEditorToolbarCreateMenuExtenders;
+	TArray<TSharedPtr<FExtender>> LevelEditorToolbarCinematicsMenuExtenders;
 
 	/* Pointer to the current level Editor instance */
 	TWeakPtr<class SLevelEditor> LevelEditorInstancePtr;
