Fixed crash when attempting to open a Level Sequence after closing down an empty sequencer tab

The sequencer tab was being erroneously referenced externally

This addresses UE-21223

diff --git a/Engine/Source/Editor/LevelEditor/Private/SLevelEditor.cpp b/Engine/Source/Editor/LevelEditor/Private/SLevelEditor.cpp
--- a/Engine/Source/Editor/LevelEditor/Private/SLevelEditor.cpp
+++ b/Engine/Source/Editor/LevelEditor/Private/SLevelEditor.cpp
@@ -535,15 +535,15 @@ void SLevelEditor::AttachSequencer( TSharedPtr<SWidget> SequencerWidget, TShared
 
 		if(SequencerWidget.IsValid() && NewSequencerAssetEditor.IsValid())
 		{
-			InvokeTab("Sequencer");
-			SequencerTab->SetOnTabClosed(SDockTab::FOnTabClosedCallback::CreateStatic(&Local::OnSequencerClosed, TWeakPtr<IAssetEditorInstance>(NewSequencerAssetEditor)));
-			SequencerTab->SetContent(SequencerWidget.ToSharedRef());
+			TSharedRef<SDockTab> Tab = InvokeTab("Sequencer");
+
+			Tab->SetOnTabClosed(SDockTab::FOnTabClosedCallback::CreateStatic(&Local::OnSequencerClosed, TWeakPtr<IAssetEditorInstance>(NewSequencerAssetEditor)));
+			Tab->SetContent(SequencerWidget.ToSharedRef());
 
 			SequencerAssetEditor = NewSequencerAssetEditor;
 		}
 		else
 		{
-			SequencerTab.Reset();
 			SequencerAssetEditor.Reset();
 		}
 	}
@@ -755,13 +755,12 @@ TSharedRef<SDockTab> SLevelEditor::SpawnLevelEditorTab( const FSpawnTabArgs& Arg
 		if (FSlateStyleRegistry::FindSlateStyle("LevelSequenceEditorStyle"))
 		{
 			// @todo sequencer: remove when world-centric mode is added
-			SequencerTab = SNew(SDockTab)
+			return SNew(SDockTab)
 				.Icon( FSlateStyleRegistry::FindSlateStyle("LevelSequenceEditorStyle")->GetBrush("LevelSequenceEditor.Tabs.Sequencer") )
 				.Label( NSLOCTEXT("Sequencer", "SequencerMainTitle", "Sequencer") )
 				[
 					SNullWidget::NullWidget
 				];
-			return SequencerTab.ToSharedRef();
 		}
 	}
 	else if( TabIdentifier == LevelEditorStatsViewerTab )
@@ -771,7 +770,7 @@ TSharedRef<SDockTab> SLevelEditor::SpawnLevelEditorTab( const FSpawnTabArgs& Arg
 			.Icon( FEditorStyle::GetBrush( "LevelEditor.Tabs.StatsViewer" ) )
 			.Label( NSLOCTEXT("LevelEditor", "StatsViewerTabTitle", "Statistics") )
 			[
-				StatsViewerModule.CreateStatsViewer()					
+				StatsViewerModule.CreateStatsViewer()
 			];
 	}
 	else if ( TabIdentifier == "WorldSettingsTab" )
@@ -799,10 +798,10 @@ TSharedRef<SDockTab> SLevelEditor::SpawnLevelEditorTab( const FSpawnTabArgs& Arg
 	return SNew(SDockTab);
 }
 
-void SLevelEditor::InvokeTab( FName TabID )
+TSharedRef<SDockTab> SLevelEditor::InvokeTab( FName TabID )
 {
 	TSharedPtr<FTabManager> LevelEditorTabManager = GetTabManager();
-	LevelEditorTabManager->InvokeTab(TabID);
+	return LevelEditorTabManager->InvokeTab(TabID);
 }
 
 void SLevelEditor::SyncDetailsToSelection()
diff --git a/Engine/Source/Editor/LevelEditor/Private/SLevelEditor.h b/Engine/Source/Editor/LevelEditor/Private/SLevelEditor.h
--- a/Engine/Source/Editor/LevelEditor/Private/SLevelEditor.h
+++ b/Engine/Source/Editor/LevelEditor/Private/SLevelEditor.h
@@ -69,10 +69,10 @@ public:
 	virtual void AppendCommands( const TSharedRef<FUICommandList>& InCommandsToAppend ) override;
 
 	/**
-	 * Given a tab ID, summons a new tab in the tab stack specified.
-	 * If SummonInStack is null, use the default location
+	 * Given a tab ID, summons a new tab in the position saved in the current layout, or in a default position.
+	 * @return the invoked tab
 	 */
-	void InvokeTab( FName TabID );
+	TSharedRef<SDockTab> InvokeTab( FName TabID );
 
 	/**
 	 * Sync the details panel to the current selection
@@ -229,8 +229,7 @@ private:
 	/** List of all actor details panels to update when selection changes */
 	TArray< TWeakPtr<class SActorDetails> > AllActorDetailPanels;
 
-	/** Attached sequencer info */
-	TSharedPtr<SDockTab> SequencerTab;
+	/** Attached sequencer asset editor */
 	TWeakPtr<IAssetEditorInstance> SequencerAssetEditor;
 };
 
