Copying //UE4/Dev-Sequencer to Dev-Main (//UE4/Dev-Main)

diff --git a/Engine/Source/Editor/LevelEditor/Private/LevelViewportLayoutThreePanes.cpp b/Engine/Source/Editor/LevelEditor/Private/LevelViewportLayoutThreePanes.cpp
--- a/Engine/Source/Editor/LevelEditor/Private/LevelViewportLayoutThreePanes.cpp
+++ b/Engine/Source/Editor/LevelEditor/Private/LevelViewportLayoutThreePanes.cpp
@@ -115,7 +115,7 @@ void FLevelViewportLayoutThreePanes::ReplaceWidget( TSharedRef< SWidget > Source
 // FLevelViewportLayoutThreePanesLeft /////////////////////////////
 
 TSharedRef<SWidget> FLevelViewportLayoutThreePanesLeft::MakeThreePanelWidget(
-	TArray<TSharedPtr< SLevelViewport >>& ViewportWidgets,
+	TMap< FName, TSharedPtr< IViewportLayoutEntity >>& ViewportWidgets,
 	const FString& ViewportKey0, const FString& ViewportKey1, const FString& ViewportKey2,
 	float PrimarySplitterPercentage, float SecondarySplitterPercentage)
 {
@@ -130,7 +130,7 @@ TSharedRef<SWidget> FLevelViewportLayoutThreePanesLeft::MakeThreePanelWidget(
 		.Value(PrimarySplitterPercentage)
 		[
 			SAssignNew( ViewportWidget0, SLevelViewport )
-			.ViewportType( LVT_Perspective )	
+			.ViewportType( LVT_Perspective )
 			.Realtime( true )
 			.ParentLayout( AsShared() )
 			.ParentLevelEditor( ParentLevelEditor )
@@ -161,9 +161,9 @@ TSharedRef<SWidget> FLevelViewportLayoutThreePanesLeft::MakeThreePanelWidget(
 			]
 		];
 
-	ViewportWidgets.Add( ViewportWidget0 );
-	ViewportWidgets.Add( ViewportWidget1 );
-	ViewportWidgets.Add( ViewportWidget2 );
+	ViewportWidgets.Add( *ViewportKey0, MakeShareable( new FLevelViewportLayoutEntity(ViewportWidget0.ToSharedRef()) ) );
+	ViewportWidgets.Add( *ViewportKey1, MakeShareable( new FLevelViewportLayoutEntity(ViewportWidget1.ToSharedRef()) ) );
+	ViewportWidgets.Add( *ViewportKey2, MakeShareable( new FLevelViewportLayoutEntity(ViewportWidget2.ToSharedRef()) ) );
 
 	// Make newly-created perspective viewports active by default
 	GCurrentLevelEditingViewportClient = &ViewportWidget0->GetLevelViewportClient();
@@ -175,7 +175,7 @@ TSharedRef<SWidget> FLevelViewportLayoutThreePanesLeft::MakeThreePanelWidget(
 // FLevelViewportLayoutThreePanesRight /////////////////////////////
 
 TSharedRef<SWidget> FLevelViewportLayoutThreePanesRight::MakeThreePanelWidget(
-	TArray<TSharedPtr< SLevelViewport >>& ViewportWidgets,
+	TMap< FName, TSharedPtr< IViewportLayoutEntity >>& ViewportWidgets,
 	const FString& ViewportKey0, const FString& ViewportKey1, const FString& ViewportKey2,
 	float PrimarySplitterPercentage, float SecondarySplitterPercentage)
 {
@@ -221,9 +221,9 @@ TSharedRef<SWidget> FLevelViewportLayoutThreePanesRight::MakeThreePanelWidget(
 			.ConfigKey( ViewportKey0 )
 		];
 
-	ViewportWidgets.Add( ViewportWidget0 );
-	ViewportWidgets.Add( ViewportWidget1 );
-	ViewportWidgets.Add( ViewportWidget2 );
+	ViewportWidgets.Add( *ViewportKey0, MakeShareable( new FLevelViewportLayoutEntity(ViewportWidget0.ToSharedRef()) ) );
+	ViewportWidgets.Add( *ViewportKey1, MakeShareable( new FLevelViewportLayoutEntity(ViewportWidget1.ToSharedRef()) ) );
+	ViewportWidgets.Add( *ViewportKey2, MakeShareable( new FLevelViewportLayoutEntity(ViewportWidget2.ToSharedRef()) ) );
 
 	// Make newly-created perspective viewports active by default
 	GCurrentLevelEditingViewportClient = &ViewportWidget0->GetLevelViewportClient();
@@ -235,7 +235,7 @@ TSharedRef<SWidget> FLevelViewportLayoutThreePanesRight::MakeThreePanelWidget(
 // FLevelViewportLayoutThreePanesTop /////////////////////////////
 
 TSharedRef<SWidget> FLevelViewportLayoutThreePanesTop::MakeThreePanelWidget(
-	TArray<TSharedPtr< SLevelViewport >>& ViewportWidgets,
+	TMap< FName, TSharedPtr< IViewportLayoutEntity >>& ViewportWidgets,
 	const FString& ViewportKey0, const FString& ViewportKey1, const FString& ViewportKey2,
 	float PrimarySplitterPercentage, float SecondarySplitterPercentage)
 {
@@ -281,9 +281,9 @@ TSharedRef<SWidget> FLevelViewportLayoutThreePanesTop::MakeThreePanelWidget(
 			]
 		];
 
-	ViewportWidgets.Add( ViewportWidget0 );
-	ViewportWidgets.Add( ViewportWidget1 );
-	ViewportWidgets.Add( ViewportWidget2 );
+	ViewportWidgets.Add( *ViewportKey0, MakeShareable( new FLevelViewportLayoutEntity(ViewportWidget0.ToSharedRef()) ) );
+	ViewportWidgets.Add( *ViewportKey1, MakeShareable( new FLevelViewportLayoutEntity(ViewportWidget1.ToSharedRef()) ) );
+	ViewportWidgets.Add( *ViewportKey2, MakeShareable( new FLevelViewportLayoutEntity(ViewportWidget2.ToSharedRef()) ) );
 
 	// Make newly-created perspective viewports active by default
 	GCurrentLevelEditingViewportClient = &ViewportWidget0->GetLevelViewportClient();
@@ -295,7 +295,7 @@ TSharedRef<SWidget> FLevelViewportLayoutThreePanesTop::MakeThreePanelWidget(
 // FLevelViewportLayoutThreePanesBottom /////////////////////////////
 
 TSharedRef<SWidget> FLevelViewportLayoutThreePanesBottom::MakeThreePanelWidget(
-	TArray<TSharedPtr< SLevelViewport >>& ViewportWidgets,
+	TMap< FName, TSharedPtr< IViewportLayoutEntity >>& ViewportWidgets,
 	const FString& ViewportKey0, const FString& ViewportKey1, const FString& ViewportKey2,
 	float PrimarySplitterPercentage, float SecondarySplitterPercentage)
 {
@@ -334,16 +334,16 @@ TSharedRef<SWidget> FLevelViewportLayoutThreePanesBottom::MakeThreePanelWidget(
 		.Value(1.0f - PrimarySplitterPercentage)
 		[
 			SAssignNew( ViewportWidget0, SLevelViewport )
-			.ViewportType( LVT_Perspective )	
+			.ViewportType( LVT_Perspective )
 			.Realtime( true )
 			.ParentLayout( AsShared() )
 			.ParentLevelEditor( ParentLevelEditor )
 			.ConfigKey( ViewportKey0 )
 		];
 
-	ViewportWidgets.Add( ViewportWidget0 );
-	ViewportWidgets.Add( ViewportWidget1 );
-	ViewportWidgets.Add( ViewportWidget2 );
+	ViewportWidgets.Add( *ViewportKey0, MakeShareable( new FLevelViewportLayoutEntity(ViewportWidget0.ToSharedRef()) ) );
+	ViewportWidgets.Add( *ViewportKey1, MakeShareable( new FLevelViewportLayoutEntity(ViewportWidget1.ToSharedRef()) ) );
+	ViewportWidgets.Add( *ViewportKey2, MakeShareable( new FLevelViewportLayoutEntity(ViewportWidget2.ToSharedRef()) ) );
 
 	// Make newly-created perspective viewports active by default
 	GCurrentLevelEditingViewportClient = &ViewportWidget0->GetLevelViewportClient();
diff --git a/Engine/Source/Editor/LevelEditor/Private/LevelViewportLayoutThreePanes.h b/Engine/Source/Editor/LevelEditor/Private/LevelViewportLayoutThreePanes.h
--- a/Engine/Source/Editor/LevelEditor/Private/LevelViewportLayoutThreePanes.h
+++ b/Engine/Source/Editor/LevelEditor/Private/LevelViewportLayoutThreePanes.h
@@ -22,7 +22,7 @@ protected:
 	virtual TSharedRef<SWidget> MakeViewportLayout(const FString& LayoutString) override;
 
 	virtual TSharedRef<SWidget> MakeThreePanelWidget(
-		TArray<TSharedPtr< SLevelViewport >>& ViewportWidgets,
+		TMap< FName, TSharedPtr< IViewportLayoutEntity >>& ViewportWidgets,
 		const FString& ViewportKey0, const FString& ViewportKey1, const FString& ViewportKey2,
 		float PrimarySplitterPercentage, float SecondarySplitterPercentage) = 0;
 
@@ -44,7 +44,7 @@ public:
 	virtual const FName& GetLayoutTypeName() const override { return LevelViewportConfigurationNames::ThreePanesLeft; }
 
 	virtual TSharedRef<SWidget> MakeThreePanelWidget(
-		TArray<TSharedPtr< SLevelViewport >>& ViewportWidgets,
+		TMap< FName, TSharedPtr< IViewportLayoutEntity >>& ViewportWidgets,
 		const FString& ViewportKey0, const FString& ViewportKey1, const FString& ViewportKey2,
 		float PrimarySplitterPercentage, float SecondarySplitterPercentage) override;
 };
@@ -58,7 +58,7 @@ public:
 	virtual const FName& GetLayoutTypeName() const override { return LevelViewportConfigurationNames::ThreePanesRight; }
 
 	virtual TSharedRef<SWidget> MakeThreePanelWidget(
-		TArray<TSharedPtr< SLevelViewport >>& ViewportWidgets,
+		TMap< FName, TSharedPtr< IViewportLayoutEntity >>& ViewportWidgets,
 		const FString& ViewportKey0, const FString& ViewportKey1, const FString& ViewportKey2,
 		float PrimarySplitterPercentage, float SecondarySplitterPercentage) override;
 };
@@ -72,7 +72,7 @@ public:
 	virtual const FName& GetLayoutTypeName() const override { return LevelViewportConfigurationNames::ThreePanesTop; }
 
 	virtual TSharedRef<SWidget> MakeThreePanelWidget(
-		TArray<TSharedPtr< SLevelViewport >>& ViewportWidgets,
+		TMap< FName, TSharedPtr< IViewportLayoutEntity >>& ViewportWidgets,
 		const FString& ViewportKey0, const FString& ViewportKey1, const FString& ViewportKey2,
 		float PrimarySplitterPercentage, float SecondarySplitterPercentage) override;
 };
@@ -86,7 +86,7 @@ public:
 	virtual const FName& GetLayoutTypeName() const override { return LevelViewportConfigurationNames::ThreePanesBottom; }
 
 	virtual TSharedRef<SWidget> MakeThreePanelWidget(
-		TArray<TSharedPtr< SLevelViewport >>& ViewportWidgets,
+		TMap< FName, TSharedPtr< IViewportLayoutEntity >>& ViewportWidgets,
 		const FString& ViewportKey0, const FString& ViewportKey1, const FString& ViewportKey2,
 		float PrimarySplitterPercentage, float SecondarySplitterPercentage) override;
 };
diff --git a/Engine/Source/Editor/LevelEditor/Private/LevelViewportLayoutTwoPanes.cpp b/Engine/Source/Editor/LevelEditor/Private/LevelViewportLayoutTwoPanes.cpp
--- a/Engine/Source/Editor/LevelEditor/Private/LevelViewportLayoutTwoPanes.cpp
+++ b/Engine/Source/Editor/LevelEditor/Private/LevelViewportLayoutTwoPanes.cpp
@@ -69,8 +69,8 @@ TSharedRef<SWidget> TLevelViewportLayoutTwoPanes<TOrientation>::MakeViewportLayo
 			.ConfigKey( ViewportKey1 )
 		];
 
-	Viewports.Add( ViewportWidget0 );
-	Viewports.Add( ViewportWidget1 );
+	Viewports.Add( *ViewportKey0, MakeShareable( new FLevelViewportLayoutEntity(ViewportWidget0.ToSharedRef()) ) );
+	Viewports.Add( *ViewportKey1, MakeShareable( new FLevelViewportLayoutEntity(ViewportWidget1.ToSharedRef()) ) );
 
 	// Make newly-created perspective viewports active by default
 	GCurrentLevelEditingViewportClient = &ViewportWidget1->GetLevelViewportClient();
diff --git a/Engine/Source/Editor/LevelEditor/Private/LevelViewportTabContent.cpp b/Engine/Source/Editor/LevelEditor/Private/LevelViewportTabContent.cpp
--- a/Engine/Source/Editor/LevelEditor/Private/LevelViewportTabContent.cpp
+++ b/Engine/Source/Editor/LevelEditor/Private/LevelViewportTabContent.cpp
@@ -30,8 +30,16 @@ TSharedPtr< class FLevelViewportLayout > FLevelViewportTabContent::ConstructView
 	else if (TypeName == LevelViewportConfigurationNames::FourPanesBottom) ViewportLayout = MakeShareable(new FLevelViewportLayoutFourPanesBottom);
 	else if (TypeName == LevelViewportConfigurationNames::FourPanesTop) ViewportLayout = MakeShareable(new FLevelViewportLayoutFourPanesTop);
 	else if (TypeName == LevelViewportConfigurationNames::OnePane) ViewportLayout = MakeShareable(new FLevelViewportLayoutOnePane);
+	else
+	{
+		FLevelEditorModule& LevelEditorModule = FModuleManager::Get().GetModuleChecked<FLevelEditorModule>("LevelEditor");
+		ViewportLayout = LevelEditorModule.FactoryCustomViewportLayout(TypeName);
+	}
 
-	check (ViewportLayout.IsValid());
+	if (!ensure(ViewportLayout.IsValid()))
+	{
+		ViewportLayout = MakeShareable(new FLevelViewportLayoutOnePane);
+	}
 	ViewportLayout->SetIsReplacement(bSwitchingLayouts);
 	return ViewportLayout;
 }
@@ -65,7 +73,7 @@ bool FLevelViewportTabContent::IsVisible() const
 	return false;
 }
 
-const TArray< TSharedPtr< SLevelViewport > >* FLevelViewportTabContent::GetViewports() const
+const TMap< FName, TSharedPtr< IViewportLayoutEntity > >* FLevelViewportTabContent::GetViewports() const
 {
 	if (ActiveLevelViewportLayout.IsValid())
 	{
diff --git a/Engine/Source/Editor/LevelEditor/Private/SLevelEditor.cpp b/Engine/Source/Editor/LevelEditor/Private/SLevelEditor.cpp
--- a/Engine/Source/Editor/LevelEditor/Private/SLevelEditor.cpp
+++ b/Engine/Source/Editor/LevelEditor/Private/SLevelEditor.cpp
@@ -37,7 +37,7 @@
 #include "SActorDetails.h"
 #include "ScopedTransaction.h"
 #include "GameFramework/WorldSettings.h"
-
+#include "LevelViewportLayout.h"
 #include "HierarchicalLODOutlinerModule.h"
 
 
@@ -300,15 +300,14 @@ bool SLevelEditor::HasActivePlayInEditorViewport() const
 		if (ViewportTab.IsValid())
 		{
 			// Get all the viewports in the layout
-			const TArray< TSharedPtr< SLevelViewport > >* LevelViewports = ViewportTab.Pin()->GetViewports();
+			const TMap< FName, TSharedPtr< IViewportLayoutEntity > >* LevelViewports = ViewportTab.Pin()->GetViewports();
 
 			if (LevelViewports != NULL)
 			{
 				// Search for a viewport with a pie session
-				for( int32 ViewportIndex = 0; ViewportIndex < LevelViewports->Num(); ++ViewportIndex )
+				for (auto& Pair : *LevelViewports)
 				{
-					const TSharedPtr< SLevelViewport >& Viewport = (*LevelViewports)[ ViewportIndex ];
-					if( Viewport->IsPlayInEditorViewportActive() )
+					if (Pair.Value->IsPlayInEditorViewportActive())
 					{
 						return true;
 					}
@@ -336,15 +335,14 @@ TSharedPtr<SLevelViewport> SLevelEditor::GetActiveViewport()
 			// Only check the viewports in the tab if its visible
 			if( ViewportTab->IsVisible() )
 			{
-				const TArray< TSharedPtr< SLevelViewport > >* LevelViewports = ViewportTab->GetViewports();
+				const TMap< FName, TSharedPtr< IViewportLayoutEntity > >* LevelViewports = ViewportTab->GetViewports();
 
 				if (LevelViewports != NULL)
 				{
-					for( int32 ViewportIndex = 0; ViewportIndex < LevelViewports->Num(); ++ViewportIndex )
+					for(auto& Pair : *LevelViewports)
 					{
-						const TSharedPtr< SLevelViewport >& Viewport = (*LevelViewports)[ ViewportIndex ];
-
-						if( Viewport->IsInForegroundTab() )
+						TSharedPtr<SLevelViewport> Viewport = Pair.Value->AsLevelViewport();
+						if( Viewport.IsValid() && Viewport->IsInForegroundTab() )
 						{
 							if( &Viewport->GetLevelViewportClient() == GCurrentLevelEditingViewportClient )
 							{
@@ -385,15 +383,14 @@ TSharedPtr<FLevelViewportTabContent> SLevelEditor::GetActiveViewportTab()
 			// Only check the viewports in the tab if its visible
 			if( ViewportTab->IsVisible() )
 			{
-				const TArray< TSharedPtr< SLevelViewport > >* LevelViewports = ViewportTab->GetViewports();
+				const TMap< FName, TSharedPtr< IViewportLayoutEntity > >* LevelViewports = ViewportTab->GetViewports();
 
 				if (LevelViewports != NULL)
 				{
-					for( int32 ViewportIndex = 0; ViewportIndex < LevelViewports->Num(); ++ViewportIndex )
+					for (auto& Pair : *LevelViewports)
 					{
-						const TSharedPtr< SLevelViewport >& Viewport = (*LevelViewports)[ ViewportIndex ];
-
-						if( Viewport->IsVisible() )
+						TSharedPtr<SLevelViewport> Viewport = Pair.Value->AsLevelViewport();
+						if( Viewport.IsValid() && Viewport->IsVisible() )
 						{
 							if( &Viewport->GetLevelViewportClient() == GCurrentLevelEditingViewportClient )
 							{
@@ -533,10 +530,10 @@ void SLevelEditor::AttachSequencer( TSharedPtr<SWidget> SequencerWidget, TShared
 			SequencerAssetEditor.Pin()->CloseWindow();
 		}
 
+		TSharedRef<SDockTab> Tab = InvokeTab("Sequencer");
+
 		if(SequencerWidget.IsValid() && NewSequencerAssetEditor.IsValid())
 		{
-			TSharedRef<SDockTab> Tab = InvokeTab("Sequencer");
-
 			Tab->SetOnTabClosed(SDockTab::FOnTabClosedCallback::CreateStatic(&Local::OnSequencerClosed, TWeakPtr<IAssetEditorInstance>(NewSequencerAssetEditor)));
 			Tab->SetContent(SequencerWidget.ToSharedRef());
 
@@ -544,6 +541,7 @@ void SLevelEditor::AttachSequencer( TSharedPtr<SWidget> SequencerWidget, TShared
 		}
 		else
 		{
+			Tab->SetContent(SNullWidget::NullWidget);
 			SequencerAssetEditor.Reset();
 		}
 	}
@@ -891,12 +889,18 @@ void SLevelEditor::OnViewportTabClosed(TSharedRef<SDockTab> ClosedTab)
 
 void SLevelEditor::SaveViewportTabInfo(TSharedRef<const FLevelViewportTabContent> ViewportTabContent)
 {
-	const TArray<TSharedPtr<SLevelViewport>>* const Viewports = ViewportTabContent->GetViewports();
+	const TMap<FName, TSharedPtr<IViewportLayoutEntity>>* const Viewports = ViewportTabContent->GetViewports();
 	if(Viewports)
 	{
 		const FString& LayoutId = ViewportTabContent->GetLayoutString();
-		for(const auto& Viewport : *Viewports)
+		for (auto& Pair : *Viewports)
 		{
+			TSharedPtr<SLevelViewport> Viewport = Pair.Value->AsLevelViewport();
+			if( !Viewport.IsValid() )
+			{
+				continue;
+			}
+
 			//@todo there could potentially be more than one of the same viewport type.  This effectively takes the last one of a specific type
 			const FLevelEditorViewportClient& LevelViewportClient = Viewport->GetLevelViewportClient();
 			const FString Key = FString::Printf(TEXT("%s[%d]"), *LayoutId, static_cast<int32>(LevelViewportClient.ViewportType));
@@ -913,12 +917,18 @@ void SLevelEditor::SaveViewportTabInfo(TSharedRef<const FLevelViewportTabContent
 
 void SLevelEditor::RestoreViewportTabInfo(TSharedRef<FLevelViewportTabContent> ViewportTabContent) const
 {
-	const TArray<TSharedPtr<SLevelViewport>>* const Viewports = ViewportTabContent->GetViewports();
+	const TMap<FName, TSharedPtr<IViewportLayoutEntity>>* const Viewports = ViewportTabContent->GetViewports();
 	if(Viewports)
 	{
 		const FString& LayoutId = ViewportTabContent->GetLayoutString();
-		for(const auto& Viewport : *Viewports)
+		for (auto& Pair : *Viewports)
 		{
+			TSharedPtr<SLevelViewport> Viewport = Pair.Value->AsLevelViewport();
+			if( !Viewport.IsValid() )
+			{
+				continue;
+			}
+
 			FLevelEditorViewportClient& LevelViewportClient = Viewport->GetLevelViewportClient();
 			bool bInitializedOrthoViewport = false;
 			for (int32 ViewportType = 0; ViewportType < LVT_MAX; ViewportType++)
@@ -1401,15 +1411,17 @@ TArray< TSharedPtr< ILevelViewport > > SLevelEditor::GetViewports() const
 		
 		if (ViewportTab.IsValid())
 		{
-			const TArray< TSharedPtr< SLevelViewport > >* LevelViewports = ViewportTab->GetViewports();
+			const TMap< FName, TSharedPtr< IViewportLayoutEntity > >* LevelViewports = ViewportTab->GetViewports();
 
 			if (LevelViewports != NULL)
 			{
-				for( int32 ViewportIndex = 0; ViewportIndex < LevelViewports->Num(); ++ViewportIndex )
+				for (auto& Pair : *LevelViewports)
 				{
-					const TSharedPtr< SLevelViewport >& Viewport = (*LevelViewports)[ ViewportIndex ];
-
-					OutViewports.Add(Viewport);
+					TSharedPtr<SLevelViewport> Viewport = Pair.Value->AsLevelViewport();
+					if( Viewport.IsValid() )
+					{
+						OutViewports.Add(Viewport);
+					}
 				}
 			}
 		}
diff --git a/Engine/Source/Editor/LevelEditor/Private/SLevelViewport.cpp b/Engine/Source/Editor/LevelEditor/Private/SLevelViewport.cpp
--- a/Engine/Source/Editor/LevelEditor/Private/SLevelViewport.cpp
+++ b/Engine/Source/Editor/LevelEditor/Private/SLevelViewport.cpp
@@ -124,7 +124,11 @@ bool SLevelViewport::IsVisible() const
 
 bool SLevelViewport::IsInForegroundTab() const
 {
-	return ViewportWidget.IsValid() && (!ParentLayout.IsValid() || ParentLayout.Pin()->IsLevelViewportVisible(*this));
+	if (ViewportWidget.IsValid() && ParentLayout.IsValid() && !ConfigKey.IsEmpty())
+	{
+		return ParentLayout.Pin()->IsLevelViewportVisible(*ConfigKey);
+	}
+	return false;
 }
 
 void SLevelViewport::Construct(const FArguments& InArgs)
@@ -1530,9 +1534,9 @@ EVisibility SLevelViewport::GetTransformToolbarVisibility() const
 
 bool SLevelViewport::IsMaximized() const
 {
-	if( ParentLayout.IsValid() )
+	if( ParentLayout.IsValid() && !ConfigKey.IsEmpty() )
 	{
-		return ParentLayout.Pin()->IsViewportMaximized( *this );
+		return ParentLayout.Pin()->IsViewportMaximized( *ConfigKey );
 	}
 
 	// Assume the viewport is always maximized if we have no layout for some reason
@@ -1609,15 +1613,19 @@ void SLevelViewport::OnToggleImmersive()
 		// We always want to animate in response to user-interactive toggling of maximized state
 		const bool bAllowAnimation = true;
 
-		ParentLayout.Pin()->RequestMaximizeViewport( SharedThis( this ), bWantMaximize, bWantImmersive, bAllowAnimation );
+		FName ViewportName = *ConfigKey;
+		if (!ViewportName.IsNone())
+		{
+			ParentLayout.Pin()->RequestMaximizeViewport( ViewportName, bWantMaximize, bWantImmersive, bAllowAnimation );
+		}
 	}
 }
 
 bool SLevelViewport::IsImmersive() const
 {
-	if( ParentLayout.IsValid() )
+	if( ParentLayout.IsValid() && !ConfigKey.IsEmpty() )
 	{
-		return ParentLayout.Pin()->IsViewportImmersive( *this );
+		return ParentLayout.Pin()->IsViewportImmersive( *ConfigKey );
 	}
 
 	// Assume the viewport is not immersive if we have no layout for some reason
@@ -2395,7 +2403,12 @@ FReply SLevelViewport::OnToggleMaximize()
 		// We always want to animate in response to user-interactive toggling of maximized state
 		const bool bAllowAnimation = true;
 
-		ParentLayoutPinned->RequestMaximizeViewport(SharedThis(this), bWantMaximize, bWantImmersive, bAllowAnimation);
+
+		FName ViewportName = *ConfigKey;
+		if (!ViewportName.IsNone())
+		{
+			ParentLayout.Pin()->RequestMaximizeViewport( ViewportName, bWantMaximize, bWantImmersive, bAllowAnimation );
+		}
 	}
 	return FReply::Handled();
 }
@@ -2406,7 +2419,12 @@ void SLevelViewport::MakeImmersive( const bool bWantImmersive, const bool bAllow
 	if( ensure( ParentLayout.IsValid() ) ) 
 	{
 		const bool bWantMaximize = IsMaximized();
-		ParentLayout.Pin()->RequestMaximizeViewport( SharedThis( this ), bWantMaximize, bWantImmersive, bAllowAnimation );
+
+		FName ViewportName = *ConfigKey;
+		if (!ViewportName.IsNone())
+		{
+			ParentLayout.Pin()->RequestMaximizeViewport( ViewportName, bWantMaximize, bWantImmersive, bAllowAnimation );
+		}
 	}
 }
 
@@ -3274,9 +3292,9 @@ void SLevelViewport::OnSetViewportConfiguration(FName ConfigurationName)
 		TSharedPtr<FLevelViewportTabContent> ViewportTabPinned = LayoutPinned->GetParentTabContent().Pin();
 		if (ViewportTabPinned.IsValid())
 		{
-			ViewportTabPinned->SetViewportConfiguration(ConfigurationName);
 			// Viewport clients are going away.  Any current one is invalid.
 			GCurrentLevelEditingViewportClient = nullptr;
+			ViewportTabPinned->SetViewportConfiguration(ConfigurationName);
 			FSlateApplication::Get().DismissAllMenus();
 		}
 	}
diff --git a/Engine/Source/Editor/LevelEditor/Private/SLevelViewportToolBar.cpp b/Engine/Source/Editor/LevelEditor/Private/SLevelViewportToolBar.cpp
--- a/Engine/Source/Editor/LevelEditor/Private/SLevelViewportToolBar.cpp
+++ b/Engine/Source/Editor/LevelEditor/Private/SLevelViewportToolBar.cpp
@@ -851,6 +851,44 @@ void SLevelViewportToolBar::GenerateViewportConfigsMenu(FMenuBuilder& MenuBuilde
 			);
 	}
 	MenuBuilder.EndSection();
+
+	MenuBuilder.BeginSection("LevelViewportSpecializedConfigs", LOCTEXT("SpecializedConfigHeader", "Specialized") );
+	{
+		FToolBarBuilder SpecializedButtons(CommandList, FMultiBoxCustomization::None);
+		SpecializedButtons.SetLabelVisibility(EVisibility::Collapsed);
+		SpecializedButtons.SetStyle(&FEditorStyle::Get(), "ViewportLayoutToolbar");
+
+
+		FLevelEditorModule& LevelEditorModule = FModuleManager::Get().GetModuleChecked<FLevelEditorModule>("LevelEditor");
+		LevelEditorModule.IterateCustomViewportLayouts([&](FName InName, const FCustomViewportLayoutDefinition& Definition){
+
+			SpecializedButtons.AddToolBarButton(
+				FExecuteAction::CreateSP(Viewport.Pin().Get(), &SLevelViewport::OnSetViewportConfiguration, InName ),
+				NAME_None,
+				Definition.DisplayName,
+				Definition.Description,
+				Definition.Icon
+			);
+
+		});
+
+		MenuBuilder.AddWidget(
+			SNew(SHorizontalBox)
+			+SHorizontalBox::Slot()
+			.AutoWidth()
+			[
+				SpecializedButtons.MakeWidget()
+			]
+			+SHorizontalBox::Slot()
+			.FillWidth(1)
+			[
+				SNullWidget::NullWidget
+			],
+			FText::GetEmpty(),
+			true
+		);
+	}
+	MenuBuilder.EndSection();
 }
 
 
diff --git a/Engine/Source/Editor/LevelEditor/Public/CustomViewportLayout.h b/Engine/Source/Editor/LevelEditor/Public/CustomViewportLayout.h
new file mode 100644
--- /dev/null
+++ b/Engine/Source/Editor/LevelEditor/Public/CustomViewportLayout.h
@@ -0,0 +1,33 @@
+// Copyright 1998-2015 Epic Games, Inc. All Rights Reserved.
+
+#pragma once
+
+
+class FLevelViewportLayout;
+
+struct FCustomViewportLayoutDefinition
+{
+	template<typename T>
+	static FCustomViewportLayoutDefinition FromType()
+	{
+		return FCustomViewportLayoutDefinition([]() -> TSharedRef<FLevelViewportLayout> {
+			return MakeShareable(new T);
+		});
+	}
+
+	FCustomViewportLayoutDefinition(TFunction<TSharedRef<FLevelViewportLayout>()> InFactoryFunction)
+		: FactoryFunction(MoveTemp(InFactoryFunction))
+	{}
+
+	/** User provided display name of the layout */
+	FText DisplayName;
+
+	/** User provided description of the layout */
+	FText Description;
+
+	/** Icon used to represent this viewport layout */
+	FSlateIcon Icon;
+
+	/** Function used to create a new instance of the viewport layout */
+	TFunction<TSharedRef<FLevelViewportLayout>()> FactoryFunction;
+};
diff --git a/Engine/Source/Editor/LevelEditor/Public/LevelEditor.h b/Engine/Source/Editor/LevelEditor/Public/LevelEditor.h
--- a/Engine/Source/Editor/LevelEditor/Public/LevelEditor.h
+++ b/Engine/Source/Editor/LevelEditor/Public/LevelEditor.h
@@ -8,6 +8,7 @@
 #include "ModuleInterface.h"
 #include "SEditorViewport.h"
 #include "Toolkits/AssetEditorToolkit.h" // For FExtensibilityManager
+#include "CustomViewportLayout.h"
 
 extern const FName LevelEditorApp;
 
@@ -223,6 +224,45 @@ public:
 	DECLARE_EVENT( FLevelEditorModule, FTakeHighResScreenShotsEvent );
 	virtual FTakeHighResScreenShotsEvent& OnTakeHighResScreenShots() { return TakeHighResScreenShotsEvent; }
 
+	/** Delegate used to capture skeltal meshes to single-frame animations when 'keeping simulation changes' */
+	DECLARE_DELEGATE_RetVal_OneParam(UAnimSequence*, FCaptureSingleFrameAnimSequence, USkeletalMeshComponent* /*Component*/);
+	virtual FCaptureSingleFrameAnimSequence& OnCaptureSingleFrameAnimSequence() { return CaptureSingleFrameAnimSequenceDelegate; }
+
+public:
+	
+	/** Register a custom viewport layout for the level editor */
+	void RegisterCustomViewportLayout(FName InLayoutName, const FCustomViewportLayoutDefinition& InDefinition)
+	{
+		CustomViewportLayouts.Add(InLayoutName, InDefinition);
+	}
+
+	/** Unregister a previously registered custom viewport layout */
+	void UnRegisterCustomViewportLayout(FName InLayoutName)
+	{
+		CustomViewportLayouts.Remove(InLayoutName);
+	}
+
+	/** Iterate all registered custom viewport layouts */
+	void IterateCustomViewportLayouts(TFunctionRef<void(FName, const FCustomViewportLayoutDefinition&)> Iter)
+	{
+		for (auto& Pair : CustomViewportLayouts)
+		{
+			Iter(Pair.Key, Pair.Value);
+		}
+	}
+
+	/** Create an instanceof a custom viewport layout from the specified layout name */
+	TSharedPtr<FLevelViewportLayout> FactoryCustomViewportLayout(FName InLayoutName) const
+	{
+		const FCustomViewportLayoutDefinition* Definition = CustomViewportLayouts.Find(InLayoutName);
+		if (Definition)
+		{
+			return Definition->FactoryFunction();
+		}
+		
+		return nullptr;
+	}
+
 private:
 	/**
 	 * Binds all global level editor commands to delegates
@@ -272,6 +312,9 @@ private:
 	/** Multicast delegate executed when a map is changed (loaded,saved,new map, etc) */
 	FMapChangedEvent MapChangedEvent;
 
+	/** Delegate used to capture skeltal meshes to single-frame animations when 'keeping simulation changes' */
+	FCaptureSingleFrameAnimSequence CaptureSingleFrameAnimSequenceDelegate;
+
 	/** All extender delegates for the level viewport menus */
 	TArray<FLevelViewportMenuExtender_SelectedObjects> LevelViewportDragDropContextMenuExtenders;
 	TArray<FLevelViewportMenuExtender_SelectedActors> LevelViewportContextMenuExtenders;
@@ -292,6 +335,9 @@ private:
 
 	/* Holds the Editor's tab manager */
 	TSharedPtr<FTabManager> LevelEditorTabManager;
+
+	/** Map of named viewport layouts to factory functions */
+	TMap<FName, FCustomViewportLayoutDefinition> CustomViewportLayouts;
 };
 
 #endif // __LevelEditor_h__
diff --git a/Engine/Source/Editor/LevelEditor/Public/LevelViewportLayout.h b/Engine/Source/Editor/LevelEditor/Public/LevelViewportLayout.h
new file mode 100644
--- /dev/null
+++ b/Engine/Source/Editor/LevelEditor/Public/LevelViewportLayout.h
@@ -0,0 +1,332 @@
+// Copyright 1998-2016 Epic Games, Inc. All Rights Reserved.
+
+
+#ifndef __LevelViewportLayout_h__
+#define __LevelViewportLayout_h__
+
+#pragma once
+
+#include "LevelEditor.h"
+#include "LevelViewportActions.h"
+
+class SLevelViewport;
+
+/** Interface that defines an entity within a viewport layout */
+class IViewportLayoutEntity : public TSharedFromThis<IViewportLayoutEntity>
+{
+public:
+	/** Virtual destruction */
+	virtual ~IViewportLayoutEntity() {};
+
+	/** Return a widget that is represents this entity */
+	virtual TSharedRef<SWidget> AsWidget() const = 0;
+
+	/** Optionally return this entity as an SLevelViewport. Legacy function for interop with code that used to deal directly with SLevelViewports */
+	virtual TSharedPtr<SLevelViewport> AsLevelViewport() const { return nullptr; }
+
+	/** Check if this entity has an active play in editor viewport */
+	virtual bool IsPlayInEditorViewportActive() const = 0;
+
+	/** Register this viewport layout entity as a game viewport, if it's currently PIE-ing */
+	virtual void RegisterGameViewportIfPIE() = 0;
+
+	/** Set keyboard focus to this viewport entity */
+	virtual void SetKeyboardFocus() = 0;
+
+	/** Called when the parent layout is being destroyed */
+	virtual void OnLayoutDestroyed() = 0;
+
+	/** Called to save this item's settings in the specified config section */
+	virtual void SaveConfig(const FString& ConfigSection) = 0;
+};
+
+class LEVELEDITOR_API FLevelViewportLayoutEntity : public IViewportLayoutEntity
+{
+public:
+
+	FLevelViewportLayoutEntity(TSharedRef<SLevelViewport> InLevelViewport);
+	virtual TSharedRef<SWidget> AsWidget() const override;
+	virtual TSharedPtr<SLevelViewport> AsLevelViewport() const override;
+
+	bool IsPlayInEditorViewportActive() const;
+	void RegisterGameViewportIfPIE();
+	void SetKeyboardFocus();
+	void OnLayoutDestroyed();
+	void SaveConfig(const FString& ConfigSection);
+
+protected:
+
+	/** This entity's level viewport */
+	TSharedRef<SLevelViewport> LevelViewport;
+};
+
+/**
+ * Base class for level viewport layout configurations
+ * Handles maximizing and restoring well as visibility of specific viewports.
+ */
+class LEVELEDITOR_API FLevelViewportLayout : public TSharedFromThis<FLevelViewportLayout>, public FTickableEditorObject
+{
+public:
+	/**
+	 * Constructor
+	 */
+	FLevelViewportLayout();
+
+	/**
+	 * Destructor
+	 */
+	virtual ~FLevelViewportLayout();
+
+	/**
+	 * Builds a viewport layout and returns the widget containing the layout
+	 * 
+	 * @param InParentDockTab		The parent dock tab widget of this viewport configuration
+	 * @param InParentTab			The parent tab object
+	 * @param LayoutString			The layout string loaded from file to custom build the layout with
+	 * @param InParentLevelEditor	Optional level editor parent to use for new viewports
+	 */
+	TSharedRef<SWidget> BuildViewportLayout( TSharedPtr<SDockTab> InParentDockTab, TSharedPtr<class FLevelViewportTabContent> InParentTab, const FString& LayoutString, TWeakPtr<class SLevelEditor> InParentLevelEditor = NULL );
+
+	/**
+	 * Makes a request to maximize a specific viewport and hide the others in this layout
+	 * 
+	 * @param	ViewportToMaximize	The viewport that should be maximized
+	 * @param	bWantMaximize		True to maximize or false to "restore"
+	 * @param	bWantImmersive		True to perform an "immersive" maximize, which transitions the viewport to fill the entire application window
+	 * @param	bAllowAnimation		True if an animated transition should be used
+	 */
+	void RequestMaximizeViewport( FName ViewportToMaximize, const bool bWantMaximize, const bool bWantImmersive, const bool bAllowAnimation = true );
+
+	/**
+	 * @return true if this layout is visible.  It is not visible if its parent tab is not active
+	 */
+	bool IsVisible() const;
+
+	/**
+	 * Checks to see the specified level viewport is visible in this layout
+	 * A viewport is visible in a layout if the layout is visible and the viewport is the maximized viewport or there is no maximized viewport
+	 *
+	 * @param InViewport	The viewport within this layout that should be checked
+	 * @return true if the viewport is visible.  
+	 */
+	bool IsLevelViewportVisible( FName InViewport ) const;
+
+	/**
+	* Checks to see if the specified level viewport supports maximizing one pane
+	*
+	* @return true if the viewport supports maximizing
+	*/
+	bool IsMaximizeSupported() const { return bIsMaximizeSupported; }
+
+	/** 
+	 * Checks to see if the specified level viewport is maximized
+	 *
+	 * @param InViewport	The viewport to check
+	 * @return true if the viewport is maximized, false otherwise
+	 */
+	bool IsViewportMaximized( FName InViewport ) const;
+
+	/** 
+	 * Checks to see if the specified level viewport is in immersive mode
+	 *
+	 * @param InViewport	The viewport to check
+	 * @return true if the viewport is immersive, false otherwise
+	 */
+	bool IsViewportImmersive( FName InViewport ) const;
+
+	/**
+	 * @return All the viewports in this configuration
+	 */
+	const TMap< FName, TSharedPtr<IViewportLayoutEntity> >& GetViewports() const { return Viewports; }
+
+	/**
+	 * Saves viewport layout information between editor sessions
+	 */
+	virtual void SaveLayoutString(const FString& LayoutString) const = 0;
+
+	virtual const FName& GetLayoutTypeName() const = 0;
+
+
+	/** FTickableEditorObject interface */
+	virtual void Tick( float DeltaTime ) override;
+	virtual bool IsTickable() const override;
+	virtual TStatId GetStatId() const override { return TStatId(); }
+
+
+	/** Tells this layout whether it was the intial layout or replaced an existing one when the user switched layouts */
+	void SetIsReplacement(bool bInIsReplacement) { bIsReplacement = bInIsReplacement; }
+
+	/** Returns the parent tab content object */
+	TWeakPtr< class FLevelViewportTabContent > GetParentTabContent() const { return ParentTabContent; }
+
+	/** Returns whether a viewport animation is currently taking place */
+	bool IsTransitioning() const { return bIsTransitioning; }
+
+protected:
+
+	/**
+	 * Maximizes a specific viewport and hides the others in this layout
+	 * 
+	 * @param	ViewportToMaximize	The viewport that should be maximized
+	 * @param	bWantMaximize		True to maximize or false to "restore"
+	 * @param	bWantImmersive		True to perform an "immersive" maximize, which transitions the viewport to fill the entire application window
+	 * @param	bAllowAnimation		True if an animated transition should be used
+	 */
+	void MaximizeViewport( FName ViewportToMaximize, const bool bWantMaximize, const bool bWantImmersive, const bool bAllowAnimation );
+
+	/**
+	 * Overridden in derived classes to set up custom layouts  
+	 *
+	 * @param LayoutString		The layout string loaded from a file
+	 * @return The base widget representing the layout.  Usually a splitter
+	 */
+	virtual TSharedRef<SWidget> MakeViewportLayout(const FString& LayoutString) = 0;
+
+	/**
+	 * Delegate called to get the visibility of the non-maximized viewports
+	 * The non-maximized viewports are not visible if there is a maximized viewport on top of them
+	 *
+	 * @param EVisibility::Visible when visible, EVisibility::Collapsed otherwise
+	 */
+	EVisibility OnGetNonMaximizedVisibility() const;
+
+	/**
+	 * Returns the widget position for viewport transition animations
+	 *
+	 * @return	Viewport position on canvas
+	 */
+	FVector2D GetMaximizedViewportPositionOnCanvas() const;
+
+	/**
+	 * Returns the widget size for viewport transition animations
+	 *
+	 * @return	Viewport size on canvas
+	 */
+	FVector2D GetMaximizedViewportSizeOnCanvas() const;
+
+	/**
+	 * Inline replaces a viewport content widget within this layout
+	 *
+	 * @param	Source	The widget to replace
+	 * @param	Replacement	The widget to replace the source widget with
+	 */
+	virtual void ReplaceWidget( TSharedRef< SWidget > Source, TSharedRef< SWidget > Replacement ) = 0;
+
+	/** If a viewport animation is in progress, finishes that transition immediately */
+	void FinishMaximizeTransition();
+
+	/**
+	 * Begins a draw throttle for responsiveness when animating a viewports size
+	 */
+	void BeginThrottleForAnimatedResize();
+
+	/**
+	 * Ends a draw throttle for responsiveness when animating a viewports size
+	 */
+	void EndThrottleForAnimatedResize();
+
+	/** Generates a layout string for persisting settings for this layout based on the runtime type of layout */
+	FString GetTypeSpecificLayoutString(const FString& LayoutString) const
+	{ 
+		if (LayoutString.IsEmpty())
+		{
+			return LayoutString;
+		}
+		return FString::Printf(TEXT("%s.%s"), *GetLayoutTypeName().ToString(), *LayoutString);
+	}
+
+	/** Called in MakeViewportLayout() functions for derived types of layout to init values common to all types */
+	void InitCommonLayoutFromString( const FString& SpecificLayoutString );
+
+	/** Called in SaveLayoutString() functions for derived types of layout to save values common to all types */
+	void SaveCommonLayoutString( const FString& SpecificLayoutString ) const;
+
+protected:
+	/** True if we've started an animation and are waiting for it to finish */
+	bool bIsTransitioning;
+
+	/** Curve for animating from a "restored" state to a maximized state */
+	FCurveSequence MaximizeAnimation;
+
+	/** List of all of the viewports in this layout */
+	TMap< FName, TSharedPtr< IViewportLayoutEntity > > Viewports;
+	
+	/** The parent tab where this layout resides */
+	TWeakPtr< SDockTab > ParentTab;
+
+	/** The parent tab content object where this layout resides */
+	TWeakPtr< class FLevelViewportTabContent > ParentTabContent;
+
+	/** The optional parent level editor for this layout */
+	TWeakPtr< class SLevelEditor > ParentLevelEditor;
+
+	/** The current maximized viewport if any */
+	FName MaximizedViewport;
+
+	/** True if the user selected this layout, false if it's the initial layout loaded */
+	bool bIsReplacement;
+
+	/** Temporarily set to true while we are querying layout metrics and want all widgets to be visible */
+	bool bIsQueryingLayoutMetrics;
+
+	/** True if the layout supports maximizing one viewport, false if the feature is disabled  */
+	bool bIsMaximizeSupported;
+
+	/** True if we're currently maximized */
+	bool bIsMaximized;
+
+	/** True if we're currently in immersive mode */
+	bool bIsImmersive;
+
+	/** True when transitioning from a maximized state */
+	bool bWasMaximized;
+
+	/** True when transitioning from an immersive state */
+	bool bWasImmersive;
+
+	/** Window-space start position of the viewport that's currently being maximized */
+	FVector2D MaximizedViewportStartPosition;
+
+	/** Window-space start size of the viewport that's currently being maximized */
+	FVector2D MaximizedViewportStartSize;
+
+	/** The overlay widget that handles what viewports should be on top (non-maximized or maximized) */
+	TWeakPtr< class SViewportsOverlay > ViewportsOverlayPtr;
+
+	/** When maximizing viewports (or making them immersive), this stores the widget we create to wrap the viewport */
+	TSharedPtr< SWidget > ViewportsOverlayWidget;
+
+	/** Dummy widget that we'll inline-replace viewport widgets with while a view is maximized (or made immersive) */
+	TSharedPtr< SWidget > ViewportReplacementWidget;
+	
+	/** Caches the window that our widgets are contained within */
+	TWeakPtr< SWindow > CachedOwnerWindow;
+
+	/** Viewport resize draw throttle request */
+	FThrottleRequest ViewportResizeThrottleRequest;
+
+	/**
+	 * Maximize/immersive commands can be queued up at startup to be executed on the first tick.
+	 * This is necessary, because these commands can't be executed until the viewport has a parent window,
+	 * which might not be there upon viewport initialization
+	 */
+	struct FMaximizeViewportCommand
+	{
+		FMaximizeViewportCommand(FName InViewport, bool bInMaximize, bool bInImmersive, bool bInToggle=true, bool bInAllowAnimation=true)
+			: Viewport(InViewport)
+			, bMaximize(bInMaximize)
+			, bImmersive(bInImmersive) 
+			, bToggle(bInToggle)
+			, bAllowAnimation(bInAllowAnimation)
+		{}
+
+		FName Viewport;
+		bool bMaximize;
+		bool bImmersive;
+		bool bToggle;
+		bool bAllowAnimation;
+	};
+	TArray<FMaximizeViewportCommand> DeferredMaximizeCommands;
+};
+
+#endif
diff --git a/Engine/Source/Editor/LevelEditor/Public/LevelViewportTabContent.h b/Engine/Source/Editor/LevelEditor/Public/LevelViewportTabContent.h
new file mode 100644
--- /dev/null
+++ b/Engine/Source/Editor/LevelEditor/Public/LevelViewportTabContent.h
@@ -0,0 +1,81 @@
+// Copyright 1998-2016 Epic Games, Inc. All Rights Reserved.
+
+
+#ifndef __LevelViewportTabContent_h__
+#define __LevelViewportTabContent_h__
+
+#pragma once
+
+#include "LevelEditor.h"
+#include "LevelEditorActions.h"
+
+class IViewportLayoutEntity;
+
+/**
+ * Represents the content in a viewport tab in the level editor.
+ * Each SDockTab holding viewports in the level editor contains and owns one of these.
+ */
+class LEVELEDITOR_API FLevelViewportTabContent : public TSharedFromThis<FLevelViewportTabContent>
+{
+public:
+	/** Starts the tab content object and creates the initial layout based on the layout string */
+	void Initialize(TSharedPtr<class SLevelEditor> InParentLevelEditor, TSharedPtr<SDockTab> InParentTab, const FString& InLayoutString);
+
+	/** Returns whether the tab is currently shown */
+	bool IsVisible() const;
+
+	/** Returns an array of viewports inside this tab */
+	const TMap< FName, TSharedPtr< IViewportLayoutEntity > >* GetViewports() const;
+
+	/**
+	 * Sets the current layout by changing the contained layout object
+	 * 
+	 * @param ConfigurationName		The name of the layout (for the names in namespace LevelViewportConfigurationNames)
+	 */
+	void SetViewportConfiguration(const FName& ConfigurationName);
+
+	/**
+	 * Returns whether the named layout is currently selected
+	 *
+	 * @param ConfigurationName		The name of the layout (for the names in namespace LevelViewportConfigurationNames)
+	 * @return						True, if the named layout is currently active
+	 */
+	bool IsViewportConfigurationSet(const FName& ConfigurationName) const;
+
+	/** @return True if this viewport belongs to the tab given */
+	bool BelongsToTab(TSharedRef<class SDockTab> InParentTab) const;
+
+	/** @return The string used to identify the layout of this viewport tab */
+	const FString& GetLayoutString() const
+	{
+		return LayoutString;
+	}
+
+private:
+	/**
+	 * Updates state after the layout is changed by SetViewportConfiguration()
+	 */
+	void UpdateViewportTabWidget();
+
+	/** Callback registered with the parent dock tab to persist state when needed */
+	void SaveLayoutString(const FString LayoutString) const;
+
+	/**
+	 * Translates layout names from namespace LevelViewportConfigurationNames into layout class and returns a new object of that class
+	 *
+	 * @param TypeName	FName of the layout type from those in LevelViewportConfigurationNames
+	 * @param bSwitchingLayouts	If true, the new layout is replacing an existing one.
+	 * @return	A new layout of the requested type.
+	 */
+	static TSharedPtr< class FLevelViewportLayout > ConstructViewportLayoutByTypeName(const FName& TypeName, bool bSwitchingLayouts);
+
+private:
+	TWeakPtr<class SDockTab> ParentTab;
+	TWeakPtr<class SLevelEditor> ParentLevelEditor;
+	FString LayoutString;
+
+	/** Current layout */
+	TSharedPtr< class FLevelViewportLayout > ActiveLevelViewportLayout;
+};
+
+#endif
diff --git a/Engine/Source/Editor/LevelEditor/Public/SLevelViewport.h b/Engine/Source/Editor/LevelEditor/Public/SLevelViewport.h
--- a/Engine/Source/Editor/LevelEditor/Public/SLevelViewport.h
+++ b/Engine/Source/Editor/LevelEditor/Public/SLevelViewport.h
@@ -12,7 +12,7 @@ class SGameLayerManager;
 /**
  * Encapsulates an SViewport and an SLevelViewportToolBar
  */
-class SLevelViewport : public SEditorViewport, public ILevelViewport
+class LEVELEDITOR_API SLevelViewport : public SEditorViewport, public ILevelViewport
 {
 public:
 	SLATE_BEGIN_ARGS( SLevelViewport )
diff --git a/Engine/Source/Editor/MovieSceneCaptureDialog/Private/MovieSceneCaptureDialogModule.cpp b/Engine/Source/Editor/MovieSceneCaptureDialog/Private/MovieSceneCaptureDialogModule.cpp
--- a/Engine/Source/Editor/MovieSceneCaptureDialog/Private/MovieSceneCaptureDialogModule.cpp
+++ b/Engine/Source/Editor/MovieSceneCaptureDialog/Private/MovieSceneCaptureDialogModule.cpp
@@ -661,9 +661,14 @@ class FMovieSceneCaptureDialogModule : public IMovieSceneCaptureDialogModule
 		FString GamePath = FPlatformProcess::GenerateApplicationPath(FApp::GetName(), FApp::GetBuildConfiguration());
 		FProcHandle ProcessHandle = FPlatformProcess::CreateProc(*GamePath, *Params, true, false, false, nullptr, 0, nullptr, nullptr);
 
-		// @todo: progress reporting, UI feedback
 		if (ProcessHandle.IsValid())
 		{
+			if (CaptureObject->bCloseEditorWhenCaptureStarts)
+			{
+				FPlatformMisc::RequestExit(false);
+				return FText();
+			}
+
 			TSharedRef<FProcHandle> SharedProcHandle = MakeShareable(new FProcHandle(ProcessHandle));
 			auto GetCaptureStatus = [=]{
 				if (!FPlatformProcess::IsProcRunning(*SharedProcHandle))
diff --git a/Engine/Source/Editor/MovieSceneTools/MovieSceneTools.Build.cs b/Engine/Source/Editor/MovieSceneTools/MovieSceneTools.Build.cs
--- a/Engine/Source/Editor/MovieSceneTools/MovieSceneTools.Build.cs
+++ b/Engine/Source/Editor/MovieSceneTools/MovieSceneTools.Build.cs
@@ -12,7 +12,7 @@ public class MovieSceneTools : ModuleRules
                 "Editor/MovieSceneTools/Private/CurveKeyEditors",
                 "Editor/MovieSceneTools/Private/TrackEditors",
 				"Editor/MovieSceneTools/Private/TrackEditors/PropertyTrackEditors",
-                "Editor/MovieSceneTools/Private/TrackEditors/ShotTrackEditor",
+                "Editor/MovieSceneTools/Private/TrackEditorThumbnail",
 				"Editor/MovieSceneTools/Private/Sections"
             }
         );
@@ -31,6 +31,7 @@ public class MovieSceneTools : ModuleRules
 		PrivateDependencyModuleNames.AddRange(
 			new string[] {
                 "ActorPickerMode",
+				"LevelSequence",
 				"MovieScene",
 				"MovieSceneTracks",
 				"BlueprintGraph",
@@ -45,8 +46,10 @@ public class MovieSceneTools : ModuleRules
 
 		PrivateIncludePathModuleNames.AddRange(
 			new string[] {
+                "AssetRegistry",
 				"AssetTools",
 				"Sequencer",
+                "Settings",
 				"SceneOutliner",
 				"PropertyEditor"
 			}
@@ -54,6 +57,7 @@ public class MovieSceneTools : ModuleRules
 
 		DynamicallyLoadedModuleNames.AddRange(
 			new string[] {
+                "AssetRegistry",
 				"AssetTools",
 				"SceneOutliner",
 				"PropertyEditor"
diff --git a/Engine/Source/Editor/MovieSceneTools/Private/CurveKeyAreas/FloatCurveKeyArea.cpp b/Engine/Source/Editor/MovieSceneTools/Private/CurveKeyAreas/FloatCurveKeyArea.cpp
--- a/Engine/Source/Editor/MovieSceneTools/Private/CurveKeyAreas/FloatCurveKeyArea.cpp
+++ b/Engine/Source/Editor/MovieSceneTools/Private/CurveKeyAreas/FloatCurveKeyArea.cpp
@@ -133,6 +133,12 @@ ERichCurveInterpMode FFloatCurveKeyArea::GetKeyInterpMode(FKeyHandle KeyHandle)
 }
 
 
+TSharedPtr<FStructOnScope> FFloatCurveKeyArea::GetKeyStruct(FKeyHandle KeyHandle)
+{
+	return MakeShareable(new FStructOnScope(FRichCurveKey::StaticStruct(), (uint8*)&Curve->GetKey(KeyHandle)));
+}
+
+
 ERichCurveTangentMode FFloatCurveKeyArea::GetKeyTangentMode(FKeyHandle KeyHandle) const
 {
 	if (Curve->IsKeyHandleValid(KeyHandle))
@@ -193,6 +199,10 @@ void FFloatCurveKeyArea::SetExtrapolationMode(ERichCurveExtrapolation ExtrapMode
 	}
 }
 
+bool FFloatCurveKeyArea::CanSetExtrapolationMode() const
+{
+	return true;
+}
 
 void FFloatCurveKeyArea::SetKeyInterpMode(FKeyHandle KeyHandle, ERichCurveInterpMode InterpMode)
 {
diff --git a/Engine/Source/Editor/MovieSceneTools/Private/CurveKeyAreas/IntegralKeyArea.cpp b/Engine/Source/Editor/MovieSceneTools/Private/CurveKeyAreas/IntegralKeyArea.cpp
--- a/Engine/Source/Editor/MovieSceneTools/Private/CurveKeyAreas/IntegralKeyArea.cpp
+++ b/Engine/Source/Editor/MovieSceneTools/Private/CurveKeyAreas/IntegralKeyArea.cpp
@@ -59,6 +59,12 @@ ERichCurveExtrapolation FIntegralCurveKeyAreaBase::GetExtrapolationMode(bool bPr
 }
 
 
+TSharedPtr<FStructOnScope> FIntegralCurveKeyAreaBase::GetKeyStruct(FKeyHandle KeyHandle)
+{
+	return MakeShareable(new FStructOnScope(FIntegralKey::StaticStruct(), (uint8*)&Curve.GetKey(KeyHandle)));
+}
+
+
 ERichCurveTangentMode FIntegralCurveKeyAreaBase::GetKeyTangentMode(FKeyHandle KeyHandle) const
 {
 	return RCTM_None;
diff --git a/Engine/Source/Editor/MovieSceneTools/Private/CurveKeyAreas/NameCurveKeyArea.cpp b/Engine/Source/Editor/MovieSceneTools/Private/CurveKeyAreas/NameCurveKeyArea.cpp
--- a/Engine/Source/Editor/MovieSceneTools/Private/CurveKeyAreas/NameCurveKeyArea.cpp
+++ b/Engine/Source/Editor/MovieSceneTools/Private/CurveKeyAreas/NameCurveKeyArea.cpp
@@ -77,6 +77,12 @@ ERichCurveInterpMode FNameCurveKeyArea::GetKeyInterpMode(FKeyHandle Keyhandle) c
 }
 
 
+TSharedPtr<FStructOnScope> FNameCurveKeyArea::GetKeyStruct(FKeyHandle KeyHandle)
+{
+	return MakeShareable(new FStructOnScope(FNameCurveKey::StaticStruct(), (uint8*)&Curve.GetKey(KeyHandle)));
+}
+
+
 ERichCurveTangentMode FNameCurveKeyArea::GetKeyTangentMode(FKeyHandle KeyHandle) const
 {
 	return RCTM_None;
diff --git a/Engine/Source/Editor/MovieSceneTools/Private/MovieSceneToolHelpers.cpp b/Engine/Source/Editor/MovieSceneTools/Private/MovieSceneToolHelpers.cpp
--- a/Engine/Source/Editor/MovieSceneTools/Private/MovieSceneToolHelpers.cpp
+++ b/Engine/Source/Editor/MovieSceneTools/Private/MovieSceneToolHelpers.cpp
@@ -2,6 +2,9 @@
 
 #include "MovieSceneToolsPrivatePCH.h"
 #include "MovieSceneSection.h"
+#include "MovieSceneCinematicShotSection.h"
+#include "LevelSequence.h"
+#include "AssetRegistryModule.h"
 
 
 /* MovieSceneToolHelpers
@@ -29,3 +32,283 @@ void MovieSceneToolHelpers::SplitSection(const TSet<TWeakObjectPtr<UMovieSceneSe
 		}
 	}
 }
+
+bool MovieSceneToolHelpers::ParseShotName(const FString& ShotName, FString& ShotPrefix, uint32& ShotNumber, uint32& TakeNumber)
+{
+	// Parse a shot name
+	// 
+	// sht010:
+	//  ShotPrefix = sht
+	//  ShotNumber = 10
+	//  TakeNumber = 1 (default)
+	// 
+	// sp020_002
+	//  ShotPrefix = sp
+	//  ShotNumber = 20
+	//  TakeNumber = 2
+	//
+	const UMovieSceneToolsProjectSettings* ProjectSettings = GetDefault<UMovieSceneToolsProjectSettings>();
+
+	uint32 FirstShotNumberIndex = INDEX_NONE;
+	uint32 LastShotNumberIndex = INDEX_NONE;
+	bool bInShotNumber = false;
+
+	uint32 FirstTakeNumberIndex = INDEX_NONE;
+	uint32 LastTakeNumberIndex = INDEX_NONE;
+	bool bInTakeNumber = false;
+
+	bool bFoundTakeSeparator = false;
+	TakeNumber = ProjectSettings->FirstTakeNumber;
+
+	for (int32 CharIndex = 0; CharIndex < ShotName.Len(); ++CharIndex)
+	{
+		if (FChar::IsDigit(ShotName[CharIndex]))
+		{
+			// Find shot number indices
+			if (FirstShotNumberIndex == INDEX_NONE)
+			{
+				bInShotNumber = true;
+				FirstShotNumberIndex = CharIndex;
+			}
+			if (bInShotNumber)
+			{
+				LastShotNumberIndex = CharIndex;
+			}
+
+			if (FirstShotNumberIndex != INDEX_NONE && LastShotNumberIndex != INDEX_NONE)
+			{
+				if (bFoundTakeSeparator)
+				{
+					// Find take number indices
+					if (FirstTakeNumberIndex == INDEX_NONE)
+					{
+						bInTakeNumber = true;
+						FirstTakeNumberIndex = CharIndex;
+					}
+					if (bInTakeNumber)
+					{
+						LastTakeNumberIndex = CharIndex;
+					}
+				}
+			}
+		}
+
+		if (FirstShotNumberIndex != INDEX_NONE && LastShotNumberIndex != INDEX_NONE)
+		{
+			if (ShotName[CharIndex] == ProjectSettings->TakeSeparator[0])
+			{
+				bFoundTakeSeparator = true;
+			}
+		}
+	}
+
+	if (FirstShotNumberIndex != INDEX_NONE)
+	{
+		ShotPrefix = ShotName.Left(FirstShotNumberIndex);
+		ShotNumber = FCString::Atoi(*ShotName.Mid(FirstShotNumberIndex, LastShotNumberIndex-FirstShotNumberIndex+1));
+	}
+
+	if (FirstTakeNumberIndex != INDEX_NONE)
+	{
+		TakeNumber = FCString::Atoi(*ShotName.Mid(FirstTakeNumberIndex, LastTakeNumberIndex-FirstTakeNumberIndex+1));
+	}
+
+	return FirstShotNumberIndex != INDEX_NONE;
+}
+
+
+FString MovieSceneToolHelpers::ComposeShotName(const FString& ShotPrefix, uint32 ShotNumber, uint32 TakeNumber)
+{
+	const UMovieSceneToolsProjectSettings* ProjectSettings = GetDefault<UMovieSceneToolsProjectSettings>();
+
+	FString ShotName = ShotPrefix;
+
+	FString ShotFormat = TEXT("%0") + FString::Printf(TEXT("%d"), ProjectSettings->ShotNumDigits) + TEXT("d");
+	FString TakeFormat = TEXT("%0") + FString::Printf(TEXT("%d"), ProjectSettings->TakeNumDigits) + TEXT("d");
+
+	ShotName += FString::Printf(*ShotFormat, ShotNumber);
+	if (TakeNumber != INDEX_NONE)
+	{
+		ShotName += ProjectSettings->TakeSeparator;
+		ShotName += FString::Printf(*TakeFormat, TakeNumber);
+	}
+	return ShotName;
+}
+
+
+FString MovieSceneToolHelpers::GenerateNewShotName(const TArray<UMovieSceneSection*>& AllSections, float Time)
+{
+	const UMovieSceneToolsProjectSettings* ProjectSettings = GetDefault<UMovieSceneToolsProjectSettings>();
+
+	UMovieSceneCinematicShotSection* BeforeShot = nullptr;
+	UMovieSceneCinematicShotSection* NextShot = nullptr;
+
+	float MinEndDiff = FLT_MAX;
+	float MinStartDiff = FLT_MAX;
+	for (auto Section : AllSections)
+	{
+		if (Section->GetEndTime() >= Time)
+		{
+			float EndDiff = Section->GetEndTime() - Time;
+			if (MinEndDiff > EndDiff)
+			{
+				MinEndDiff = EndDiff;
+				BeforeShot = Cast<UMovieSceneCinematicShotSection>(Section);
+			}
+		}
+		if (Section->GetStartTime() <= Time)
+		{
+			float StartDiff = Time - Section->GetStartTime();
+			if (MinStartDiff > StartDiff)
+			{
+				MinStartDiff = StartDiff;
+				NextShot = Cast<UMovieSceneCinematicShotSection>(Section);
+			}
+		}
+	}
+	
+	// There aren't any shots, let's create the first shot name
+	if (BeforeShot == nullptr || NextShot == nullptr)
+	{
+		// Default case
+	}
+	// This is the last shot
+	else if (BeforeShot == NextShot)
+	{
+		FString NextShotPrefix = ProjectSettings->ShotPrefix;
+		uint32 NextShotNumber = ProjectSettings->FirstShotNumber;
+		uint32 NextTakeNumber = ProjectSettings->FirstTakeNumber;
+
+		if (ParseShotName(NextShot->GetShotDisplayName().ToString(), NextShotPrefix, NextShotNumber, NextTakeNumber))
+		{
+			uint32 NewShotNumber = NextShotNumber + ProjectSettings->ShotIncrement;
+			return ComposeShotName(NextShotPrefix, NewShotNumber, ProjectSettings->FirstTakeNumber);
+		}
+	}
+	// This is in between two shots
+	else 
+	{
+		FString BeforeShotPrefix = ProjectSettings->ShotPrefix;
+		uint32 BeforeShotNumber = ProjectSettings->FirstShotNumber;
+		uint32 BeforeTakeNumber = ProjectSettings->FirstTakeNumber;
+
+		FString NextShotPrefix = ProjectSettings->ShotPrefix;
+		uint32 NextShotNumber = ProjectSettings->FirstShotNumber;
+		uint32 NextTakeNumber = ProjectSettings->FirstTakeNumber;
+
+		if (ParseShotName(BeforeShot->GetShotDisplayName().ToString(), BeforeShotPrefix, BeforeShotNumber, BeforeTakeNumber) &&
+			ParseShotName(NextShot->GetShotDisplayName().ToString(), NextShotPrefix, NextShotNumber, NextTakeNumber))
+		{
+			uint32 NewShotNumber = BeforeShotNumber + ( (NextShotNumber - BeforeShotNumber) / 2); // what if we can't find one? or conflicts with another?
+			return ComposeShotName(BeforeShotPrefix, NewShotNumber, ProjectSettings->FirstTakeNumber);
+		}
+	}
+
+	// Default case
+	return ComposeShotName(ProjectSettings->ShotPrefix, ProjectSettings->FirstShotNumber, ProjectSettings->FirstTakeNumber);
+}
+
+void MovieSceneToolHelpers::GatherTakes(const UMovieSceneSection* Section, TArray<uint32>& TakeNumbers)
+{
+	const UMovieSceneCinematicShotSection* Shot = Cast<const UMovieSceneCinematicShotSection>(Section);
+	
+	if (Shot->GetSequence() == nullptr)
+	{
+		return;
+	}
+
+	FAssetData ShotData(Shot->GetSequence()->GetOuter());
+
+	FString ShotPackagePath = ShotData.PackagePath.ToString();
+	int32 ShotLastSlashPos = INDEX_NONE;
+	ShotPackagePath.FindLastChar(TCHAR('/'), ShotLastSlashPos);
+	ShotPackagePath = ShotPackagePath.Left(ShotLastSlashPos);
+
+	FString ShotPrefix;
+	uint32 ShotNumber = INDEX_NONE;
+	uint32 TakeNumber = INDEX_NONE;
+
+	if (ParseShotName(Shot->GetShotDisplayName().ToString(), ShotPrefix, ShotNumber, TakeNumber))
+	{
+		// Gather up all level sequence assets
+		FAssetRegistryModule& AssetRegistryModule = FModuleManager::LoadModuleChecked<FAssetRegistryModule>(TEXT("AssetRegistry"));
+		TArray<FAssetData> ObjectList;
+		AssetRegistryModule.Get().GetAssetsByClass(ULevelSequence::StaticClass()->GetFName(), ObjectList);
+
+		for (auto AssetObject : ObjectList)
+		{
+			FString AssetPackagePath = AssetObject.PackagePath.ToString();
+			int32 AssetLastSlashPos = INDEX_NONE;
+			AssetPackagePath.FindLastChar(TCHAR('/'), AssetLastSlashPos);
+			AssetPackagePath = AssetPackagePath.Left(AssetLastSlashPos);
+
+			if (AssetPackagePath == ShotPackagePath)
+			{
+				FString AssetShotPrefix;
+				uint32 AssetShotNumber = INDEX_NONE;
+				uint32 AssetTakeNumber = INDEX_NONE;
+
+				if (ParseShotName(AssetObject.AssetName.ToString(), AssetShotPrefix, AssetShotNumber, AssetTakeNumber))
+				{
+					if (AssetShotPrefix == ShotPrefix && AssetShotNumber == ShotNumber)
+					{
+						TakeNumbers.Add(AssetTakeNumber);
+					}
+				}
+			}
+		}
+	}
+
+	TakeNumbers.Sort();
+}
+
+UObject* MovieSceneToolHelpers::GetTake(const UMovieSceneSection* Section, uint32 TakeNumber)
+{
+	const UMovieSceneCinematicShotSection* Shot = Cast<const UMovieSceneCinematicShotSection>(Section);
+
+	FAssetData ShotData(Shot->GetSequence()->GetOuter());
+
+	FString ShotPackagePath = ShotData.PackagePath.ToString();
+	int32 ShotLastSlashPos = INDEX_NONE;
+	ShotPackagePath.FindLastChar(TCHAR('/'), ShotLastSlashPos);
+	ShotPackagePath = ShotPackagePath.Left(ShotLastSlashPos);
+
+	FString ShotPrefix;
+	uint32 ShotNumber = INDEX_NONE;
+	uint32 TakeNumberDummy = INDEX_NONE;
+
+	if (ParseShotName(Shot->GetShotDisplayName().ToString(), ShotPrefix, ShotNumber, TakeNumberDummy))
+	{
+		// Gather up all level sequence assets
+		FAssetRegistryModule& AssetRegistryModule = FModuleManager::LoadModuleChecked<FAssetRegistryModule>(TEXT("AssetRegistry"));
+		TArray<FAssetData> ObjectList;
+		AssetRegistryModule.Get().GetAssetsByClass(ULevelSequence::StaticClass()->GetFName(), ObjectList);
+
+		for (auto AssetObject : ObjectList)
+		{
+			FString AssetPackagePath = AssetObject.PackagePath.ToString();
+			int32 AssetLastSlashPos = INDEX_NONE;
+			AssetPackagePath.FindLastChar(TCHAR('/'), AssetLastSlashPos);
+			AssetPackagePath = AssetPackagePath.Left(AssetLastSlashPos);
+
+			if (AssetPackagePath == ShotPackagePath)
+			{
+				FString AssetShotPrefix;
+				uint32 AssetShotNumber = INDEX_NONE;
+				uint32 AssetTakeNumber = INDEX_NONE;
+
+				if (ParseShotName(AssetObject.AssetName.ToString(), AssetShotPrefix, AssetShotNumber, AssetTakeNumber))
+				{
+					if (AssetShotPrefix == ShotPrefix && AssetShotNumber == ShotNumber && TakeNumber == AssetTakeNumber)
+					{
+						return AssetObject.GetAsset();
+					}
+				}
+			}
+		}
+	}
+
+	return nullptr;
+}
+
+
diff --git a/Engine/Source/Editor/MovieSceneTools/Private/MovieSceneToolsModule.cpp b/Engine/Source/Editor/MovieSceneTools/Private/MovieSceneToolsModule.cpp
--- a/Engine/Source/Editor/MovieSceneTools/Private/MovieSceneToolsModule.cpp
+++ b/Engine/Source/Editor/MovieSceneTools/Private/MovieSceneToolsModule.cpp
@@ -18,7 +18,8 @@
 #include "VisibilityPropertyTrackEditor.h"
 
 #include "TransformTrackEditor.h"
-#include "ShotTrackEditor.h"
+#include "CameraCutTrackEditor.h"
+#include "CinematicShotTrackEditor.h"
 #include "SlomoTrackEditor.h"
 #include "SubTrackEditor.h"
 #include "AudioTrackEditor.h"
@@ -36,7 +37,9 @@
 #include "SequencerClipboardReconciler.h"
 #include "ClipboardTypes.h"
 #include "Curves/CurveBase.h"
+#include "ISettingsModule.h"
 
+#define LOCTEXT_NAMESPACE "FMovieSceneToolsModule"
 
 /**
  * Implements the MovieSceneTools module.
@@ -50,6 +53,15 @@ public:
 
 	virtual void StartupModule() override
 	{
+		if (ISettingsModule* SettingsModule = FModuleManager::GetModulePtr<ISettingsModule>("Settings"))
+		{
+			SettingsModule->RegisterSettings("Project", "Editor", "Level Sequences",
+				LOCTEXT("RuntimeSettingsName", "Level Sequences"),
+				LOCTEXT("RuntimeSettingsDescription", "Configure project settings relating to Level Sequences"),
+				GetMutableDefault<UMovieSceneToolsProjectSettings>()
+			);
+		}
+
 		ISequencerModule& SequencerModule = FModuleManager::Get().LoadModuleChecked<ISequencerModule>( "Sequencer" );
 
 		// register property track editors
@@ -68,7 +80,8 @@ public:
 		ParticleTrackCreateEditorHandle = SequencerModule.RegisterTrackEditor_Handle( FOnCreateTrackEditor::CreateStatic( &FParticleTrackEditor::CreateTrackEditor ) );
 		ParticleParameterTrackCreateEditorHandle = SequencerModule.RegisterTrackEditor_Handle( FOnCreateTrackEditor::CreateStatic( &FParticleParameterTrackEditor::CreateTrackEditor ) );
 		PathTrackCreateEditorHandle = SequencerModule.RegisterTrackEditor_Handle( FOnCreateTrackEditor::CreateStatic( &F3DPathTrackEditor::CreateTrackEditor ) );
-		ShotTrackCreateEditorHandle = SequencerModule.RegisterTrackEditor_Handle( FOnCreateTrackEditor::CreateStatic( &FShotTrackEditor::CreateTrackEditor ) );
+		CameraCutTrackCreateEditorHandle = SequencerModule.RegisterTrackEditor_Handle( FOnCreateTrackEditor::CreateStatic( &FCameraCutTrackEditor::CreateTrackEditor ) );
+		CinematicShotTrackCreateEditorHandle = SequencerModule.RegisterTrackEditor_Handle( FOnCreateTrackEditor::CreateStatic( &FCinematicShotTrackEditor::CreateTrackEditor ) );
 		SlomoTrackCreateEditorHandle = SequencerModule.RegisterTrackEditor_Handle( FOnCreateTrackEditor::CreateStatic( &FSlomoTrackEditor::CreateTrackEditor ) );
 		SubTrackCreateEditorHandle = SequencerModule.RegisterTrackEditor_Handle( FOnCreateTrackEditor::CreateStatic( &FSubTrackEditor::CreateTrackEditor ) );
 		TransformTrackCreateEditorHandle = SequencerModule.RegisterTrackEditor_Handle( FOnCreateTrackEditor::CreateStatic( &F3DTransformTrackEditor::CreateTrackEditor ) );
@@ -81,6 +94,11 @@ public:
 
 	virtual void ShutdownModule() override
 	{
+		if (ISettingsModule* SettingsModule = FModuleManager::GetModulePtr<ISettingsModule>("Settings"))
+		{
+			SettingsModule->UnregisterSettings("Project", "Editor", "Level Sequences");
+		}
+
 		if (!FModuleManager::Get().IsModuleLoaded("Sequencer"))
 		{
 			return;
@@ -104,7 +122,8 @@ public:
 		SequencerModule.UnRegisterTrackEditor_Handle( ParticleTrackCreateEditorHandle );
 		SequencerModule.UnRegisterTrackEditor_Handle( ParticleParameterTrackCreateEditorHandle );
 		SequencerModule.UnRegisterTrackEditor_Handle( PathTrackCreateEditorHandle );
-		SequencerModule.UnRegisterTrackEditor_Handle( ShotTrackCreateEditorHandle );
+		SequencerModule.UnRegisterTrackEditor_Handle( CameraCutTrackCreateEditorHandle );
+		SequencerModule.UnRegisterTrackEditor_Handle( CinematicShotTrackCreateEditorHandle );
 		SequencerModule.UnRegisterTrackEditor_Handle( SlomoTrackCreateEditorHandle );
 		SequencerModule.UnRegisterTrackEditor_Handle( SubTrackCreateEditorHandle );
 		SequencerModule.UnRegisterTrackEditor_Handle( TransformTrackCreateEditorHandle );
@@ -164,7 +183,8 @@ private:
 	FDelegateHandle ParticleTrackCreateEditorHandle;
 	FDelegateHandle ParticleParameterTrackCreateEditorHandle;
 	FDelegateHandle PathTrackCreateEditorHandle;
-	FDelegateHandle ShotTrackCreateEditorHandle;
+	FDelegateHandle CameraCutTrackCreateEditorHandle;
+	FDelegateHandle CinematicShotTrackCreateEditorHandle;
 	FDelegateHandle SlomoTrackCreateEditorHandle;
 	FDelegateHandle SubTrackCreateEditorHandle;
 	FDelegateHandle TransformTrackCreateEditorHandle;
@@ -175,3 +195,6 @@ private:
 
 
 IMPLEMENT_MODULE( FMovieSceneToolsModule, MovieSceneTools );
+
+#undef LOCTEXT_NAMESPACE
+
diff --git a/Engine/Source/Editor/MovieSceneTools/Private/MovieSceneToolsPrivatePCH.h b/Engine/Source/Editor/MovieSceneTools/Private/MovieSceneToolsPrivatePCH.h
--- a/Engine/Source/Editor/MovieSceneTools/Private/MovieSceneToolsPrivatePCH.h
+++ b/Engine/Source/Editor/MovieSceneTools/Private/MovieSceneToolsPrivatePCH.h
@@ -18,8 +18,11 @@
 #include "CommonMovieSceneTools.h"
 #include "MovieSceneToolHelpers.h"
 #include "MovieSceneCommonHelpers.h"
+#include "MovieSceneToolsProjectSettings.h"
 
+#include "ThumbnailSection.h"
+#include "CameraCutSection.h"
+#include "CinematicShotSection.h"
 
-#include "ShotSequencerSection.h"
-#include "ShotTrackThumbnail.h"
-#include "ShotTrackThumbnailPool.h"
+#include "TrackEditorThumbnail.h"
+#include "TrackEditorThumbnailPool.h"
diff --git a/Engine/Source/Editor/MovieSceneTools/Private/MovieSceneToolsProjectSettings.cpp b/Engine/Source/Editor/MovieSceneTools/Private/MovieSceneToolsProjectSettings.cpp
new file mode 100644
--- /dev/null
+++ b/Engine/Source/Editor/MovieSceneTools/Private/MovieSceneToolsProjectSettings.cpp
@@ -0,0 +1,17 @@
+// Copyright 1998-2016 Epic Games, Inc. All Rights Reserved.
+
+#include "MovieSceneToolsPrivatePCH.h"
+#include "MovieSceneToolsProjectSettings.h"
+
+
+UMovieSceneToolsProjectSettings::UMovieSceneToolsProjectSettings()
+	: DefaultStartTime(0.f)
+	, DefaultDuration(5.f)
+	, ShotPrefix(TEXT("sht"))
+	, FirstShotNumber(10)
+	, ShotIncrement(10)
+	, ShotNumDigits(4)
+	, TakeNumDigits(3)
+	, FirstTakeNumber(1)
+	, TakeSeparator(TEXT("_"))
+{ }
diff --git a/Engine/Source/Editor/MovieSceneTools/Private/Sections/CameraCutSection.cpp b/Engine/Source/Editor/MovieSceneTools/Private/Sections/CameraCutSection.cpp
new file mode 100644
--- /dev/null
+++ b/Engine/Source/Editor/MovieSceneTools/Private/Sections/CameraCutSection.cpp
@@ -0,0 +1,158 @@
+// Copyright 1998-2016 Epic Games, Inc. All Rights Reserved.
+
+#include "MovieSceneToolsPrivatePCH.h"
+#include "ISectionLayoutBuilder.h"
+#include "Runtime/MovieSceneTracks/Public/Sections/MovieSceneCameraCutSection.h"
+#include "Runtime/Engine/Public/Slate/SceneViewport.h"
+#include "SInlineEditableTextBlock.h"
+
+
+#define LOCTEXT_NAMESPACE "FCameraCutSection"
+
+
+/* FCameraCutSection structors
+ *****************************************************************************/
+
+FCameraCutSection::FCameraCutSection(TSharedPtr<ISequencer> InSequencer, TSharedPtr<FTrackEditorThumbnailPool> InThumbnailPool, UMovieSceneSection& InSection) : FThumbnailSection(InSequencer, InThumbnailPool, InSection)
+{
+}
+
+FCameraCutSection::~FCameraCutSection()
+{
+}
+
+
+/* ISequencerSection interface
+ *****************************************************************************/
+
+void FCameraCutSection::BuildSectionContextMenu(FMenuBuilder& MenuBuilder, const FGuid& ObjectBinding)
+{
+	UWorld* World = GEditor->GetEditorWorldContext().World();
+
+	if (World == nullptr)
+	{
+		return;
+	}
+
+	// get list of available cameras
+	TArray<AActor*> AllCameras;
+	TArray<AActor*> SelectedCameras;
+
+	for (FActorIterator ActorIt(World); ActorIt; ++ActorIt)
+	{
+		AActor* Actor = *ActorIt;
+
+		if ((Actor != GetCameraObject()) && Actor->IsListedInSceneOutliner() && Actor->IsA<ACameraActor>())
+		{
+			AllCameras.Add(Actor);
+		}
+	}
+
+	if (AllCameras.Num() == 0)
+	{
+		return;
+	}
+
+	MenuBuilder.BeginSection(NAME_None, LOCTEXT("ChangeCameraMenuText", "Change Camera"));
+	{
+		for (auto EachCamera : AllCameras)
+		{
+			FText ActorLabel = FText::FromString(EachCamera->GetActorLabel());
+
+			MenuBuilder.AddMenuEntry(
+				FText::Format(LOCTEXT("SetCameraMenuEntryTextFormat", "{0}"), ActorLabel),
+				FText::Format(LOCTEXT("SetCameraMenuEntryTooltipFormat", "Assign {0} to this camera cut"), ActorLabel),
+				FSlateIcon(),
+				FUIAction(FExecuteAction::CreateRaw(this, &FCameraCutSection::HandleSetCameraMenuEntryExecute, EachCamera))
+			);
+		}
+	}
+	MenuBuilder.EndSection();
+}
+
+
+FText FCameraCutSection::GetDisplayName() const
+{
+	return NSLOCTEXT("FCameraCutSection", "", "CameraCuts");
+}
+
+
+/* FThumbnailSection interface
+ *****************************************************************************/
+
+ACameraActor* FCameraCutSection::GetCameraObject() const
+{
+	UMovieSceneCameraCutSection* CameraCutSection = Cast<UMovieSceneCameraCutSection>(Section);
+
+	TArray<UObject*> OutObjects;
+	// @todo sequencer: the director track may be able to get cameras from sub-movie scenes
+	SequencerPtr.Pin()->GetRuntimeObjects(SequencerPtr.Pin()->GetRootMovieSceneSequenceInstance(),  CameraCutSection->GetCameraGuid(), OutObjects);
+
+	ACameraActor* ReturnCam = nullptr;
+
+	if (OutObjects.Num() > 0)
+	{
+		ReturnCam = Cast<ACameraActor>(OutObjects[0]);
+
+		if (ReturnCam != nullptr)
+		{
+			return ReturnCam;
+		}
+	}
+
+	// Otherwise look in the current movie scene
+	OutObjects.Empty();
+	SequencerPtr.Pin()->GetRuntimeObjects(SequencerPtr.Pin()->GetFocusedMovieSceneSequenceInstance(), CameraCutSection->GetCameraGuid(), OutObjects);
+
+	if (OutObjects.Num() > 0)
+	{
+		ReturnCam = Cast<ACameraActor>(OutObjects[0]);
+
+		if (ReturnCam != nullptr)
+		{
+			return ReturnCam;
+		}
+	}
+
+	return ReturnCam;
+}
+
+FText FCameraCutSection::HandleThumbnailTextBlockText() const
+{
+	ACameraActor* CameraActor = GetCameraObject();
+	if (CameraActor)
+	{
+		return FText::FromString(CameraActor->GetActorLabel());
+	}
+
+	return FText::GetEmpty();
+}
+
+
+/* FCameraCutSection callbacks
+ *****************************************************************************/
+
+void FCameraCutSection::HandleSetCameraMenuEntryExecute(AActor* InCamera)
+{
+	auto Sequencer = SequencerPtr.Pin();
+
+	if (Sequencer.IsValid())
+	{
+		FGuid ObjectGuid = Sequencer->GetHandleToObject(InCamera, true);
+
+		UMovieSceneCameraCutSection* CameraCutSection = Cast<UMovieSceneCameraCutSection>(Section);
+
+		CameraCutSection->SetFlags(RF_Transactional);
+
+		const FScopedTransaction Transaction(LOCTEXT("SetCameraCut", "Set Camera Cut"));
+
+		CameraCutSection->Modify();
+	
+		CameraCutSection->SetCameraGuid(ObjectGuid);
+	
+		Sequencer->NotifyMovieSceneDataChanged();
+	}
+}
+
+
+#undef LOCTEXT_NAMESPACE
diff --git a/Engine/Source/Editor/MovieSceneTools/Private/Sections/CameraCutSection.h b/Engine/Source/Editor/MovieSceneTools/Private/Sections/CameraCutSection.h
new file mode 100644
--- /dev/null
+++ b/Engine/Source/Editor/MovieSceneTools/Private/Sections/CameraCutSection.h
@@ -0,0 +1,44 @@
+// Copyright 1998-2016 Epic Games, Inc. All Rights Reserved.
+
+#pragma once
+
+class FThumbnailSection;
+class FTrackEditorThumbnail;
+class FTrackEditorThumbnailPool;
+class IMenu;
+class ISectionLayoutBuilder;
+class UMovieSceneCameraCutSection;
+
+
+/**
+ * CameraCut section, which paints and ticks the appropriate section.
+ */
+class FCameraCutSection
+	: public FThumbnailSection
+{
+public:
+
+	/** Create and initialize a new instance. */
+	FCameraCutSection(TSharedPtr<ISequencer> InSequencer, TSharedPtr<FTrackEditorThumbnailPool> InThumbnailPool, UMovieSceneSection& InSection);
+
+	/** Virtual destructor. */
+	virtual ~FCameraCutSection();
+
+public:
+
+	// ISequencerSection interface
+
+	virtual void BuildSectionContextMenu(FMenuBuilder& MenuBuilder, const FGuid& ObjectBinding) override;
+	virtual FText GetDisplayName() const override;
+
+	// FThumbnail interface
+
+	virtual ACameraActor* GetCameraObject() const override;
+	virtual FText HandleThumbnailTextBlockText() const override;
+
+private:
+
+	/** Callback for executing a "Set Camera" menu entry in the context menu. */
+	void HandleSetCameraMenuEntryExecute(AActor* InCamera);
+
+};
diff --git a/Engine/Source/Editor/MovieSceneTools/Private/Sections/CinematicShotSection.cpp b/Engine/Source/Editor/MovieSceneTools/Private/Sections/CinematicShotSection.cpp
new file mode 100644
--- /dev/null
+++ b/Engine/Source/Editor/MovieSceneTools/Private/Sections/CinematicShotSection.cpp
@@ -0,0 +1,168 @@
+// Copyright 1998-2016 Epic Games, Inc. All Rights Reserved.
+
+#include "MovieSceneToolsPrivatePCH.h"
+#include "ISectionLayoutBuilder.h"
+#include "Runtime/MovieSceneTracks/Public/Sections/MovieSceneCinematicShotSection.h"
+#include "CinematicShotTrackEditor.h"
+#include "Runtime/Engine/Public/Slate/SceneViewport.h"
+#include "SInlineEditableTextBlock.h"
+#include "CinematicShotTrackEDitor.h"
+#include "MovieSceneToolHelpers.h"
+
+#define LOCTEXT_NAMESPACE "FCinematicShotSection"
+
+
+/* FCinematicShotSection structors
+ *****************************************************************************/
+
+FCinematicShotSection::FCinematicShotSection(TSharedPtr<ISequencer> InSequencer, TSharedPtr<FTrackEditorThumbnailPool> InThumbnailPool, UMovieSceneSection& InSection, TSharedPtr<FCinematicShotTrackEditor> InCinematicShotTrackEditor)
+	: FThumbnailSection(InSequencer, InThumbnailPool, InSection)
+	, SectionObject(*CastChecked<UMovieSceneCinematicShotSection>(&InSection))
+	, Sequencer(InSequencer)
+	, CinematicShotTrackEditor(InCinematicShotTrackEditor)
+{
+	if (InSequencer->HasSequenceInstanceForSection(InSection))
+	{
+		SequenceInstance = InSequencer->GetSequenceInstanceForSection(InSection);
+	}
+}
+
+
+FCinematicShotSection::~FCinematicShotSection()
+{
+}
+
+int32 FCinematicShotSection::OnPaintSection(const FGeometry& AllottedGeometry, const FSlateRect& SectionClippingRect, FSlateWindowElementList& OutDrawElements, int32 LayerId, bool bParentEnabled) const
+{
+	if (SequenceInstance.IsValid())
+	{
+		return FThumbnailSection::OnPaintSection(AllottedGeometry, SectionClippingRect, OutDrawElements, LayerId, bParentEnabled);
+	}
+	else
+	{
+		const ESlateDrawEffect::Type DrawEffects = bParentEnabled ? ESlateDrawEffect::None : ESlateDrawEffect::DisabledEffect;
+
+		FSlateDrawElement::MakeBox( 
+			OutDrawElements,
+			LayerId,
+			AllottedGeometry.ToPaintGeometry(),
+			FEditorStyle::GetBrush("Sequencer.GenericSection.Background"),
+			SectionClippingRect,
+			DrawEffects,
+			FLinearColor::Black
+		); 
+
+		return LayerId + 2;
+	}
+}
+
+void FCinematicShotSection::BuildSectionContextMenu(FMenuBuilder& MenuBuilder, const FGuid& ObjectBinding)
+{
+	MenuBuilder.BeginSection(NAME_None, LOCTEXT("ShotMenuText", "Shot"));
+	{
+		if (SequenceInstance.IsValid())
+		{
+			MenuBuilder.AddSubMenu(
+				LOCTEXT("TakesMenu", "Takes"),
+				LOCTEXT("TakesMenuTooltip", "Shot takes"),
+				FNewMenuDelegate::CreateLambda([=](FMenuBuilder& InMenuBuilder){ AddTakesMenu(InMenuBuilder); }));
+
+			MenuBuilder.AddMenuEntry(
+				LOCTEXT("NewTake", "New Take"),
+				FText::Format(LOCTEXT("NewTakeTooltip", "Create a new take for {0}"), SectionObject.GetShotDisplayName()),
+				FSlateIcon(),
+				FUIAction(FExecuteAction::CreateSP(CinematicShotTrackEditor.Pin().ToSharedRef(), &FCinematicShotTrackEditor::NewTake, &SectionObject))
+			);
+		}
+
+		MenuBuilder.AddMenuEntry(
+			LOCTEXT("InsertNewShot", "Insert Shot"),
+			FText::Format(LOCTEXT("InsertNewShotTooltip", "Insert a new shot after {0}"), SectionObject.GetShotDisplayName()),
+			FSlateIcon(),
+			FUIAction(FExecuteAction::CreateSP(CinematicShotTrackEditor.Pin().ToSharedRef(), &FCinematicShotTrackEditor::InsertShot, &SectionObject))
+		);
+
+		if (SequenceInstance.IsValid())
+		{
+			MenuBuilder.AddMenuEntry(
+				LOCTEXT("DuplicateShot", "Duplicate Shot"),
+				FText::Format(LOCTEXT("DuplicateShotTooltip", "Duplicate {0} to create a new shot"), SectionObject.GetShotDisplayName()),
+				FSlateIcon(),
+				FUIAction(FExecuteAction::CreateSP(CinematicShotTrackEditor.Pin().ToSharedRef(), &FCinematicShotTrackEditor::DuplicateShot, &SectionObject))
+			);
+
+			/*
+			//@todo
+			MenuBuilder.AddMenuEntry(
+				LOCTEXT("RenameShot", "Rename Shot"),
+				FText::Format(LOCTEXT("RenameShotTooltip", "Rename {0}"), SectionObject.GetShotDisplayName()),
+				FSlateIcon(),
+				FUIAction(FExecuteAction::CreateSP(CinematicShotTrackEditor.Pin().ToSharedRef(), &FCinematicShotTrackEditor::RenameShot, &SectionObject))
+			);
+			*/
+		}
+	}
+	MenuBuilder.EndSection();
+}
+
+void FCinematicShotSection::AddTakesMenu(FMenuBuilder& MenuBuilder)
+{
+	TArray<uint32> TakeNumbers;
+	MovieSceneToolHelpers::GatherTakes(&SectionObject, TakeNumbers);
+
+	for (auto TakeNumber : TakeNumbers)
+	{
+		MenuBuilder.AddMenuEntry(
+			FText::Format(LOCTEXT("TakeNumber", "Take {0}"), FText::AsNumber(TakeNumber)),
+			FText::Format(LOCTEXT("TakeNumberTooltip", "Switch to take {0}"), FText::AsNumber(TakeNumber)),
+			FSlateIcon(),
+			FUIAction(FExecuteAction::CreateSP(CinematicShotTrackEditor.Pin().ToSharedRef(), &FCinematicShotTrackEditor::SwitchTake, &SectionObject, TakeNumber))
+		);
+	}
+}
+
+FText FCinematicShotSection::GetDisplayName() const
+{
+	return NSLOCTEXT("FCinematicShotSection", "", "Shot");
+}
+
+/* FCinematicShotSection callbacks
+ *****************************************************************************/
+
+ACameraActor* FCinematicShotSection::GetCameraObject() const
+{
+	return CinematicShotTrackEditor.IsValid() ? CinematicShotTrackEditor.Pin()->GetCinematicShotCamera().Get() : nullptr;
+}
+
+FText FCinematicShotSection::HandleThumbnailTextBlockText() const
+{
+	return SectionObject.GetShotDisplayName();
+}
+
+
+void FCinematicShotSection::HandleThumbnailTextBlockTextCommitted(const FText& NewShotName, ETextCommit::Type CommitType)
+{
+	if (CommitType == ETextCommit::OnEnter && !HandleThumbnailTextBlockText().EqualTo(NewShotName))
+	{
+		SectionObject.Modify();
+
+		const FScopedTransaction Transaction(LOCTEXT("SetShotName", "Set Shot Name"));
+	
+		SectionObject.SetShotDisplayName(NewShotName);
+	}
+}
+
+FReply FCinematicShotSection::OnSectionDoubleClicked(const FGeometry& SectionGeometry, const FPointerEvent& MouseEvent)
+{
+	if( MouseEvent.GetEffectingButton() == EKeys::LeftMouseButton )
+	{
+		if (SequenceInstance.IsValid())
+		{
+			Sequencer.Pin()->FocusSequenceInstance(SequenceInstance.Pin().ToSharedRef());
+		}
+	}
+
+	return FReply::Handled();
+}
+
+#undef LOCTEXT_NAMESPACE
diff --git a/Engine/Source/Editor/MovieSceneTools/Private/Sections/CinematicShotSection.h b/Engine/Source/Editor/MovieSceneTools/Private/Sections/CinematicShotSection.h
new file mode 100644
--- /dev/null
+++ b/Engine/Source/Editor/MovieSceneTools/Private/Sections/CinematicShotSection.h
@@ -0,0 +1,63 @@
+// Copyright 1998-2016 Epic Games, Inc. All Rights Reserved.
+
+#pragma once
+
+class FThumbnailSection;
+class FTrackEditorThumbnail;
+class FTrackEditorThumbnailPool;
+class FCinematicShotTrackEditor; 
+class IMenu;
+class ISectionLayoutBuilder;
+class UMovieSceneCinematicShotSection;
+class FMovieSceneSequenceInstance;
+
+
+/**
+ * CinematicShot section, which paints and ticks the appropriate section.
+ */
+class FCinematicShotSection
+	: public FThumbnailSection
+{
+public:
+
+	/** Create and initialize a new instance. */
+	FCinematicShotSection(TSharedPtr<ISequencer> InSequencer, TSharedPtr<FTrackEditorThumbnailPool> InThumbnailPool, UMovieSceneSection& InSection, TSharedPtr<FCinematicShotTrackEditor> InCinematicShotTrackEditor);
+
+	/** Virtual destructor. */
+	virtual ~FCinematicShotSection();
+
+public:
+
+	// ISequencerSection interface
+
+	virtual int32 OnPaintSection( const FGeometry& AllottedGeometry, const FSlateRect& SectionClippingRect, FSlateWindowElementList& OutDrawElements, int32 LayerId, bool bParentEnabled ) const override;
+	virtual void BuildSectionContextMenu(FMenuBuilder& MenuBuilder, const FGuid& ObjectBinding) override;
+	virtual FText GetDisplayName() const override;
+	virtual FReply OnSectionDoubleClicked(const FGeometry& SectionGeometry, const FPointerEvent& MouseEvent) override;
+
+	// FThumbnail interface
+
+	virtual ACameraActor* GetCameraObject() const override;
+	virtual bool CanRename() const override { return true; }
+	virtual FText HandleThumbnailTextBlockText() const override;
+	virtual void HandleThumbnailTextBlockTextCommitted(const FText& NewThumbnailName, ETextCommit::Type CommitType) override;
+
+private:
+
+	/** Add shot takes menu */
+	void AddTakesMenu(FMenuBuilder& MenuBuilder);
+
+private:
+
+	/** The section we are visualizing */
+	UMovieSceneCinematicShotSection& SectionObject;
+
+	/** The instance that this section is part of */
+	TWeakPtr<FMovieSceneSequenceInstance> SequenceInstance;
+
+	/** Sequencer interface */
+	TWeakPtr<ISequencer> Sequencer;
+
+	/** The cinematic shot track editor that contains this section */
+	TWeakPtr<FCinematicShotTrackEditor> CinematicShotTrackEditor;
+};
diff --git a/Engine/Source/Editor/MovieSceneTools/Private/Sections/ThumbnailSection.cpp b/Engine/Source/Editor/MovieSceneTools/Private/Sections/ThumbnailSection.cpp
new file mode 100644
--- /dev/null
+++ b/Engine/Source/Editor/MovieSceneTools/Private/Sections/ThumbnailSection.cpp
@@ -0,0 +1,317 @@
+// Copyright 1998-2016 Epic Games, Inc. All Rights Reserved.
+
+#include "MovieSceneToolsPrivatePCH.h"
+#include "ISectionLayoutBuilder.h"
+#include "Runtime/Engine/Public/Slate/SceneViewport.h"
+#include "SInlineEditableTextBlock.h"
+
+
+#define LOCTEXT_NAMESPACE "FThumbnailSection"
+
+
+namespace ThumbnailSectionConstants
+{
+	const uint32 ThumbnailHeight = 90;
+	const uint32 TrackHeight = ThumbnailHeight+10; // some extra padding
+	const uint32 TrackWidth = 90;
+	const float SectionGripSize = 4.0f;
+	const FName ThumbnailGripBrushName = FName("Sequencer.Thumbnail.SectionHandle");
+}
+
+
+/* FThumbnailSection structors
+ *****************************************************************************/
+
+FThumbnailSection::FThumbnailSection(TSharedPtr<ISequencer> InSequencer, TSharedPtr<FTrackEditorThumbnailPool> InThumbnailPool, UMovieSceneSection& InSection)
+	: Section(&InSection)
+	, SequencerPtr(InSequencer)
+	, ThumbnailPool(InThumbnailPool)
+	, Thumbnails()
+	, ThumbnailWidth(0)
+	, StoredSize(ForceInit)
+	, StoredStartTime(0.f)
+	, VisibleTimeRange(TRange<float>::Empty())
+	, InternalViewportScene(nullptr)
+	, InternalViewportClient(nullptr)
+{
+	// @todo Sequencer optimize This code shouldn't even be called if this is offscreen
+	// the following code could be generated on demand when the section is scrolled onscreen
+	InternalViewportClient = MakeShareable(new FLevelEditorViewportClient(nullptr));
+	{
+		InternalViewportClient->ViewportType = LVT_Perspective;
+		InternalViewportClient->bDisableInput = true;
+		InternalViewportClient->bDrawAxes = false;
+		InternalViewportClient->EngineShowFlags = FEngineShowFlags(ESFIM_Game);
+		InternalViewportClient->EngineShowFlags.DisableAdvancedFeatures();
+		InternalViewportClient->SetRealtime( false );
+	}
+
+	InternalViewportScene = MakeShareable(new FSceneViewport(InternalViewportClient.Get(), nullptr));
+	InternalViewportClient->Viewport = InternalViewportScene.Get();
+
+	WhiteBrush = FEditorStyle::GetBrush("WhiteBrush");
+}
+
+
+FThumbnailSection::~FThumbnailSection()
+{
+	if (ThumbnailPool.IsValid())
+	{
+		ThumbnailPool.Pin()->RemoveThumbnailsNeedingRedraw(Thumbnails);
+	}
+	
+	if (InternalViewportClient.IsValid())
+	{
+		InternalViewportClient->Viewport = nullptr;
+	}
+}
+
+
+/* FThumbnailSection interface
+ *****************************************************************************/
+
+void FThumbnailSection::DrawViewportThumbnail(TSharedPtr<FTrackEditorThumbnail> TrackEditorThumbnail)
+{
+	if (InternalViewportScene.IsValid())
+	{
+		SequencerPtr.Pin()->SetGlobalTime(TrackEditorThumbnail->GetTime());
+		InternalViewportClient->SetActorLock(GetCameraObject());
+		InternalViewportClient->UpdateViewForLockedActor();
+		GWorld->SendAllEndOfFrameUpdates();
+		InternalViewportScene->Draw(false);
+		TrackEditorThumbnail->CopyTextureIn((FSlateRenderTargetRHI*)InternalViewportScene->GetViewportRenderTargetTexture());
+	}
+}
+
+
+uint32 FThumbnailSection::GetThumbnailWidth() const
+{
+	return ThumbnailWidth;
+}
+
+
+void FThumbnailSection::RegenerateViewportThumbnails(const FIntPoint& Size)
+{
+	ACameraActor* CameraActor = GetCameraObject();
+
+	if (CameraActor != nullptr && InternalViewportScene.IsValid())
+	{
+		ThumbnailPool.Pin()->RemoveThumbnailsNeedingRedraw(Thumbnails);
+		Thumbnails.Empty();
+
+		if (Size.X == 0 || Size.Y == 0)
+		{
+			return;
+		}
+
+		StoredSize = Size;
+		StoredStartTime = Section->GetStartTime();
+	
+		float AspectRatio = CameraActor->CameraComponent->AspectRatio;
+
+		int32 NewThumbnailWidth = FMath::TruncToInt(ThumbnailSectionConstants::TrackWidth * AspectRatio);
+		int32 ThumbnailCount = FMath::DivideAndRoundUp(StoredSize.X, NewThumbnailWidth);
+
+		if (InternalViewportClient.IsValid())
+		{
+			InternalViewportClient->SetAllowCinematicPreview(false);
+			InternalViewportClient->SetActorLock(CameraActor);
+		}
+
+		if (NewThumbnailWidth != ThumbnailWidth)
+		{
+			ThumbnailWidth = NewThumbnailWidth;
+			InternalViewportScene->UpdateViewportRHI(false, ThumbnailWidth, ThumbnailSectionConstants::ThumbnailHeight, EWindowMode::Windowed);
+		}
+
+		float StartTime = Section->GetStartTime();
+		float EndTime = Section->GetEndTime();
+		float DeltaTime = EndTime - StartTime;
+
+		int32 TotalX = StoredSize.X;
+
+		// @todo sequencer optimize this to use a single viewport and a single thumbnail
+		// instead paste the textures into the thumbnail at specific offsets
+		for (int32 ThumbnailIndex = 0; ThumbnailIndex < ThumbnailCount; ++ThumbnailIndex)
+		{
+			int32 StartX = ThumbnailIndex * ThumbnailWidth;
+			float FractionThrough = (float)StartX / (float)TotalX;
+			float Time = StartTime + DeltaTime * FractionThrough;
+
+			int32 NextStartX = (ThumbnailIndex+1) * ThumbnailWidth;
+			float NextFractionThrough = (float)NextStartX / (float)TotalX;
+			float NextTime = FMath::Min(StartTime + DeltaTime * NextFractionThrough, EndTime);
+
+			FIntPoint ThumbnailSize(NextStartX-StartX, ThumbnailSectionConstants::ThumbnailHeight);
+
+			check(FractionThrough >= 0.f && FractionThrough <= 1.f && NextFractionThrough >= 0.f);
+			TSharedPtr<FTrackEditorThumbnail> NewThumbnail = MakeShareable(new FTrackEditorThumbnail(SharedThis(this), ThumbnailSize, TRange<float>(Time, NextTime)));
+
+			Thumbnails.Add(NewThumbnail);
+		}
+
+		// @todo Sequencer Optimize Only say a thumbnail needs redraw if it is onscreen
+		ThumbnailPool.Pin()->AddThumbnailsNeedingRedraw(Thumbnails);
+	}
+}
+
+
+/* ISequencerSection interface
+ *****************************************************************************/
+
+bool FThumbnailSection::ShouldDrawKeyAreaBackground() const
+{
+	return false;
+}
+
+bool FThumbnailSection::AreSectionsConnected() const
+{
+	return true;
+}
+
+TSharedRef<SWidget> FThumbnailSection::GenerateSectionWidget()
+{
+	return SNew(SBox)
+		.HAlign(HAlign_Left)
+		.VAlign(VAlign_Top)
+		.Padding(FMargin(15.0f, 7.0f, 0.0f, 0.0f))
+		[
+			SNew(SInlineEditableTextBlock)
+				.ToolTipText(CanRename() ? LOCTEXT("RenameThumbnail", "Click or hit F2 to rename") : FText::GetEmpty())
+				.Text(this, &FThumbnailSection::HandleThumbnailTextBlockText)
+				.ShadowOffset(FVector2D(1,1))
+				.OnTextCommitted(this, &FThumbnailSection::HandleThumbnailTextBlockTextCommitted)
+				.IsReadOnly(!CanRename())
+		];
+}
+
+
+FName FThumbnailSection::GetSectionGripLeftBrushName() const
+{
+	return ThumbnailSectionConstants::ThumbnailGripBrushName;
+}
+
+
+FName FThumbnailSection::GetSectionGripRightBrushName() const
+{
+	return ThumbnailSectionConstants::ThumbnailGripBrushName;
+}
+
+
+float FThumbnailSection::GetSectionGripSize() const
+{
+	return ThumbnailSectionConstants::SectionGripSize;
+}
+
+
+float FThumbnailSection::GetSectionHeight() const
+{
+	return ThumbnailSectionConstants::TrackHeight;
+}
+
+
+UMovieSceneSection* FThumbnailSection::GetSectionObject() 
+{
+	return Section;
+}
+
+
+FText FThumbnailSection::GetSectionTitle() const
+{
+	return FText::GetEmpty();
+}
+
+
+int32 FThumbnailSection::OnPaintSection(const FGeometry& AllottedGeometry, const FSlateRect& SectionClippingRect, FSlateWindowElementList& OutDrawElements, int32 LayerId, bool bParentEnabled) const
+{
+	const ESlateDrawEffect::Type DrawEffects = bParentEnabled ? ESlateDrawEffect::None : ESlateDrawEffect::DisabledEffect;
+
+	if (Thumbnails.Num() != 0)
+	{
+		FGeometry ThumbnailAreaGeometry = AllottedGeometry.MakeChild( FVector2D(GetSectionGripSize(), 0.0f), AllottedGeometry.GetDrawSize() - FVector2D( GetSectionGripSize()*2, 0.0f ) );
+
+		FSlateDrawElement::MakeBox(
+				OutDrawElements,
+				LayerId,
+				ThumbnailAreaGeometry.ToPaintGeometry(),
+				FEditorStyle::GetBrush("Sequencer.GenericSection.Background"),
+				SectionClippingRect,
+				DrawEffects
+			);
+
+		// @todo Sequencer: Need a way to visualize the key here
+
+		int32 ThumbnailIndex = 0;
+		for( const TSharedPtr<FTrackEditorThumbnail>& Thumbnail : Thumbnails )
+		{
+			FGeometry TruncatedGeometry = ThumbnailAreaGeometry.MakeChild(
+				Thumbnail->GetSize(),
+				FSlateLayoutTransform(ThumbnailAreaGeometry.Scale, FVector2D( ThumbnailIndex * ThumbnailWidth, 5.f))
+				);
+			
+			FSlateDrawElement::MakeViewport(
+				OutDrawElements,
+				LayerId,
+				TruncatedGeometry.ToPaintGeometry(),
+				Thumbnail,
+				SectionClippingRect,
+				false,
+				false,
+				DrawEffects,
+				FLinearColor::White
+			);
+
+			if(Thumbnail->GetFadeInCurve() > 0.0f )
+			{
+				FSlateDrawElement::MakeBox(
+					OutDrawElements,
+					LayerId+1,
+					TruncatedGeometry.ToPaintGeometry(),
+					WhiteBrush,
+					SectionClippingRect,
+					DrawEffects,
+					FLinearColor(1.0f, 1.0f, 1.0f, Thumbnail->GetFadeInCurve())
+					);
+			}
+
+			++ThumbnailIndex;
+		}
+	}
+	else
+	{
+		FSlateDrawElement::MakeBox( 
+			OutDrawElements,
+			LayerId,
+			AllottedGeometry.ToPaintGeometry(),
+			FEditorStyle::GetBrush("Sequencer.GenericSection.Background"),
+			SectionClippingRect,
+			DrawEffects
+		); 
+	}
+
+	return LayerId + 2;
+}
+
+
+void FThumbnailSection::Tick(const FGeometry& AllottedGeometry, const FGeometry& ParentGeometry, const double InCurrentTime, const float InDeltaTime)
+{
+	FTimeToPixel TimeToPixelConverter(AllottedGeometry, TRange<float>(Section->GetStartTime(), Section->GetEndTime()));
+
+	// get the visible time range
+	VisibleTimeRange = TRange<float>(
+		TimeToPixelConverter.PixelToTime(-AllottedGeometry.Position.X),
+		TimeToPixelConverter.PixelToTime(-AllottedGeometry.Position.X + ParentGeometry.Size.X)
+	);
+
+	FIntPoint AllocatedSize = AllottedGeometry.MakeChild( FVector2D( GetSectionGripSize(), 0.0f ), FVector2D( AllottedGeometry.Size.X - (GetSectionGripSize()*2), AllottedGeometry.Size.Y) ).Size.IntPoint();
+	AllocatedSize.X = FMath::Max(AllocatedSize.X, 1);
+
+	float StartTime = Section->GetStartTime();
+
+	if (!FSlateApplication::Get().HasAnyMouseCaptor() && ( AllocatedSize.X != StoredSize.X || !FMath::IsNearlyEqual(StartTime, StoredStartTime) ) )
+	{
+		RegenerateViewportThumbnails(AllocatedSize);
+	}
+}
+
+#undef LOCTEXT_NAMESPACE
diff --git a/Engine/Source/Editor/MovieSceneTools/Private/Sections/ThumbnailSection.h b/Engine/Source/Editor/MovieSceneTools/Private/Sections/ThumbnailSection.h
new file mode 100644
--- /dev/null
+++ b/Engine/Source/Editor/MovieSceneTools/Private/Sections/ThumbnailSection.h
@@ -0,0 +1,114 @@
+// Copyright 1998-2016 Epic Games, Inc. All Rights Reserved.
+
+#pragma once
+
+
+class FTrackEditorThumbnail;
+class FTrackEditorThumbnailPool;
+class IMenu;
+class ISectionLayoutBuilder;
+class UMovieSceneSection;
+
+
+/**
+ * Thumbnail section, which paints and ticks the appropriate section.
+ */
+class FThumbnailSection
+	: public ISequencerSection
+	, public TSharedFromThis<FThumbnailSection>
+{
+public:
+
+	/** Create and initialize a new instance. */
+	FThumbnailSection(TSharedPtr<ISequencer> InSequencer, TSharedPtr<FTrackEditorThumbnailPool> InThumbnailPool, UMovieSceneSection& InSection);
+
+	/** Virtual destructor. */
+	virtual ~FThumbnailSection();
+
+public:
+
+	/** Draws the passed in viewport thumbnail and copies it to the thumbnail's texture. */
+	void DrawViewportThumbnail(TSharedPtr<FTrackEditorThumbnail> TrackEditorThumbnail);
+
+	/** @return The sequencer widget owning the MovieScene section. */
+	TSharedRef<SWidget> GetSequencerWidget()
+	{
+		return SequencerPtr.Pin()->GetSequencerWidget();
+	}
+
+	/** Gets the thumbnail width. */
+	uint32 GetThumbnailWidth() const;
+
+	/** Gets the time range of what in the sequencer is visible. */
+	TRange<float> GetVisibleTimeRange() const
+	{
+		return VisibleTimeRange;
+	}
+
+	/** Regenerates all viewports and thumbnails at the new size. */
+	void RegenerateViewportThumbnails(const FIntPoint& Size);
+
+	/** Get the camera that this thumbnail should draw from */
+	virtual ACameraActor* GetCameraObject() const { return nullptr; }
+
+	/** Get whether the text is renameable */
+	virtual bool CanRename() const { return false; }
+
+	/** Callback for getting the text of the track name text block. */
+	virtual FText HandleThumbnailTextBlockText() const { return FText::GetEmpty(); }
+
+	/** Callback for when the text of the track name text block has changed. */
+	virtual void HandleThumbnailTextBlockTextCommitted(const FText& NewThumbnailName, ETextCommit::Type CommitType) { }
+
+public:
+
+	// ISequencerSection interface
+
+	virtual bool ShouldDrawKeyAreaBackground() const override;
+	virtual bool AreSectionsConnected() const override;
+	virtual void GenerateSectionLayout(ISectionLayoutBuilder& LayoutBuilder) const override { }
+	virtual TSharedRef<SWidget> GenerateSectionWidget() override;
+	virtual FName GetSectionGripLeftBrushName() const override;
+	virtual FName GetSectionGripRightBrushName() const override;
+	virtual float GetSectionGripSize() const override;
+	virtual float GetSectionHeight() const override;
+	virtual UMovieSceneSection* GetSectionObject() override;
+	virtual FText GetSectionTitle() const override;
+	virtual int32 OnPaintSection( const FGeometry& AllottedGeometry, const FSlateRect& SectionClippingRect, FSlateWindowElementList& OutDrawElements, int32 LayerId, bool bParentEnabled ) const override;
+	virtual void Tick( const FGeometry& AllottedGeometry, const FGeometry& ParentGeometry, const double InCurrentTime, const float InDeltaTime ) override;
+
+protected:
+
+	/** The section we are visualizing. */
+	UMovieSceneSection* Section;
+
+	/** The parent sequencer we are a part of. */
+	TWeakPtr<ISequencer> SequencerPtr;
+
+	/** The thumbnail pool that we are sending all of our thumbnails to. */
+	TWeakPtr<FTrackEditorThumbnailPool> ThumbnailPool;
+
+	/** A list of all thumbnails this CameraCut section has. */
+	TArray<TSharedPtr<FTrackEditorThumbnail>> Thumbnails;
+
+	/** The width of our thumbnails. */
+	uint32 ThumbnailWidth;
+
+	/** The stored size of this section in the Slate geometry. */
+	FIntPoint StoredSize;
+
+	/** The stored start time, to query for invalidations. */
+	float StoredStartTime;
+
+	/** Cached Time Range of the visible parent section area. */
+	TRange<float> VisibleTimeRange;
+	
+	/** An internal viewport scene we use to render the thumbnails with. */
+	TSharedPtr<FSceneViewport> InternalViewportScene;
+
+	/** An internal editor viewport client to render the thumbnails with. */
+	TSharedPtr<FLevelEditorViewportClient> InternalViewportClient;
+	
+	/** Fade brush. */
+	const FSlateBrush* WhiteBrush;
+};
diff --git a/Engine/Source/Editor/MovieSceneTools/Private/TrackEditorThumbnail/TrackEditorThumbnail.cpp b/Engine/Source/Editor/MovieSceneTools/Private/TrackEditorThumbnail/TrackEditorThumbnail.cpp
new file mode 100644
--- /dev/null
+++ b/Engine/Source/Editor/MovieSceneTools/Private/TrackEditorThumbnail/TrackEditorThumbnail.cpp
@@ -0,0 +1,107 @@
+// Copyright 1998-2016 Epic Games, Inc. All Rights Reserved.
+
+#include "MovieSceneToolsPrivatePCH.h"
+#include "Runtime/Engine/Public/Slate/SlateTextures.h"
+
+
+namespace TrackEditorThumbnailConstants
+{
+	const double ThumbnailFadeInDuration = 0.25f;
+}
+
+
+/* FTrackEditorThumbnail structors
+ *****************************************************************************/
+
+FTrackEditorThumbnail::FTrackEditorThumbnail(TSharedPtr<FThumbnailSection> InSection, const FIntPoint& InSize, TRange<float> InTimeRange)
+	: OwningSection(InSection)
+	, Size(InSize)
+	, Texture(nullptr)
+	, TimeRange(InTimeRange)
+	, FadeInCurve(0.0f, TrackEditorThumbnailConstants::ThumbnailFadeInDuration)
+{
+	Texture = new FSlateTexture2DRHIRef(GetSize().X, GetSize().Y, PF_B8G8R8A8, nullptr, TexCreate_Dynamic, true);
+
+	BeginInitResource( Texture );
+
+}
+
+
+FTrackEditorThumbnail::~FTrackEditorThumbnail()
+{
+	BeginReleaseResource(Texture);
+	FlushRenderingCommands();
+
+	if (Texture) 
+	{
+		delete Texture;
+	}
+}
+
+
+/* FTrackEditorThumbnail interface
+ *****************************************************************************/
+
+void FTrackEditorThumbnail::CopyTextureIn(FSlateRenderTargetRHI* InTexture)
+{
+	ENQUEUE_UNIQUE_RENDER_COMMAND_TWOPARAMETER(ReadTexture,
+		FSlateRenderTargetRHI*, RenderTarget, InTexture,
+		FSlateTexture2DRHIRef*, TargetTexture, Texture,
+	{
+		RHICmdList.CopyToResolveTarget(RenderTarget->GetRHIRef(), TargetTexture->GetTypedResource(), false, FResolveParams());
+	});
+}
+
+
+void FTrackEditorThumbnail::DrawThumbnail()
+{
+	OwningSection.Pin()->DrawViewportThumbnail(SharedThis(this));
+	FadeInCurve.PlayReverse( OwningSection.Pin()->GetSequencerWidget() );
+}
+
+
+float FTrackEditorThumbnail::GetFadeInCurve() const 
+{
+	return FadeInCurve.GetLerp();
+}
+
+
+float FTrackEditorThumbnail::GetTime() const
+{
+	return TimeRange.GetLowerBoundValue();
+}
+
+
+bool FTrackEditorThumbnail::IsValid() const
+{
+	return OwningSection.IsValid();
+}
+
+
+bool FTrackEditorThumbnail::IsVisible() const
+{
+	return OwningSection.IsValid()
+		? TimeRange.Overlaps(OwningSection.Pin()->GetVisibleTimeRange())
+		: false;
+}
+
+
+/* ISlateViewport interface
+ *****************************************************************************/
+
+FIntPoint FTrackEditorThumbnail::GetSize() const
+{
+	return Size;
+}
+
+
+FSlateShaderResource* FTrackEditorThumbnail::GetViewportRenderTargetTexture() const
+{
+	return Texture;
+}
+
+
+bool FTrackEditorThumbnail::RequiresVsync() const 
+{
+	return false;
+}
diff --git a/Engine/Source/Editor/MovieSceneTools/Private/TrackEditorThumbnail/TrackEditorThumbnail.h b/Engine/Source/Editor/MovieSceneTools/Private/TrackEditorThumbnail/TrackEditorThumbnail.h
new file mode 100644
--- /dev/null
+++ b/Engine/Source/Editor/MovieSceneTools/Private/TrackEditorThumbnail/TrackEditorThumbnail.h
@@ -0,0 +1,68 @@
+// Copyright 1998-2016 Epic Games, Inc. All Rights Reserved.
+
+#pragma once
+
+
+class FThumbnailSection;
+class FSlateRenderTargetRHI;
+class FSlateTexture2DRHIRef;
+
+
+/**
+ * Track Editor Thumbnail, which keeps a Texture to be displayed by a viewport.
+ */
+class FTrackEditorThumbnail
+	: public ISlateViewport
+	, public TSharedFromThis<FTrackEditorThumbnail>
+{
+public:
+
+	/** Create and initialize a new instance. */
+	FTrackEditorThumbnail(TSharedPtr<FThumbnailSection> InSection, const FIntPoint& InSize, TRange<float> InTimeRange);
+
+	/** Virtual destructor. */
+	virtual ~FTrackEditorThumbnail();
+
+public:
+
+	/** Copies the incoming render target to this thumbnails texture. */
+	void CopyTextureIn(FSlateRenderTargetRHI* InTexture);
+
+	/** Renders the thumbnail to the texture. */
+	void DrawThumbnail();
+
+	/** Gets the curve for fading in the thumbnail. */
+	float GetFadeInCurve() const;
+
+	/** Gets the time that this thumbnail is a rendering of. */
+	float GetTime() const;
+
+	bool IsValid() const;
+
+	/** Returns whether this thumbnail is visible based on the thumbnail section geometry visibility. */
+	bool IsVisible() const;
+
+public:
+
+	// ISlateViewport interface
+
+	virtual FIntPoint GetSize() const override;
+	virtual FSlateShaderResource* GetViewportRenderTargetTexture() const override;
+	virtual bool RequiresVsync() const override;
+
+private:
+
+	/** Parent thumbnail section we are a thumbnail of. */
+	TWeakPtr<FThumbnailSection> OwningSection;
+	
+	FIntPoint Size;
+
+	/** The Texture RHI that holds the thumbnail. */
+	FSlateTexture2DRHIRef* Texture;
+
+	/** Where in time this thumbnail is a rendering of. */
+	TRange<float> TimeRange;
+
+	/** Fade curve to display while the thumbnail is redrawing. */
+	FCurveSequence FadeInCurve;
+};
