[INTEGRATE] Merged Dev-Editor fix for access to private type, and removed hack workaround
Merged from CL 2892765

diff --git a/Engine/Plugins/MovieScene/LevelSequenceEditor/Source/LevelSequenceEditor/Private/CinematicLevelViewportLayout.cpp b/Engine/Plugins/MovieScene/LevelSequenceEditor/Source/LevelSequenceEditor/Private/CinematicLevelViewportLayout.cpp
--- a/Engine/Plugins/MovieScene/LevelSequenceEditor/Source/LevelSequenceEditor/Private/CinematicLevelViewportLayout.cpp
+++ b/Engine/Plugins/MovieScene/LevelSequenceEditor/Source/LevelSequenceEditor/Private/CinematicLevelViewportLayout.cpp
@@ -106,7 +106,7 @@ TSharedRef<SWidget> FCinematicLevelViewportLayout_TwoPane::MakeViewportLayout(co
 			SAssignNew(ViewportWidget0, SLevelViewport)
 			.ViewportType(LVT_Perspective)
 			.ParentLayout(AsShared())
-			.ParentLevelEditor(*(TWeakPtr<ILevelEditor>*)&ParentLevelEditor)
+			.ParentLevelEditor(ParentLevelEditor)
 			.ConfigKey(ViewportKey0)
 		]
 
diff --git a/Engine/Source/Editor/LevelEditor/Private/LevelViewportLayout.cpp b/Engine/Source/Editor/LevelEditor/Private/LevelViewportLayout.cpp
--- a/Engine/Source/Editor/LevelEditor/Private/LevelViewportLayout.cpp
+++ b/Engine/Source/Editor/LevelEditor/Private/LevelViewportLayout.cpp
@@ -199,7 +199,7 @@ FLevelViewportLayout::~FLevelViewportLayout()
 }
 
 
-TSharedRef<SWidget> FLevelViewportLayout::BuildViewportLayout( TSharedPtr<SDockTab> InParentDockTab, TSharedPtr<FLevelViewportTabContent> InParentTab, const FString& LayoutString, TWeakPtr<SLevelEditor> InParentLevelEditor )
+TSharedRef<SWidget> FLevelViewportLayout::BuildViewportLayout( TSharedPtr<SDockTab> InParentDockTab, TSharedPtr<FLevelViewportTabContent> InParentTab, const FString& LayoutString, TWeakPtr<ILevelEditor> InParentLevelEditor )
 {
 	// We don't support reconfiguring an existing layout object, as this makes handling of transitions
 	// particularly difficult.  Instead just destroy the old layout and create a new layout object.
diff --git a/Engine/Source/Editor/LevelEditor/Private/SLevelEditor.h b/Engine/Source/Editor/LevelEditor/Private/SLevelEditor.h
--- a/Engine/Source/Editor/LevelEditor/Private/SLevelEditor.h
+++ b/Engine/Source/Editor/LevelEditor/Private/SLevelEditor.h
@@ -86,12 +86,12 @@ public:
 	 * @param MyGeometry		Information about the size and position of the viewport widget
 	 * @param InKeyEvent	The event which just occurred	
 	 */
-	virtual FReply OnKeyDownInViewport( const FGeometry& MyGeometry, const FKeyEvent& InKeyEvent );
+	virtual FReply OnKeyDownInViewport( const FGeometry& MyGeometry, const FKeyEvent& InKeyEvent ) override;
 
 	bool CanCloseApp();
 
 	/** Returns the full action list for this level editor instance */
-	const TSharedPtr< FUICommandList >& GetLevelEditorActions() const
+	virtual const TSharedPtr< FUICommandList >& GetLevelEditorActions() const override
 	{
 		return LevelEditorCommands;
 	}
@@ -113,7 +113,7 @@ public:
 	}
 
 	// Tab Management
-	TSharedRef<FTabManager> GetTabManager() const;
+	virtual TSharedRef<FTabManager> GetTabManager() const override;
 	
 	/** Attaches a sequencer asset editor used to animate objects in the level to this level editor */
 	void AttachSequencer( TSharedPtr<SWidget> SequencerWidget, TSharedPtr<IAssetEditorInstance> NewSequencerAssetEditor );
diff --git a/Engine/Source/Editor/LevelEditor/Private/SLevelViewport.cpp b/Engine/Source/Editor/LevelEditor/Private/SLevelViewport.cpp
--- a/Engine/Source/Editor/LevelEditor/Private/SLevelViewport.cpp
+++ b/Engine/Source/Editor/LevelEditor/Private/SLevelViewport.cpp
@@ -417,7 +417,7 @@ FReply SLevelViewport::OnKeyDown( const FGeometry& MyGeometry, const FKeyEvent&
 		//  the SLevelEditor is not a direct parent of the viewport.  
 		if ( this->IsImmersive() && !Reply.IsEventHandled() )
 		{
-			TSharedPtr<SLevelEditor> ParentLevelEditorSharedPtr = ParentLevelEditor.Pin();
+			TSharedPtr<ILevelEditor> ParentLevelEditorSharedPtr = ParentLevelEditor.Pin();
 			if( ParentLevelEditorSharedPtr.IsValid() )
 			{
 				Reply = ParentLevelEditorSharedPtr->OnKeyDownInViewport( MyGeometry, InKeyEvent );
diff --git a/Engine/Source/Editor/LevelEditor/Public/ILevelEditor.h b/Engine/Source/Editor/LevelEditor/Public/ILevelEditor.h
--- a/Engine/Source/Editor/LevelEditor/Public/ILevelEditor.h
+++ b/Engine/Source/Editor/LevelEditor/Public/ILevelEditor.h
@@ -28,6 +28,15 @@ public:
 	/** Get the thumbnail pool used by this level editor */
 	virtual TSharedPtr< class FAssetThumbnailPool > GetThumbnailPool() const = 0;
 
+	/** Access the level editor's tab manager */
+	virtual TSharedRef<FTabManager> GetTabManager() const = 0;
+
+	/** Access the level editor's action command list */
+	virtual const TSharedPtr< FUICommandList >& GetLevelEditorActions() const = 0;
+
+	/** Called to process a key down event in a viewport when in immersive mode */
+	virtual FReply OnKeyDownInViewport( const FGeometry& MyGeometry, const FKeyEvent& InKeyEvent ) = 0;
+
 	/** Append commands to the command list for the level editor */
 	virtual void AppendCommands( const TSharedRef<FUICommandList>& InCommandsToAppend ) = 0;
 
diff --git a/Engine/Source/Editor/LevelEditor/Public/LevelViewportLayout.h b/Engine/Source/Editor/LevelEditor/Public/LevelViewportLayout.h
--- a/Engine/Source/Editor/LevelEditor/Public/LevelViewportLayout.h
+++ b/Engine/Source/Editor/LevelEditor/Public/LevelViewportLayout.h
@@ -9,7 +9,8 @@
 #include "LevelEditor.h"
 #include "LevelViewportActions.h"
 
-class SLevelViewport;
+class ILevelEditor;
+class SLevelEditor;
 
 /** Interface that defines an entity within a viewport layout */
 class IViewportLayoutEntity : public TSharedFromThis<IViewportLayoutEntity>
@@ -85,7 +86,7 @@ public:
 	 * @param LayoutString			The layout string loaded from file to custom build the layout with
 	 * @param InParentLevelEditor	Optional level editor parent to use for new viewports
 	 */
-	TSharedRef<SWidget> BuildViewportLayout( TSharedPtr<SDockTab> InParentDockTab, TSharedPtr<class FLevelViewportTabContent> InParentTab, const FString& LayoutString, TWeakPtr<class SLevelEditor> InParentLevelEditor = NULL );
+	TSharedRef<SWidget> BuildViewportLayout( TSharedPtr<SDockTab> InParentDockTab, TSharedPtr<class FLevelViewportTabContent> InParentTab, const FString& LayoutString, TWeakPtr<ILevelEditor> InParentLevelEditor = NULL );
 
 	/**
 	 * Makes a request to maximize a specific viewport and hide the others in this layout
@@ -258,7 +259,7 @@ protected:
 	TWeakPtr< class FLevelViewportTabContent > ParentTabContent;
 
 	/** The optional parent level editor for this layout */
-	TWeakPtr< class SLevelEditor > ParentLevelEditor;
+	TWeakPtr< ILevelEditor > ParentLevelEditor;
 
 	/** The current maximized viewport if any */
 	FName MaximizedViewport;
diff --git a/Engine/Source/Editor/LevelEditor/Public/SLevelViewport.h b/Engine/Source/Editor/LevelEditor/Public/SLevelViewport.h
--- a/Engine/Source/Editor/LevelEditor/Public/SLevelViewport.h
+++ b/Engine/Source/Editor/LevelEditor/Public/SLevelViewport.h
@@ -8,6 +8,7 @@
 
 class SActorPreview;
 class SGameLayerManager;
+class ILevelEditor;
 
 /**
  * Encapsulates an SViewport and an SLevelViewportToolBar
@@ -22,7 +23,7 @@ public:
 
 		SLATE_ARGUMENT( TWeakPtr<class FEditorModeTools>, EditorModeTools )
 		SLATE_ARGUMENT( TSharedPtr<class FLevelViewportLayout>, ParentLayout )
-		SLATE_ARGUMENT( TWeakPtr<class ILevelEditor>, ParentLevelEditor )
+		SLATE_ARGUMENT( TWeakPtr<ILevelEditor>, ParentLevelEditor )
 		SLATE_ARGUMENT( ELevelViewportType, ViewportType )
 		SLATE_ARGUMENT( bool, Realtime )
 		SLATE_ARGUMENT( FString, ConfigKey )
@@ -259,7 +260,7 @@ public:
 	FString GetDeviceProfileString( ) const;
 
 	/** Get the parent level editor for this viewport */
-	TWeakPtr<class SLevelEditor> GetParentLevelEditor() const { return ParentLevelEditor; }
+	TWeakPtr<ILevelEditor> GetParentLevelEditor() const { return ParentLevelEditor; }
 
 	/** Called to get the level text */
 	FText GetCurrentLevelText( bool bDrawOnlyLabel ) const;
@@ -642,7 +643,7 @@ private:
 	TWeakPtr<class FLevelViewportLayout> ParentLayout;
 
 	/** Pointer to the parent level editor for this viewport */
-	TWeakPtr<class SLevelEditor> ParentLevelEditor;
+	TWeakPtr<ILevelEditor> ParentLevelEditor;
 
 	/** Viewport overlay widget exposed to game systems when running play-in-editor */
 	TSharedPtr<SOverlay> PIEViewportOverlayWidget;
