VR Editor: Removed unused code (no effective change)

diff --git a/Engine/Config/ConsoleVariables.ini b/Engine/Config/ConsoleVariables.ini
--- a/Engine/Config/ConsoleVariables.ini
+++ b/Engine/Config/ConsoleVariables.ini
@@ -38,5 +38,4 @@
 ; If Linux editor crashes in FMallocBinned with callstack that mentions MeshUtilities, you may need to uncomment this.
 ;r.TriangleOrderOptimization=2
 
-; @todo vreditor: Consider disabling for VR editing.  Too high of an MSAA count prevents the editor from running at a decent framerate in vr
 ; r.MSAA.CompositingSampleCount=1
\ No newline at end of file
diff --git a/Engine/Source/Editor/LevelEditor/Private/LevelEditorActions.cpp b/Engine/Source/Editor/LevelEditor/Private/LevelEditorActions.cpp
--- a/Engine/Source/Editor/LevelEditor/Private/LevelEditorActions.cpp
+++ b/Engine/Source/Editor/LevelEditor/Private/LevelEditorActions.cpp
@@ -3083,7 +3083,7 @@ void FLevelEditorCommands::RegisterCommands()
 	UI_COMMAND( AddMatinee, "Add Matinee", "Creates a new matinee actor to edit", EUserInterfaceActionType::Button, FInputChord() );
 	UI_COMMAND( EditMatinee, "Edit Matinee", "Selects a Matinee to edit", EUserInterfaceActionType::Button, FInputChord() );
 
-	UI_COMMAND( ToggleVR, "Toggle VR", "Toggles VR (Virtual Reality) mode", EUserInterfaceActionType::ToggleButton, FInputChord( EModifierKey::Alt, EKeys::Tilde ) );	// @todo vreditor: Make sure this hotkey is OK
+	UI_COMMAND( ToggleVR, "Toggle VR", "Toggles VR (Virtual Reality) mode", EUserInterfaceActionType::ToggleButton, FInputChord( EModifierKey::Alt, EKeys::Tilde ) );
 
 	UI_COMMAND( OpenLevelBlueprint, "Open Level Blueprint", "Edit the Level Blueprint for the current level", EUserInterfaceActionType::Button, FInputChord() );
 	UI_COMMAND( CheckOutProjectSettingsConfig, "Check Out", "Checks out the project settings config file so the game mode can be set.", EUserInterfaceActionType::Button, FInputChord() );
diff --git a/Engine/Source/Editor/VREditor/VirtualHand.h b/Engine/Source/Editor/VREditor/VirtualHand.h
--- a/Engine/Source/Editor/VREditor/VirtualHand.h
+++ b/Engine/Source/Editor/VREditor/VirtualHand.h
@@ -5,30 +5,6 @@
 #include "VREditorWorldInteractionTypes.h"
 
 
-// @todo vreditor: Horrible that we have to roll our own input events.  This should all be handled by the plugin
-struct FSTEMControllerState
-{
-	int Buttons;
-	int ButtonsPrevious;
-	float TriggerVal;
-	float LastTriggerVal;
-	float JoystickX;
-	float JoystickY;
-	FVector2D LastJoystickVal;
-
-
-	FSTEMControllerState()
-	{
-		Buttons = 0;
-		ButtonsPrevious = 0;
-		TriggerVal = 0;
-		LastTriggerVal = 0;
-		JoystickX = 0;
-		JoystickY = 0;
-		LastJoystickVal = FVector2D();
-	}
-};
-
 
 /** Represents a single virtual hand */
 struct FVirtualHand
@@ -58,9 +34,6 @@ struct FVirtualHand
 	/** True if this hand has a motion controller (or both!) */
 	bool bHaveMotionController;
 
-	/** STEM controller state for this hand */
-	FSTEMControllerState STEMControllerState;
-
 
 	//
 	// Graphics
