Copying //UE4/Dev-Framework to //UE4/Dev-Main (Source: //UE4/Dev-Framework @ 2972815)

diff --git a/Engine/Source/Editor/LandscapeEditor/Private/Tests/LandscapeAutomationTests.cpp b/Engine/Source/Editor/LandscapeEditor/Private/Tests/LandscapeAutomationTests.cpp
--- a/Engine/Source/Editor/LandscapeEditor/Private/Tests/LandscapeAutomationTests.cpp
+++ b/Engine/Source/Editor/LandscapeEditor/Private/Tests/LandscapeAutomationTests.cpp
@@ -160,8 +160,8 @@ bool FLandscapeEditorTest::RunTest(const FString& Parameters)
 
 		//Create the path
 		ScreenshotParameters.CurrentWindow = AllWindows[0];
-		const FString TestName = TEXT("NewLandscapeTest");
-		FString PathName = FPaths::AutomationDir() + TestName / FPlatformProperties::PlatformName();
+		const FString LandscapeTestName = TEXT("NewLandscapeTest");
+		FString PathName = FPaths::AutomationDir() + LandscapeTestName / FPlatformProperties::PlatformName();
 		FPaths::MakePathRelativeTo(PathName, *FPaths::RootDir());
 		ScreenshotParameters.ScreenshotName = FString::Printf(TEXT("%s/%d.png"), *PathName, FEngineVersion::Current().GetChangelist());
 
diff --git a/Engine/Source/Editor/LevelEditor/Private/SLevelViewport.cpp b/Engine/Source/Editor/LevelEditor/Private/SLevelViewport.cpp
--- a/Engine/Source/Editor/LevelEditor/Private/SLevelViewport.cpp
+++ b/Engine/Source/Editor/LevelEditor/Private/SLevelViewport.cpp
@@ -1123,34 +1123,34 @@ void SLevelViewport::BindCommands()
 	UICommandListRef.SetCanProduceActionForCommand( FUICommandList::FCanProduceActionForCommand::CreateSP(this, &SLevelViewport::CanProduceActionForCommand) );
 }
 	
-void SLevelViewport::BindOptionCommands( FUICommandList& CommandList )
+void SLevelViewport::BindOptionCommands( FUICommandList& OutCommandList )
 {
 	const FLevelViewportCommands& ViewportActions = FLevelViewportCommands::Get();
 
-	CommandList.MapAction( 
+	OutCommandList.MapAction( 
 		ViewportActions.AdvancedSettings,
 		FExecuteAction::CreateSP( this, &SLevelViewport::OnAdvancedSettings ) );
 
-	CommandList.MapAction( 
+	OutCommandList.MapAction( 
 		ViewportActions.ToggleMaximize,
 		FExecuteAction::CreateSP( this, &SLevelViewport::OnToggleMaximizeMode ),
 		FCanExecuteAction::CreateSP( this, &SLevelViewport::CanToggleMaximizeMode ) );
 
 	
-	CommandList.MapAction(
+	OutCommandList.MapAction(
 		ViewportActions.ToggleGameView,
 		FExecuteAction::CreateSP( this, &SLevelViewport::ToggleGameView ),
 		FCanExecuteAction::CreateSP( this, &SLevelViewport::CanToggleGameView ),
 		FIsActionChecked::CreateSP( this, &SLevelViewport::IsInGameView ) );
 
-	CommandList.MapAction(
+	OutCommandList.MapAction(
 		ViewportActions.ToggleImmersive,
 		FExecuteAction::CreateSP( this, &SLevelViewport::OnToggleImmersive),
 		FCanExecuteAction(),
 		FIsActionChecked::CreateSP( this, &SLevelViewport::IsImmersive ) );
 
 
-	CommandList.MapAction(
+	OutCommandList.MapAction(
 		ViewportActions.ToggleCinematicPreview,
 		FExecuteAction::CreateSP( this, &SLevelViewport::OnToggleAllowCinematicPreview ),
 		FCanExecuteAction(),
@@ -1158,20 +1158,20 @@ void SLevelViewport::BindOptionCommands( FUICommandList& CommandList )
 		);
 
 
-	CommandList.MapAction(
+	OutCommandList.MapAction(
 		ViewportActions.CreateCamera,
 		FExecuteAction::CreateSP( this, &SLevelViewport::OnCreateCameraActor ),
 		FCanExecuteAction(),
 		FCanExecuteAction::CreateSP( this, &SLevelViewport::IsPerspectiveViewport ) 
 		);
 
-	CommandList.MapAction(
+	OutCommandList.MapAction(
 		ViewportActions.HighResScreenshot,
 		FExecuteAction::CreateSP( this, &SLevelViewport::OnTakeHighResScreenshot ),
 		FCanExecuteAction()
 		);
 
-	CommandList.MapAction(
+	OutCommandList.MapAction(
 		ViewportActions.ToggleActorPilotCameraView,
 		FExecuteAction::CreateSP(this, &SLevelViewport::ToggleActorPilotCameraView),
 		FCanExecuteAction(),
@@ -1181,28 +1181,28 @@ void SLevelViewport::BindOptionCommands( FUICommandList& CommandList )
 	// Map each bookmark action
 	for( int32 BookmarkIndex = 0; BookmarkIndex < AWorldSettings::MAX_BOOKMARK_NUMBER; ++BookmarkIndex )
 	{
-		CommandList.MapAction( 
+		OutCommandList.MapAction( 
 			ViewportActions.JumpToBookmarkCommands[BookmarkIndex],
 			FExecuteAction::CreateSP( this, &SLevelViewport::OnJumpToBookmark, BookmarkIndex )
 			);
 
-		CommandList.MapAction( 
+		OutCommandList.MapAction( 
 			ViewportActions.SetBookmarkCommands[BookmarkIndex],
 			FExecuteAction::CreateSP( this, &SLevelViewport::OnSetBookmark, BookmarkIndex )
 			);
 
-		CommandList.MapAction( 
+		OutCommandList.MapAction( 
 			ViewportActions.ClearBookmarkCommands[BookmarkIndex],
 			FExecuteAction::CreateSP( this, &SLevelViewport::OnClearBookMark, BookmarkIndex )
 			);
 	}
 
-	CommandList.MapAction(
+	OutCommandList.MapAction(
 		ViewportActions.ClearAllBookMarks,
 		FExecuteAction::CreateSP( this, &SLevelViewport::OnClearAllBookMarks )
 		);
 
-	CommandList.MapAction(
+	OutCommandList.MapAction(
 		ViewportActions.ToggleViewportToolbar,
 		FExecuteAction::CreateSP( this, &SLevelViewport::OnToggleShowFullToolbar ),
 		FCanExecuteAction(),
@@ -1210,95 +1210,95 @@ void SLevelViewport::BindOptionCommands( FUICommandList& CommandList )
 		);
 }
 
-void SLevelViewport::BindViewCommands( FUICommandList& CommandList )
+void SLevelViewport::BindViewCommands( FUICommandList& OutCommandList )
 {
 	const FLevelViewportCommands& ViewportActions = FLevelViewportCommands::Get();
 
-	CommandList.MapAction(
+	OutCommandList.MapAction(
 		ViewportActions.FindInLevelScriptBlueprint,
 		FExecuteAction::CreateSP( this, &SLevelViewport::FindSelectedInLevelScript ),
 		FCanExecuteAction::CreateSP( this, &SLevelViewport::CanFindSelectedInLevelScript ) 
 		);
 
-	CommandList.MapAction(
+	OutCommandList.MapAction(
 		ViewportActions.EjectActorPilot,
 		FExecuteAction::CreateSP( this, &SLevelViewport::OnActorUnlock ),
 		FCanExecuteAction::CreateSP( this, &SLevelViewport::CanExecuteActorUnlock )
 		);
 
-	CommandList.MapAction(
+	OutCommandList.MapAction(
 		ViewportActions.PilotSelectedActor,
 		FExecuteAction::CreateSP( this, &SLevelViewport::OnActorLockSelected ),
 		FCanExecuteAction::CreateSP( this, &SLevelViewport::CanExecuteActorLockSelected )
 		);
 
-	CommandList.MapAction(
+	OutCommandList.MapAction(
 		ViewportActions.ViewportConfig_OnePane,
 		FExecuteAction::CreateSP(this, &SLevelViewport::OnSetViewportConfiguration, LevelViewportConfigurationNames::OnePane),
 		FCanExecuteAction(),
 		FIsActionChecked::CreateSP(this, &SLevelViewport::IsViewportConfigurationSet, LevelViewportConfigurationNames::OnePane));
 
-	CommandList.MapAction(
+	OutCommandList.MapAction(
 		ViewportActions.ViewportConfig_TwoPanesH,
 		FExecuteAction::CreateSP( this, &SLevelViewport::OnSetViewportConfiguration, LevelViewportConfigurationNames::TwoPanesHoriz ),
 		FCanExecuteAction(),
 		FIsActionChecked::CreateSP( this, &SLevelViewport::IsViewportConfigurationSet, LevelViewportConfigurationNames::TwoPanesHoriz ));
 
-	CommandList.MapAction(
+	OutCommandList.MapAction(
 		ViewportActions.ViewportConfig_TwoPanesV,
 		FExecuteAction::CreateSP( this, &SLevelViewport::OnSetViewportConfiguration, LevelViewportConfigurationNames::TwoPanesVert ),
 		FCanExecuteAction(),
 		FIsActionChecked::CreateSP( this, &SLevelViewport::IsViewportConfigurationSet, LevelViewportConfigurationNames::TwoPanesVert ));
 
-	CommandList.MapAction(
+	OutCommandList.MapAction(
 		ViewportActions.ViewportConfig_ThreePanesLeft,
 		FExecuteAction::CreateSP( this, &SLevelViewport::OnSetViewportConfiguration, LevelViewportConfigurationNames::ThreePanesLeft ),
 		FCanExecuteAction(),
 		FIsActionChecked::CreateSP( this, &SLevelViewport::IsViewportConfigurationSet, LevelViewportConfigurationNames::ThreePanesLeft ));
 
-	CommandList.MapAction(
+	OutCommandList.MapAction(
 		ViewportActions.ViewportConfig_ThreePanesRight,
 		FExecuteAction::CreateSP( this, &SLevelViewport::OnSetViewportConfiguration, LevelViewportConfigurationNames::ThreePanesRight ),
 		FCanExecuteAction(),
 		FIsActionChecked::CreateSP( this, &SLevelViewport::IsViewportConfigurationSet, LevelViewportConfigurationNames::ThreePanesRight ));
 
-	CommandList.MapAction(
+	OutCommandList.MapAction(
 		ViewportActions.ViewportConfig_ThreePanesTop,
 		FExecuteAction::CreateSP( this, &SLevelViewport::OnSetViewportConfiguration, LevelViewportConfigurationNames::ThreePanesTop ),
 		FCanExecuteAction(),
 		FIsActionChecked::CreateSP( this, &SLevelViewport::IsViewportConfigurationSet, LevelViewportConfigurationNames::ThreePanesTop ));
 
-	CommandList.MapAction(
+	OutCommandList.MapAction(
 		ViewportActions.ViewportConfig_ThreePanesBottom,
 		FExecuteAction::CreateSP( this, &SLevelViewport::OnSetViewportConfiguration, LevelViewportConfigurationNames::ThreePanesBottom ),
 		FCanExecuteAction(),
 		FIsActionChecked::CreateSP( this, &SLevelViewport::IsViewportConfigurationSet, LevelViewportConfigurationNames::ThreePanesBottom ));
 
-	CommandList.MapAction(
+	OutCommandList.MapAction(
 		ViewportActions.ViewportConfig_FourPanesLeft,
 		FExecuteAction::CreateSP( this, &SLevelViewport::OnSetViewportConfiguration, LevelViewportConfigurationNames::FourPanesLeft ),
 		FCanExecuteAction(),
 		FIsActionChecked::CreateSP( this, &SLevelViewport::IsViewportConfigurationSet, LevelViewportConfigurationNames::FourPanesLeft ));
 
-	CommandList.MapAction(
+	OutCommandList.MapAction(
 		ViewportActions.ViewportConfig_FourPanesRight,
 		FExecuteAction::CreateSP( this, &SLevelViewport::OnSetViewportConfiguration, LevelViewportConfigurationNames::FourPanesRight ),
 		FCanExecuteAction(),
 		FIsActionChecked::CreateSP( this, &SLevelViewport::IsViewportConfigurationSet, LevelViewportConfigurationNames::FourPanesRight ));
 
-	CommandList.MapAction(
+	OutCommandList.MapAction(
 		ViewportActions.ViewportConfig_FourPanesTop,
 		FExecuteAction::CreateSP( this, &SLevelViewport::OnSetViewportConfiguration, LevelViewportConfigurationNames::FourPanesTop ),
 		FCanExecuteAction(),
 		FIsActionChecked::CreateSP( this, &SLevelViewport::IsViewportConfigurationSet, LevelViewportConfigurationNames::FourPanesTop ));
 
-	CommandList.MapAction(
+	OutCommandList.MapAction(
 		ViewportActions.ViewportConfig_FourPanesBottom,
 		FExecuteAction::CreateSP( this, &SLevelViewport::OnSetViewportConfiguration, LevelViewportConfigurationNames::FourPanesBottom ),
 		FCanExecuteAction(),
 		FIsActionChecked::CreateSP( this, &SLevelViewport::IsViewportConfigurationSet, LevelViewportConfigurationNames::FourPanesBottom ));
 
-	CommandList.MapAction(
+	OutCommandList.MapAction(
 		ViewportActions.ViewportConfig_FourPanes2x2,
 		FExecuteAction::CreateSP( this, &SLevelViewport::OnSetViewportConfiguration, LevelViewportConfigurationNames::FourPanes2x2 ),
 		FCanExecuteAction(),
@@ -1307,7 +1307,7 @@ void SLevelViewport::BindViewCommands( FUICommandList& CommandList )
 	auto ProcessViewportTypeActions = [&](FName InViewportTypeName, const FViewportTypeDefinition& InDefinition){
 		if (InDefinition.ActivationCommand.IsValid())
 		{
-			CommandList.MapAction(InDefinition.ActivationCommand, FUIAction(
+			OutCommandList.MapAction(InDefinition.ActivationCommand, FUIAction(
 				FExecuteAction::CreateSP(this, &SLevelViewport::ToggleViewportTypeActivationWithinLayout, InViewportTypeName),
 				FCanExecuteAction(),
 				FIsActionChecked::CreateSP(this, &SLevelViewport::IsViewportTypeWithinLayoutEqual, InViewportTypeName)
@@ -1321,7 +1321,7 @@ void SLevelViewport::BindViewCommands( FUICommandList& CommandList )
 	for (FLevelViewportCommands::TBufferVisualizationModeCommandMap::TConstIterator It = ViewportActions.BufferVisualizationModeCommands.CreateConstIterator(); It; ++It)
 	{
 		const FLevelViewportCommands::FBufferVisualizationRecord& Record = It.Value();
-		CommandList.MapAction(
+		OutCommandList.MapAction(
 			Record.Command,
 			FExecuteAction::CreateSP( this, &SLevelViewport::ChangeBufferVisualizationMode, Record.Name ),
 			FCanExecuteAction(),
@@ -1330,9 +1330,9 @@ void SLevelViewport::BindViewCommands( FUICommandList& CommandList )
 }
 
 
-void SLevelViewport::BindShowCommands( FUICommandList& CommandList )
+void SLevelViewport::BindShowCommands( FUICommandList& OutCommandList )
 {
-	CommandList.MapAction( 
+	OutCommandList.MapAction( 
 		FLevelViewportCommands::Get().UseDefaultShowFlags,
 		FExecuteAction::CreateSP( this, &SLevelViewport::OnUseDefaultShowFlags, false ) );
 
@@ -1344,7 +1344,7 @@ void SLevelViewport::BindShowCommands( FUICommandList& CommandList )
 		const FShowFlagData& SFData = ShowFlagData[ShowFlag];
 
 		// NOTE: There should be one command per show flag so using ShowFlag as the index to ShowFlagCommands is acceptable
-		CommandList.MapAction(
+		OutCommandList.MapAction(
 			FLevelViewportCommands::Get().ShowFlagCommands[ ShowFlag ].ShowMenuItem,
 			FExecuteAction::CreateSP( this, &SLevelViewport::ToggleShowFlag, SFData.EngineShowFlagIndex ),
 			FCanExecuteAction(),
@@ -1354,11 +1354,11 @@ void SLevelViewport::BindShowCommands( FUICommandList& CommandList )
 	// Show Volumes
 	{
 		// Map 'Show All' and 'Hide All' commands
-		CommandList.MapAction(
+		OutCommandList.MapAction(
 			FLevelViewportCommands::Get().ShowAllVolumes,
 			FExecuteAction::CreateSP( this, &SLevelViewport::OnToggleAllVolumeActors, true ) );
 
-		CommandList.MapAction(
+		OutCommandList.MapAction(
 			FLevelViewportCommands::Get().HideAllVolumes,
 			FExecuteAction::CreateSP( this, &SLevelViewport::OnToggleAllVolumeActors, false ) );
 
@@ -1368,7 +1368,7 @@ void SLevelViewport::BindShowCommands( FUICommandList& CommandList )
 
 		for( int32 VolumeClassIndex = 0; VolumeClassIndex < VolumeClasses.Num(); ++VolumeClassIndex )
 		{
-			CommandList.MapAction(
+			OutCommandList.MapAction(
 				FLevelViewportCommands::Get().ShowVolumeCommands[ VolumeClassIndex ].ShowMenuItem,
 				FExecuteAction::CreateSP( this, &SLevelViewport::ToggleShowVolumeClass, VolumeClassIndex ),
 				FCanExecuteAction(),
@@ -1379,11 +1379,11 @@ void SLevelViewport::BindShowCommands( FUICommandList& CommandList )
 	// Show Layers
 	{
 		// Map 'Show All' and 'Hide All' commands
-		CommandList.MapAction(
+		OutCommandList.MapAction(
 			FLevelViewportCommands::Get().ShowAllLayers,
 			FExecuteAction::CreateSP( this, &SLevelViewport::OnToggleAllLayers, true ) );
 
-		CommandList.MapAction(
+		OutCommandList.MapAction(
 			FLevelViewportCommands::Get().HideAllLayers,
 			FExecuteAction::CreateSP( this, &SLevelViewport::OnToggleAllLayers, false ) );
 	}
@@ -1391,18 +1391,18 @@ void SLevelViewport::BindShowCommands( FUICommandList& CommandList )
 	// Show Sprite Categories
 	{
 		// Map 'Show All' and 'Hide All' commands
-		CommandList.MapAction(
+		OutCommandList.MapAction(
 			FLevelViewportCommands::Get().ShowAllSprites,
 			FExecuteAction::CreateSP( this, &SLevelViewport::OnToggleAllSpriteCategories, true ) );
 
-		CommandList.MapAction(
+		OutCommandList.MapAction(
 			FLevelViewportCommands::Get().HideAllSprites,
 			FExecuteAction::CreateSP( this, &SLevelViewport::OnToggleAllSpriteCategories, false ) );
 
 		// Bind each show flag to the same delegate.  We use the delegate payload system to figure out what show flag we are dealing with
 		for( int32 CategoryIndex = 0; CategoryIndex < GUnrealEd->SpriteIDToIndexMap.Num(); ++CategoryIndex )
 		{
-			CommandList.MapAction(
+			OutCommandList.MapAction(
 				FLevelViewportCommands::Get().ShowSpriteCommands[ CategoryIndex ].ShowMenuItem,
 				FExecuteAction::CreateSP( this, &SLevelViewport::ToggleSpriteCategory, CategoryIndex ),
 				FCanExecuteAction(),
@@ -1413,7 +1413,7 @@ void SLevelViewport::BindShowCommands( FUICommandList& CommandList )
 	// Show Stat Categories
 	{
 		// Map 'Hide All' command
-		CommandList.MapAction(
+		OutCommandList.MapAction(
 			FLevelViewportCommands::Get().HideAllStats,
 			FExecuteAction::CreateSP(this, &SLevelViewport::OnToggleAllStatCommands, false));
 
@@ -1432,9 +1432,9 @@ void SLevelViewport::BindShowCommands( FUICommandList& CommandList )
 	}
 }
 
-void SLevelViewport::BindDropCommands( FUICommandList& CommandList )
+void SLevelViewport::BindDropCommands( FUICommandList& OutCommandList )
 {
-	CommandList.MapAction( 
+	OutCommandList.MapAction( 
 		FLevelViewportCommands::Get().ApplyMaterialToActor,
 		FExecuteAction::CreateSP( this, &SLevelViewport::OnApplyMaterialToViewportTarget ) );
 }
@@ -1497,15 +1497,15 @@ const FSlateBrush* SLevelViewport::OnGetViewportBorderBrush() const
 
 FSlateColor SLevelViewport::OnGetViewportBorderColorAndOpacity() const
 {
-	auto ColorAndOpacity = FLinearColor::White;
+	FLinearColor ViewportBorderColorAndOpacity = FLinearColor::White;
 	if( FSlateApplication::Get().IsNormalExecution() )
 	{
 		if( ViewTransitionAnim.IsPlaying() )
 		{
-			ColorAndOpacity = FLinearColor( 1.0f, 1.0f, 1.0f, 1.0f - ViewTransitionAnim.GetLerp() );
+			ViewportBorderColorAndOpacity = FLinearColor( 1.0f, 1.0f, 1.0f, 1.0f - ViewTransitionAnim.GetLerp() );
 		}
 	}
-	return ColorAndOpacity;
+	return ViewportBorderColorAndOpacity;
 }
 
 EVisibility SLevelViewport::OnGetViewportContentVisibility() const
@@ -3779,13 +3779,13 @@ EVisibility SLevelViewport::GetLockedIconVisibility() const
 
 FText SLevelViewport::GetLockedIconToolTip() const
 {
-	FText ToolTip;
+	FText ToolTipText;
 	if (IsAnyActorLocked())
 	{
-		ToolTip = FText::Format( LOCTEXT("ActorLockedIcon_ToolTip", "Viewport Locked to {0}"), FText::FromString( LevelViewportClient->GetActiveActorLock().Get()->GetActorLabel() ) );
+		ToolTipText = FText::Format( LOCTEXT("ActorLockedIcon_ToolTip", "Viewport Locked to {0}"), FText::FromString( LevelViewportClient->GetActiveActorLock().Get()->GetActorLabel() ) );
 	}
 
-	return ToolTip;
+	return ToolTipText;
 }
 
 UWorld* SLevelViewport::GetWorld() const
diff --git a/Engine/Source/Editor/MaterialEditor/Private/SMaterialPalette.cpp b/Engine/Source/Editor/MaterialEditor/Private/SMaterialPalette.cpp
--- a/Engine/Source/Editor/MaterialEditor/Private/SMaterialPalette.cpp
+++ b/Engine/Source/Editor/MaterialEditor/Private/SMaterialPalette.cpp
@@ -38,8 +38,7 @@ void SMaterialPaletteItem::Construct(const FArguments& InArgs, FCreateWidgetForA
 	// Find icons
 	const FSlateBrush* IconBrush = FEditorStyle::GetBrush(TEXT("NoBrush"));
 	FSlateColor IconColor = FSlateColor::UseForeground();
-	FString ToolTip = GraphAction->GetTooltipDescription();
-	FText IconToolTip = FText::FromString(ToolTip);
+	FText IconToolTip = FText::FromString(GraphAction->GetTooltipDescription());
 	bool bIsReadOnly = false;
 
 	TSharedRef<SWidget> IconWidget = CreateIconWidget( IconToolTip, IconBrush, IconColor );
diff --git a/Engine/Source/Editor/MeshPaint/Private/MeshPaintEdMode.cpp b/Engine/Source/Editor/MeshPaint/Private/MeshPaintEdMode.cpp
--- a/Engine/Source/Editor/MeshPaint/Private/MeshPaintEdMode.cpp
+++ b/Engine/Source/Editor/MeshPaint/Private/MeshPaintEdMode.cpp
@@ -3855,14 +3855,14 @@ void FEdModeMeshPaint::PushInstanceVertexColorsToMesh()
 	const bool bMeshSelected = GetSelectedMeshInfo( NumBaseVertexColorBytes, NumInstanceVertexColorBytes, bHasInstanceMaterialAndTexture );
 	if ( bMeshSelected && NumInstanceVertexColorBytes > 0 )
 	{
-		FSuppressableWarningDialog::FSetupInfo Info( LOCTEXT("PushInstanceVertexColorsPrompt_Message", "Copying the instance vertex colors to the source mesh will replace any of the source mesh's pre-existing vertex colors and affect every instance of the source mesh." ),
+		FSuppressableWarningDialog::FSetupInfo SetupInfo( LOCTEXT("PushInstanceVertexColorsPrompt_Message", "Copying the instance vertex colors to the source mesh will replace any of the source mesh's pre-existing vertex colors and affect every instance of the source mesh." ),
 													 LOCTEXT("PushInstanceVertexColorsPrompt_Title", "Warning: Copying vertex data overwrites all instances" ), "Warning_PushInstanceVertexColorsPrompt" );
 
-		Info.ConfirmText = LOCTEXT("PushInstanceVertexColorsPrompt_ConfirmText", "Continue");
-		Info.CancelText = LOCTEXT("PushInstanceVertexColorsPrompt_CancelText", "Abort");
-		Info.CheckBoxText = LOCTEXT("PushInstanceVertexColorsPrompt_CheckBoxText","Always copy vertex colors without prompting");
+		SetupInfo.ConfirmText = LOCTEXT("PushInstanceVertexColorsPrompt_ConfirmText", "Continue");
+		SetupInfo.CancelText = LOCTEXT("PushInstanceVertexColorsPrompt_CancelText", "Abort");
+		SetupInfo.CheckBoxText = LOCTEXT("PushInstanceVertexColorsPrompt_CheckBoxText","Always copy vertex colors without prompting");
 
-		FSuppressableWarningDialog VertexColorCopyWarning( Info );
+		FSuppressableWarningDialog VertexColorCopyWarning( SetupInfo );
 
 		// Prompt the user to see if they really want to push the vert colors to the source mesh and to explain
 		// the ramifications of doing so. This uses a suppressible dialog so that the user has the choice to always ignore the warning.
diff --git a/Engine/Source/Editor/MovieSceneTools/Private/TrackEditors/MaterialTrackEditor.cpp b/Engine/Source/Editor/MovieSceneTools/Private/TrackEditors/MaterialTrackEditor.cpp
--- a/Engine/Source/Editor/MovieSceneTools/Private/TrackEditors/MaterialTrackEditor.cpp
+++ b/Engine/Source/Editor/MovieSceneTools/Private/TrackEditors/MaterialTrackEditor.cpp
@@ -175,13 +175,13 @@ bool FComponentMaterialTrackEditor::SupportsType( TSubclassOf<UMovieSceneTrack>
 
 UMaterialInterface* FComponentMaterialTrackEditor::GetMaterialInterfaceForTrack( FGuid ObjectBinding, UMovieSceneMaterialTrack* MaterialTrack )
 {
-	TSharedPtr<ISequencer> Sequencer = GetSequencer();
-	if (!Sequencer.IsValid())
+	TSharedPtr<ISequencer> SequencerPtr = GetSequencer();
+	if (!SequencerPtr.IsValid())
 	{
 		return nullptr;
 	}
 
-	UPrimitiveComponent* Component = Cast<UPrimitiveComponent>(Sequencer->FindSpawnedObjectOrTemplate( ObjectBinding ));
+	UPrimitiveComponent* Component = Cast<UPrimitiveComponent>(SequencerPtr->FindSpawnedObjectOrTemplate( ObjectBinding ));
 	UMovieSceneComponentMaterialTrack* ComponentMaterialTrack = Cast<UMovieSceneComponentMaterialTrack>( MaterialTrack );
 	if ( Component != nullptr && ComponentMaterialTrack != nullptr )
 	{
diff --git a/Engine/Source/Editor/MovieSceneTools/Private/TrackEditors/ParticleParameterTrackEditor.cpp b/Engine/Source/Editor/MovieSceneTools/Private/TrackEditors/ParticleParameterTrackEditor.cpp
--- a/Engine/Source/Editor/MovieSceneTools/Private/TrackEditors/ParticleParameterTrackEditor.cpp
+++ b/Engine/Source/Editor/MovieSceneTools/Private/TrackEditors/ParticleParameterTrackEditor.cpp
@@ -73,8 +73,8 @@ TSharedRef<SWidget> FParticleParameterTrackEditor::OnGetAddParameterMenuContent(
 {
 	FMenuBuilder AddParameterMenuBuilder( true, nullptr );
 
-	TSharedPtr<ISequencer> Sequencer = GetSequencer();
-	AEmitter* Emitter = Sequencer.IsValid() ? Cast<AEmitter>( Sequencer->FindSpawnedObjectOrTemplate( ObjectBinding ) ) : nullptr;
+	TSharedPtr<ISequencer> SequencerPtr = GetSequencer();
+	AEmitter* Emitter = SequencerPtr.IsValid() ? Cast<AEmitter>( SequencerPtr->FindSpawnedObjectOrTemplate( ObjectBinding ) ) : nullptr;
 
 	if ( Emitter != nullptr )
 	{
diff --git a/Engine/Source/Editor/MovieSceneTools/Private/TrackEditors/ParticleTrackEditor.cpp b/Engine/Source/Editor/MovieSceneTools/Private/TrackEditors/ParticleTrackEditor.cpp
--- a/Engine/Source/Editor/MovieSceneTools/Private/TrackEditors/ParticleTrackEditor.cpp
+++ b/Engine/Source/Editor/MovieSceneTools/Private/TrackEditors/ParticleTrackEditor.cpp
@@ -304,8 +304,8 @@ void FParticleTrackEditor::BuildObjectBindingTrackMenu(FMenuBuilder& MenuBuilder
 
 void FParticleTrackEditor::AddParticleKey( const FGuid ObjectGuid )
 {
-	TSharedPtr<ISequencer> Sequencer = GetSequencer();
-	UObject* Object = Sequencer.IsValid() ? Sequencer->FindSpawnedObjectOrTemplate(ObjectGuid) : nullptr;
+	TSharedPtr<ISequencer> SequencerPtr = GetSequencer();
+	UObject* Object = SequencerPtr.IsValid() ? SequencerPtr->FindSpawnedObjectOrTemplate(ObjectGuid) : nullptr;
 
 	if (Object)
 	{
diff --git a/Engine/Source/Editor/MovieSceneTools/Private/TrackEditors/SkeletalAnimationTrackEditor.cpp b/Engine/Source/Editor/MovieSceneTools/Private/TrackEditors/SkeletalAnimationTrackEditor.cpp
--- a/Engine/Source/Editor/MovieSceneTools/Private/TrackEditors/SkeletalAnimationTrackEditor.cpp
+++ b/Engine/Source/Editor/MovieSceneTools/Private/TrackEditors/SkeletalAnimationTrackEditor.cpp
@@ -165,9 +165,9 @@ void FSkeletalAnimationTrackEditor::AddKey(const FGuid& ObjectGuid)
 
 bool FSkeletalAnimationTrackEditor::HandleAssetAdded(UObject* Asset, const FGuid& TargetObjectGuid)
 {
-	TSharedPtr<ISequencer> Sequencer = GetSequencer();
+	TSharedPtr<ISequencer> SequencerPtr = GetSequencer();
 
-	if (Asset->IsA<UAnimSequence>() && Sequencer.IsValid())
+	if (Asset->IsA<UAnimSequence>() && SequencerPtr.IsValid())
 	{
 		UAnimSequence* AnimSequence = Cast<UAnimSequence>(Asset);
 		
@@ -179,7 +179,7 @@ bool FSkeletalAnimationTrackEditor::HandleAssetAdded(UObject* Asset, const FGuid
 			{
 				
 
-				UObject* Object = Sequencer->FindSpawnedObjectOrTemplate(TargetObjectGuid);
+				UObject* Object = SequencerPtr->FindSpawnedObjectOrTemplate(TargetObjectGuid);
 				AnimatablePropertyChanged(FOnKeyProperty::CreateRaw(this, &FSkeletalAnimationTrackEditor::AddKeyInternal, Object, AnimSequence));
 
 				return true;
@@ -256,13 +256,13 @@ void FSkeletalAnimationTrackEditor::OnAnimationAssetSelected(const FAssetData& A
 	FSlateApplication::Get().DismissAllMenus();
 
 	UObject* SelectedObject = AssetData.GetAsset();
-	TSharedPtr<ISequencer> Sequencer = GetSequencer();
+	TSharedPtr<ISequencer> SequencerPtr = GetSequencer();
 
-	if (SelectedObject && SelectedObject->IsA(UAnimSequence::StaticClass()) && Sequencer.IsValid())
+	if (SelectedObject && SelectedObject->IsA(UAnimSequence::StaticClass()) && SequencerPtr.IsValid())
 	{
 		UAnimSequence* AnimSequence = CastChecked<UAnimSequence>(AssetData.GetAsset());
 
-		UObject* Object = Sequencer->FindSpawnedObjectOrTemplate(ObjectBinding);
+		UObject* Object = SequencerPtr->FindSpawnedObjectOrTemplate(ObjectBinding);
 		AnimatablePropertyChanged( FOnKeyProperty::CreateRaw( this, &FSkeletalAnimationTrackEditor::AddKeyInternal, Object, AnimSequence) );
 	}
 }
@@ -296,8 +296,8 @@ bool FSkeletalAnimationTrackEditor::AddKeyInternal( float KeyTime, UObject* Obje
 
 USkeleton* FSkeletalAnimationTrackEditor::AcquireSkeletonFromObjectGuid(const FGuid& Guid)
 {
-	TSharedPtr<ISequencer> Sequencer = GetSequencer();
-	UObject* BoundObject = Sequencer.IsValid() ? Sequencer->FindSpawnedObjectOrTemplate(Guid) : nullptr;
+	TSharedPtr<ISequencer> SequencerPtr = GetSequencer();
+	UObject* BoundObject = SequencerPtr.IsValid() ? SequencerPtr->FindSpawnedObjectOrTemplate(Guid) : nullptr;
 
 	if (AActor* Actor = Cast<AActor>(BoundObject))
 	{
diff --git a/Engine/Source/Editor/MovieSceneTools/Public/PropertyTrackEditor.h b/Engine/Source/Editor/MovieSceneTools/Public/PropertyTrackEditor.h
--- a/Engine/Source/Editor/MovieSceneTools/Public/PropertyTrackEditor.h
+++ b/Engine/Source/Editor/MovieSceneTools/Public/PropertyTrackEditor.h
@@ -90,10 +90,10 @@ public:
 
 	~FPropertyTrackEditor()
 	{
-		TSharedPtr<ISequencer> Sequencer = FMovieSceneTrackEditor::GetSequencer();
-		if ( Sequencer.IsValid() )
+		TSharedPtr<ISequencer> SequencerPtr = FMovieSceneTrackEditor::GetSequencer();
+		if ( SequencerPtr.IsValid() )
 		{
-			ISequencerObjectChangeListener& ObjectChangeListener = Sequencer->GetObjectChangeListener();
+			ISequencerObjectChangeListener& ObjectChangeListener = SequencerPtr->GetObjectChangeListener();
 			for ( FName WatchedPropertyTypeName : WatchedPropertyTypeNames )
 			{
 				ObjectChangeListener.GetOnAnimatablePropertyChanged( WatchedPropertyTypeName ).RemoveAll( this );
diff --git a/Engine/Source/Editor/PListEditor/Private/PListNodeArray.cpp b/Engine/Source/Editor/PListEditor/Private/PListNodeArray.cpp
--- a/Engine/Source/Editor/PListEditor/Private/PListNodeArray.cpp
+++ b/Engine/Source/Editor/PListEditor/Private/PListNodeArray.cpp
@@ -57,7 +57,7 @@ TSharedRef<ITableRow> FPListNodeArray::GenerateWidget(const TSharedRef<STableVie
 }
 
 /** Generate a widget for the specified column name */
-TSharedRef<SWidget> FPListNodeArray::GenerateWidgetForColumn(const FName& ColumnName, int32 Depth, ITableRow* RowPtr)
+TSharedRef<SWidget> FPListNodeArray::GenerateWidgetForColumn(const FName& ColumnName, int32 InDepth, ITableRow* RowPtr)
 {
 	TableRow = RowPtr;
 	check(TableRow);
@@ -78,7 +78,7 @@ TSharedRef<SWidget> FPListNodeArray::GenerateWidgetForColumn(const FName& Column
 				+ SHorizontalBox::Slot()
 				[
 					SNew(SSpacer)
-					.Size(FVector2D(20 * Depth, 0))
+					.Size(FVector2D(20 * InDepth, 0))
 				]
 
 				// Editable key value
@@ -111,7 +111,7 @@ TSharedRef<SWidget> FPListNodeArray::GenerateWidgetForColumn(const FName& Column
 				[
 					// Space before expander
 					SNew(SSpacer)
-					.Size(FVector2D(20 * (Depth - 1), 0))
+					.Size(FVector2D(20 * (InDepth - 1), 0))
 				]
 
 				+ SHorizontalBox::Slot()
diff --git a/Engine/Source/Editor/PListEditor/Private/PListNodeBoolean.cpp b/Engine/Source/Editor/PListEditor/Private/PListNodeBoolean.cpp
--- a/Engine/Source/Editor/PListEditor/Private/PListNodeBoolean.cpp
+++ b/Engine/Source/Editor/PListEditor/Private/PListNodeBoolean.cpp
@@ -51,7 +51,7 @@ TSharedRef<ITableRow> FPListNodeBoolean::GenerateWidget(const TSharedRef<STableV
 }
 
 /** Generate a widget for the specified column name */
-TSharedRef<SWidget> FPListNodeBoolean::GenerateWidgetForColumn(const FName& ColumnName, int32 Depth, ITableRow* RowPtr)
+TSharedRef<SWidget> FPListNodeBoolean::GenerateWidgetForColumn(const FName& ColumnName, int32 InDepth, ITableRow* RowPtr)
 {
 	if(ColumnName == "PListKeyColumn")
 	{
@@ -65,7 +65,7 @@ TSharedRef<SWidget> FPListNodeBoolean::GenerateWidgetForColumn(const FName& Colu
 			+ SHorizontalBox::Slot()
 			[
 				SNew(SSpacer)
-				.Size(FVector2D(20 * Depth, 0))
+				.Size(FVector2D(20 * InDepth, 0))
 			]
 
 			// Editable key value
diff --git a/Engine/Source/Editor/PListEditor/Private/PListNodeDictionary.cpp b/Engine/Source/Editor/PListEditor/Private/PListNodeDictionary.cpp
--- a/Engine/Source/Editor/PListEditor/Private/PListNodeDictionary.cpp
+++ b/Engine/Source/Editor/PListEditor/Private/PListNodeDictionary.cpp
@@ -50,7 +50,7 @@ TSharedRef<ITableRow> FPListNodeDictionary::GenerateWidget(const TSharedRef<STab
 }
 
 /** Generate a widget for the specified column name */
-TSharedRef<SWidget> FPListNodeDictionary::GenerateWidgetForColumn(const FName& ColumnName, int32 Depth, ITableRow* RowPtr)
+TSharedRef<SWidget> FPListNodeDictionary::GenerateWidgetForColumn(const FName& ColumnName, int32 InDepth, ITableRow* RowPtr)
 {
 	TableRow = RowPtr;
 	check(TableRow);
@@ -71,7 +71,7 @@ TSharedRef<SWidget> FPListNodeDictionary::GenerateWidgetForColumn(const FName& C
 				+ SHorizontalBox::Slot()
 				[
 					SNew(SSpacer)
-					.Size(FVector2D(20 * Depth, 0))
+					.Size(FVector2D(20 * InDepth, 0))
 				]
 
 				+ SHorizontalBox::Slot()
@@ -100,7 +100,7 @@ TSharedRef<SWidget> FPListNodeDictionary::GenerateWidgetForColumn(const FName& C
 				[
 					// Space before expander
 					SNew(SSpacer)
-					.Size(FVector2D(20 * (Depth - 1), 0))
+					.Size(FVector2D(20 * (InDepth - 1), 0))
 				]
 
 				+ SHorizontalBox::Slot()
diff --git a/Engine/Source/Editor/PListEditor/Private/PListNodeFile.cpp b/Engine/Source/Editor/PListEditor/Private/PListNodeFile.cpp
--- a/Engine/Source/Editor/PListEditor/Private/PListNodeFile.cpp
+++ b/Engine/Source/Editor/PListEditor/Private/PListNodeFile.cpp
@@ -62,7 +62,7 @@ TSharedRef<ITableRow> FPListNodeFile::GenerateWidget(const TSharedRef<STableView
 }
 
 /** Generate a widget for the specified column name */
-TSharedRef<SWidget> FPListNodeFile::GenerateWidgetForColumn(const FName& ColumnName, int32 Depth, ITableRow* RowPtr) 
+TSharedRef<SWidget> FPListNodeFile::GenerateWidgetForColumn(const FName& ColumnName, int32 InDepth, ITableRow* RowPtr) 
 {
 	return GenerateInvalidRow(NSLOCTEXT("PListNodeArray", "PListNodeFileArrayUsesColumns", "PListNodeFile does not use columns"));
 }
diff --git a/Engine/Source/Editor/PListEditor/Private/PListNodeString.cpp b/Engine/Source/Editor/PListEditor/Private/PListNodeString.cpp
--- a/Engine/Source/Editor/PListEditor/Private/PListNodeString.cpp
+++ b/Engine/Source/Editor/PListEditor/Private/PListNodeString.cpp
@@ -59,7 +59,7 @@ TSharedRef<ITableRow> FPListNodeString::GenerateWidget(const TSharedRef<STableVi
 }
 
 /** Generate a widget for the specified column name */
-TSharedRef<SWidget> FPListNodeString::GenerateWidgetForColumn(const FName& ColumnName, int32 Depth, ITableRow* RowPtr)
+TSharedRef<SWidget> FPListNodeString::GenerateWidgetForColumn(const FName& ColumnName, int32 InDepth, ITableRow* RowPtr)
 {
 	if(ColumnName == "PListKeyColumn")
 	{
@@ -73,7 +73,7 @@ TSharedRef<SWidget> FPListNodeString::GenerateWidgetForColumn(const FName& Colum
 			+ SHorizontalBox::Slot()
 			[
 				SNew(SSpacer)
-				.Size(FVector2D(20 * Depth, 0))
+				.Size(FVector2D(20 * InDepth, 0))
 			]
 
 			// Editable key value
diff --git a/Engine/Source/Editor/Persona/Private/AnimGraphNodeDetails.cpp b/Engine/Source/Editor/Persona/Private/AnimGraphNodeDetails.cpp
--- a/Engine/Source/Editor/Persona/Private/AnimGraphNodeDetails.cpp
+++ b/Engine/Source/Editor/Persona/Private/AnimGraphNodeDetails.cpp
@@ -641,23 +641,23 @@ void SParentPlayerTreeRow::Construct(const FArguments& InArgs, const TSharedRef<
 
 TSharedRef<SWidget> SParentPlayerTreeRow::GenerateWidgetForColumn(const FName& ColumnName)
 {
-	TSharedPtr<SHorizontalBox> Box;
-	SAssignNew(Box, SHorizontalBox);
+	TSharedPtr<SHorizontalBox> HorizBox;
+	SAssignNew(HorizBox, SHorizontalBox);
 
 	if(ColumnName == "Name")
 	{
-		Box->AddSlot()
+		HorizBox->AddSlot()
 			.VAlign(VAlign_Center)
 			.AutoWidth()
 			[
 				SNew(SExpanderArrow, SharedThis(this))
 			];
 
-		Item->GenerateNameWidget(Box);
+		Item->GenerateNameWidget(HorizBox);
 	}
 	else if(Item->Override)
 	{
-		Box->AddSlot()
+		HorizBox->AddSlot()
 			.Padding(2)
 			.VAlign(VAlign_Center)
 			.AutoWidth()
@@ -676,7 +676,7 @@ TSharedRef<SWidget> SParentPlayerTreeRow::GenerateWidgetForColumn(const FName& C
 		
 		TArray<const UClass*> AllowedClasses;
 		AllowedClasses.Add(UAnimationAsset::StaticClass());
-		Box->AddSlot()
+		HorizBox->AddSlot()
 			.VAlign(VAlign_Center)
 			.AutoWidth()
 			[
@@ -687,7 +687,7 @@ TSharedRef<SWidget> SParentPlayerTreeRow::GenerateWidgetForColumn(const FName& C
 				.AllowedClass(GetCurrentAssetToUse()->GetClass())
 			];
 
-		Box->AddSlot()
+		HorizBox->AddSlot()
 			.VAlign(VAlign_Center)
 			.AutoWidth()
 			[
@@ -704,7 +704,7 @@ TSharedRef<SWidget> SParentPlayerTreeRow::GenerateWidgetForColumn(const FName& C
 			];
 	}
 
-	return Box.ToSharedRef();
+	return HorizBox.ToSharedRef();
 }
 
 bool SParentPlayerTreeRow::OnShouldFilterAsset(const FAssetData& AssetData)
diff --git a/Engine/Source/Editor/Persona/Private/AnimationEditorViewportClient.cpp b/Engine/Source/Editor/Persona/Private/AnimationEditorViewportClient.cpp
--- a/Engine/Source/Editor/Persona/Private/AnimationEditorViewportClient.cpp
+++ b/Engine/Source/Editor/Persona/Private/AnimationEditorViewportClient.cpp
@@ -1293,12 +1293,12 @@ void FAnimationViewportClient::ProcessClick(class FSceneView& View, class HHitPr
 }
 }
 
-bool FAnimationViewportClient::InputWidgetDelta( FViewport* Viewport, EAxisList::Type CurrentAxis, FVector& Drag, FRotator& Rot, FVector& Scale )
+bool FAnimationViewportClient::InputWidgetDelta( FViewport* InViewport, EAxisList::Type CurrentAxis, FVector& Drag, FRotator& Rot, FVector& Scale )
 {
 	// Get some useful info about buttons being held down
-	const bool bCtrlDown = Viewport->KeyState(EKeys::LeftControl) || Viewport->KeyState(EKeys::RightControl);
-	const bool bShiftDown = Viewport->KeyState(EKeys::LeftShift) || Viewport->KeyState(EKeys::RightShift);
-	const bool bMouseButtonDown = Viewport->KeyState( EKeys::LeftMouseButton ) || Viewport->KeyState( EKeys::MiddleMouseButton ) || Viewport->KeyState( EKeys::RightMouseButton );
+	const bool bCtrlDown = InViewport->KeyState(EKeys::LeftControl) || InViewport->KeyState(EKeys::RightControl);
+	const bool bShiftDown = InViewport->KeyState(EKeys::LeftShift) || InViewport->KeyState(EKeys::RightShift);
+	const bool bMouseButtonDown = InViewport->KeyState( EKeys::LeftMouseButton ) || InViewport->KeyState( EKeys::MiddleMouseButton ) || InViewport->KeyState( EKeys::RightMouseButton );
 
 	bool bHandled = false;
 
@@ -1449,7 +1449,7 @@ bool FAnimationViewportClient::InputWidgetDelta( FViewport* Viewport, EAxisList:
 			}
 		}
 
-		Viewport->Invalidate();
+		InViewport->Invalidate();
 	}
 
 	return bHandled;
@@ -1668,10 +1668,10 @@ void FAnimationViewportClient::SetViewportType(ELevelViewportType InViewportType
 	FocusViewportOnPreviewMesh();
 }
 
-bool FAnimationViewportClient::InputKey( FViewport* Viewport, int32 ControllerId, FKey Key, EInputEvent Event, float AmountDepressed, bool bGamepad )
+bool FAnimationViewportClient::InputKey( FViewport* InViewport, int32 ControllerId, FKey Key, EInputEvent Event, float AmountDepressed, bool bGamepad )
 {
-	const int32 HitX = Viewport->GetMouseX();
-	const int32 HitY = Viewport->GetMouseY();
+	const int32 HitX = InViewport->GetMouseX();
+	const int32 HitY = InViewport->GetMouseY();
 
 	bool bHandled = false;
 
@@ -1726,7 +1726,7 @@ bool FAnimationViewportClient::InputKey( FViewport* Viewport, int32 ControllerId
 	// Pass keys to standard controls, if we didn't consume input
 	return (bHandled)
 		? true
-		: FEditorViewportClient::InputKey(Viewport,  ControllerId,  Key,  Event,  AmountDepressed,  bGamepad);
+		: FEditorViewportClient::InputKey(InViewport,  ControllerId,  Key,  Event,  AmountDepressed,  bGamepad);
 }
 
 void FAnimationViewportClient::SetWidgetMode(FWidget::EWidgetMode InMode)
diff --git a/Engine/Source/Editor/Persona/Private/Persona.cpp b/Engine/Source/Editor/Persona/Private/Persona.cpp
--- a/Engine/Source/Editor/Persona/Private/Persona.cpp
+++ b/Engine/Source/Editor/Persona/Private/Persona.cpp
@@ -1020,12 +1020,12 @@ void FPersona::ExtendDefaultPersonaToolbar()
 
 UBlueprint* FPersona::GetBlueprintObj() const
 {
-	auto EditingObjects = GetEditingObjects();
-	for (int32 i = 0; i < EditingObjects.Num(); ++i)
+	const TArray<UObject*>& EditingObjs = GetEditingObjects();
+	for (int32 i = 0; i < EditingObjs.Num(); ++i)
 	{
-		if (EditingObjects[i]->IsA<UAnimBlueprint>()) {return (UBlueprint*)EditingObjects[i];}
+		if (EditingObjs[i]->IsA<UAnimBlueprint>()) {return (UBlueprint*)EditingObjs[i];}
 	}
-	return NULL;
+	return nullptr;
 }
 
 UObject* FPersona::GetPreviewAnimationAsset() const
diff --git a/Engine/Source/Editor/Persona/Private/SAnimationSequenceBrowser.cpp b/Engine/Source/Editor/Persona/Private/SAnimationSequenceBrowser.cpp
--- a/Engine/Source/Editor/Persona/Private/SAnimationSequenceBrowser.cpp
+++ b/Engine/Source/Editor/Persona/Private/SAnimationSequenceBrowser.cpp
@@ -810,7 +810,7 @@ TSharedRef<SToolTip> SAnimationSequenceBrowser::CreateCustomAssetToolTip(FAssetD
 		];
 
 	TSharedPtr<SHorizontalBox> ContentBox = nullptr;
-	TSharedRef<SToolTip> ToolTip = SNew(SToolTip)
+	TSharedRef<SToolTip> ToolTipWidget = SNew(SToolTip)
 	.TextMargin(1)
 	.BorderImage(FEditorStyle::GetBrush("ContentBrowser.TileViewTooltip.ToolTipBorder"))
 	[
@@ -880,7 +880,7 @@ TSharedRef<SToolTip> SAnimationSequenceBrowser::CreateCustomAssetToolTip(FAssetD
 		]
 	];
 
-	return ToolTip;
+	return ToolTipWidget;
 }
 
 void SAnimationSequenceBrowser::CreateAssetTooltipResources()
diff --git a/Engine/Source/Editor/Persona/Private/SMorphTargetViewer.cpp b/Engine/Source/Editor/Persona/Private/SMorphTargetViewer.cpp
--- a/Engine/Source/Editor/Persona/Private/SMorphTargetViewer.cpp
+++ b/Engine/Source/Editor/Persona/Private/SMorphTargetViewer.cpp
@@ -397,8 +397,8 @@ TSharedPtr<SWidget> SMorphTargetViewer::OnGetContextMenuContent() const
 		FUIAction Action = FUIAction( FExecuteAction::CreateSP( this, &SMorphTargetViewer::OnDeleteMorphTargets ), 
 									  FCanExecuteAction::CreateSP( this, &SMorphTargetViewer::CanPerformDelete ) );
 		const FText Label = LOCTEXT("DeleteMorphTargetButtonLabel", "Delete");
-		const FText ToolTip = LOCTEXT("DeleteMorphTargetButtonTooltip", "Deletes the selected morph targets.");
-		MenuBuilder.AddMenuEntry( Label, ToolTip, FSlateIcon(), Action);
+		const FText ToolTipText = LOCTEXT("DeleteMorphTargetButtonTooltip", "Deletes the selected morph targets.");
+		MenuBuilder.AddMenuEntry( Label, ToolTipText, FSlateIcon(), Action);
 	}
 	MenuBuilder.EndSection();
 
diff --git a/Engine/Source/Editor/Persona/Private/SRetargetSourceWindow.cpp b/Engine/Source/Editor/Persona/Private/SRetargetSourceWindow.cpp
--- a/Engine/Source/Editor/Persona/Private/SRetargetSourceWindow.cpp
+++ b/Engine/Source/Editor/Persona/Private/SRetargetSourceWindow.cpp
@@ -390,8 +390,8 @@ TSharedPtr<SWidget> SRetargetSourceWindow::OnGetContextMenuContent() const
 		{
 			FUIAction Action = FUIAction( FExecuteAction::CreateSP( this, &SRetargetSourceWindow::OnAddRetargetSource ) );
 			const FText Label = LOCTEXT("AddRetargetSourceActionLabel", "Add...");
-			const FText ToolTip = LOCTEXT("AddRetargetSourceActionTooltip", "Add new retarget source.");
-			MenuBuilder.AddMenuEntry( Label, ToolTip, FSlateIcon(), Action);
+			const FText ToolTipText = LOCTEXT("AddRetargetSourceActionTooltip", "Add new retarget source.");
+			MenuBuilder.AddMenuEntry( Label, ToolTipText, FSlateIcon(), Action);
 		}
 		MenuBuilder.EndSection();
 
@@ -400,22 +400,22 @@ TSharedPtr<SWidget> SRetargetSourceWindow::OnGetContextMenuContent() const
 			FUIAction Action = FUIAction( FExecuteAction::CreateSP( this, &SRetargetSourceWindow::OnRenameRetargetSource ), 
 				FCanExecuteAction::CreateSP( this, &SRetargetSourceWindow::CanPerformRename ) );
 			const FText Label = LOCTEXT("RenameRetargetSourceActionLabel", "Rename");
-			const FText ToolTip = LOCTEXT("RenameRetargetSourceActionTooltip", "Rename the selected retarget source.");
-			MenuBuilder.AddMenuEntry( Label, ToolTip, FSlateIcon(), Action);
+			const FText ToolTipText = LOCTEXT("RenameRetargetSourceActionTooltip", "Rename the selected retarget source.");
+			MenuBuilder.AddMenuEntry( Label, ToolTipText, FSlateIcon(), Action);
 		}
 		{
 			FUIAction Action = FUIAction( FExecuteAction::CreateSP( this, &SRetargetSourceWindow::OnDeleteRetargetSource ), 
 				FCanExecuteAction::CreateSP( this, &SRetargetSourceWindow::CanPerformDelete ) );
 			const FText Label = LOCTEXT("DeleteRetargetSourceActionLabel", "Delete");
-			const FText ToolTip = LOCTEXT("DeleteRetargetSourceActionTooltip", "Deletes the selected retarget sources.");
-			MenuBuilder.AddMenuEntry( Label, ToolTip, FSlateIcon(), Action);
+			const FText ToolTipText = LOCTEXT("DeleteRetargetSourceActionTooltip", "Deletes the selected retarget sources.");
+			MenuBuilder.AddMenuEntry( Label, ToolTipText, FSlateIcon(), Action);
 		}
 		{
 			FUIAction Action = FUIAction( FExecuteAction::CreateSP( this, &SRetargetSourceWindow::OnRefreshRetargetSource ), 
 				FCanExecuteAction::CreateSP( this, &SRetargetSourceWindow::CanPerformRefresh ) );
 			const FText Label = LOCTEXT("RefreshRetargetSourceActionLabel", "Refresh");
-			const FText ToolTip = LOCTEXT("RefreshRetargetSourceActionTooltip", "Refreshs the selected retarget sources.");
-			MenuBuilder.AddMenuEntry( Label, ToolTip, FSlateIcon(), Action);
+			const FText ToolTipText = LOCTEXT("RefreshRetargetSourceActionTooltip", "Refreshs the selected retarget sources.");
+			MenuBuilder.AddMenuEntry( Label, ToolTipText, FSlateIcon(), Action);
 		}
 		MenuBuilder.EndSection();
 
diff --git a/Engine/Source/Editor/Persona/Private/SSkeletonAnimNotifies.cpp b/Engine/Source/Editor/Persona/Private/SSkeletonAnimNotifies.cpp
--- a/Engine/Source/Editor/Persona/Private/SSkeletonAnimNotifies.cpp
+++ b/Engine/Source/Editor/Persona/Private/SSkeletonAnimNotifies.cpp
@@ -199,16 +199,16 @@ TSharedPtr<SWidget> SSkeletonAnimNotifies::OnGetContextMenuContent() const
 			FUIAction Action = FUIAction( FExecuteAction::CreateSP( this, &SSkeletonAnimNotifies::OnDeleteAnimNotify ), 
 				FCanExecuteAction::CreateSP( this, &SSkeletonAnimNotifies::CanPerformDelete ) );
 			const FText Label = LOCTEXT("DeleteAnimNotifyButtonLabel", "Delete");
-			const FText ToolTip = LOCTEXT("DeleteAnimNotifyButtonTooltip", "Deletes the selected anim notifies.");
-			MenuBuilder.AddMenuEntry( Label, ToolTip, FSlateIcon(), Action);
+			const FText ToolTipText = LOCTEXT("DeleteAnimNotifyButtonTooltip", "Deletes the selected anim notifies.");
+			MenuBuilder.AddMenuEntry( Label, ToolTipText, FSlateIcon(), Action);
 		}
 
 		{
 			FUIAction Action = FUIAction( FExecuteAction::CreateSP( this, &SSkeletonAnimNotifies::OnRenameAnimNotify ), 
 				FCanExecuteAction::CreateSP( this, &SSkeletonAnimNotifies::CanPerformRename ) );
 			const FText Label = LOCTEXT("RenameAnimNotifyButtonLabel", "Rename");
-			const FText ToolTip = LOCTEXT("RenameAnimNotifyButtonTooltip", "Renames the selected anim notifies.");
-			MenuBuilder.AddMenuEntry( Label, ToolTip, FSlateIcon(), Action);
+			const FText ToolTipText = LOCTEXT("RenameAnimNotifyButtonTooltip", "Renames the selected anim notifies.");
+			MenuBuilder.AddMenuEntry( Label, ToolTipText, FSlateIcon(), Action);
 		}
 	}
 	MenuBuilder.EndSection();
diff --git a/Engine/Source/Editor/Persona/Private/SSkeletonBlendProfiles.cpp b/Engine/Source/Editor/Persona/Private/SSkeletonBlendProfiles.cpp
--- a/Engine/Source/Editor/Persona/Private/SSkeletonBlendProfiles.cpp
+++ b/Engine/Source/Editor/Persona/Private/SSkeletonBlendProfiles.cpp
@@ -174,9 +174,9 @@ TSharedPtr<SWidget> SSkeletonBlendProfiles::GetContextMenu() const
 			Action.ExecuteAction = FExecuteAction::CreateSP(this, &SSkeletonBlendProfiles::OnRequestNewBlendProfile);
 
 			const FText Label = LOCTEXT("Menu_NewProfileLabel", "Create New...");
-			const FText ToolTip = LOCTEXT("Menu_NewProfileToolTip", "Create a new blend profile in the current skeleton.");
+			const FText ToolTipText = LOCTEXT("Menu_NewProfileToolTip", "Create a new blend profile in the current skeleton.");
 
-			MenuBuilder.AddMenuEntry(Label, ToolTip, FSlateIcon(), Action);
+			MenuBuilder.AddMenuEntry(Label, ToolTipText, FSlateIcon(), Action);
 		}
 	}
 	MenuBuilder.EndSection();
@@ -189,9 +189,9 @@ TSharedPtr<SWidget> SSkeletonBlendProfiles::GetContextMenu() const
 			Action.CanExecuteAction = FCanExecuteAction::CreateSP(this, &SSkeletonBlendProfiles::MenuCanDelete);
 
 			const FText Label = LOCTEXT("Menu_DeleteProfileLabel", "Delete");
-			const FText ToolTip = LOCTEXT("Menu_DeleteProfileToolTip", "Delete selected blend profiles.");
+			const FText ToolTipText = LOCTEXT("Menu_DeleteProfileToolTip", "Delete selected blend profiles.");
 
-			MenuBuilder.AddMenuEntry(Label, ToolTip, FSlateIcon(), Action);
+			MenuBuilder.AddMenuEntry(Label, ToolTipText, FSlateIcon(), Action);
 		}
 
 		{
@@ -200,9 +200,9 @@ TSharedPtr<SWidget> SSkeletonBlendProfiles::GetContextMenu() const
 			Action.CanExecuteAction = FCanExecuteAction::CreateSP(this, &SSkeletonBlendProfiles::MenuCanRename);
 
 			const FText Label = LOCTEXT("Menu_RenameProfileLabel", "Rename");
-			const FText ToolTip = LOCTEXT("Menu_RenameProfileToolTip", "Rename selected profile.");
+			const FText ToolTipText = LOCTEXT("Menu_RenameProfileToolTip", "Rename selected profile.");
 
-			MenuBuilder.AddMenuEntry(Label, ToolTip, FSlateIcon(), Action);
+			MenuBuilder.AddMenuEntry(Label, ToolTipText, FSlateIcon(), Action);
 		}
 	}
 	MenuBuilder.EndSection();
diff --git a/Engine/Source/Editor/Persona/Private/SSkeletonSlotNames.cpp b/Engine/Source/Editor/Persona/Private/SSkeletonSlotNames.cpp
--- a/Engine/Source/Editor/Persona/Private/SSkeletonSlotNames.cpp
+++ b/Engine/Source/Editor/Persona/Private/SSkeletonSlotNames.cpp
@@ -266,8 +266,8 @@ TSharedPtr<SWidget> SSkeletonSlotNames::OnGetContextMenuContent() const
 			FUIAction Action = FUIAction(FExecuteAction::CreateSP(this, &SSkeletonSlotNames::OnDeleteSlotGroup, SlotInfo->Name));
 			Action.CanExecuteAction.BindSP(this, &SSkeletonSlotNames::CanDeleteSlotGroup, SlotInfo->Name);
 			const FText Label = LOCTEXT("AnimSlotManagerContextMenuDeleteSlotGroupLabel", "Delete Slot Group");
-			const FText ToolTip = LOCTEXT("AnimSlotManagerContextMenuDeleteSlotGroupTooltip", "Delete this slot group.");
-			MenuBuilder.AddMenuEntry(Label, ToolTip, FSlateIcon(), Action);
+			const FText ToolTipText = LOCTEXT("AnimSlotManagerContextMenuDeleteSlotGroupTooltip", "Delete this slot group.");
+			MenuBuilder.AddMenuEntry(Label, ToolTipText, FSlateIcon(), Action);
 		}
 		MenuBuilder.EndSection();
 	}
@@ -289,8 +289,8 @@ TSharedPtr<SWidget> SSkeletonSlotNames::OnGetContextMenuContent() const
 
 			FUIAction Action = FUIAction(FExecuteAction::CreateSP(this, &SSkeletonSlotNames::OnRenameSlot, SlotInfo->Name));
 			const FText Label = LOCTEXT("AnimSlotManagerContextMenuRenameSlotLabel", "Rename Slot");
-			const FText ToolTip = LOCTEXT("AnimSlotManagerContextMenuRenameSlotTooltip", "Rename this slot");
-			MenuBuilder.AddMenuEntry(Label, ToolTip, FSlateIcon(), Action);
+			const FText ToolTipText = LOCTEXT("AnimSlotManagerContextMenuRenameSlotTooltip", "Rename this slot");
+			MenuBuilder.AddMenuEntry(Label, ToolTipText, FSlateIcon(), Action);
 		}
 		// Delete Slot
 		{
@@ -298,8 +298,8 @@ TSharedPtr<SWidget> SSkeletonSlotNames::OnGetContextMenuContent() const
 
 			FUIAction Action = FUIAction(FExecuteAction::CreateSP(this, &SSkeletonSlotNames::OnDeleteSlot, SlotInfo->Name));
 			const FText Label = LOCTEXT("AnimSlotManagerContextMenuDeleteSlotLabel", "Delete Slot");
-			const FText ToolTip = LOCTEXT("AnimSlotManagerContextMenuDeleteSlotTooltip", "Delete this slot.");
-			MenuBuilder.AddMenuEntry(Label, ToolTip, FSlateIcon(), Action);
+			const FText ToolTipText = LOCTEXT("AnimSlotManagerContextMenuDeleteSlotTooltip", "Delete this slot.");
+			MenuBuilder.AddMenuEntry(Label, ToolTipText, FSlateIcon(), Action);
 		}
 		MenuBuilder.EndSection();
 	}
@@ -309,15 +309,15 @@ TSharedPtr<SWidget> SSkeletonSlotNames::OnGetContextMenuContent() const
 	{
 		FUIAction Action = FUIAction(FExecuteAction::CreateSP(this, &SSkeletonSlotNames::OnAddSlot));
 		const FText Label = LOCTEXT("AnimSlotManagerContextMenuAddSlotLabel", "Add Slot");
-		const FText ToolTip = LOCTEXT("AnimSlotManagerContextMenuAddSlotTooltip", "Adds a new Slot");
-		MenuBuilder.AddMenuEntry(Label, ToolTip, FSlateIcon(), Action);
+		const FText ToolTipText = LOCTEXT("AnimSlotManagerContextMenuAddSlotTooltip", "Adds a new Slot");
+		MenuBuilder.AddMenuEntry(Label, ToolTipText, FSlateIcon(), Action);
 	}
 	// Add Group
 	{
 		FUIAction Action = FUIAction(FExecuteAction::CreateSP(this, &SSkeletonSlotNames::OnAddGroup));
 		const FText Label = LOCTEXT("AnimSlotManagerContextMenuAddGroupLabel", "Add Group");
-		const FText ToolTip = LOCTEXT("AnimSlotManagerContextMenuAddGroupTooltip", "Adds a new Group");
-		MenuBuilder.AddMenuEntry(Label, ToolTip, FSlateIcon(), Action);
+		const FText ToolTipText = LOCTEXT("AnimSlotManagerContextMenuAddGroupTooltip", "Adds a new Group");
+		MenuBuilder.AddMenuEntry(Label, ToolTipText, FSlateIcon(), Action);
 	}
 	MenuBuilder.EndSection();
 
@@ -331,13 +331,13 @@ void SSkeletonSlotNames::FillSetSlotGroupSubMenu(FMenuBuilder& MenuBuilder)
 	{
 		const FName& GroupName = SlotGroup.GroupName;
 
-		const FText ToolTip = FText::Format(LOCTEXT("ContextMenuSetSlotSubMenuToolTip", "Changes slot's group to {0}"), FText::FromName(GroupName));
+		const FText ToolTipText = FText::Format(LOCTEXT("ContextMenuSetSlotSubMenuToolTip", "Changes slot's group to {0}"), FText::FromName(GroupName));
 /*		FString Label = Class->GetDisplayNameText().ToString();*/
 		const FText Label = FText::FromName(GroupName);
 
 		FUIAction UIAction;
 		UIAction.ExecuteAction.BindRaw(this, &SSkeletonSlotNames::ContextMenuOnSetSlot,	GroupName);
-		MenuBuilder.AddMenuEntry(Label, ToolTip, FSlateIcon(), UIAction);
+		MenuBuilder.AddMenuEntry(Label, ToolTipText, FSlateIcon(), UIAction);
 	}
 }
 
diff --git a/Engine/Source/Editor/Persona/Private/SSkeletonSmartNameManager.cpp b/Engine/Source/Editor/Persona/Private/SSkeletonSmartNameManager.cpp
--- a/Engine/Source/Editor/Persona/Private/SSkeletonSmartNameManager.cpp
+++ b/Engine/Source/Editor/Persona/Private/SSkeletonSmartNameManager.cpp
@@ -206,16 +206,16 @@ TSharedPtr<SWidget> SSkeletonSmartNameManager::OnGetContextMenu() const
 			Action.ExecuteAction = FExecuteAction::CreateSP(this, &SSkeletonSmartNameManager::OnRenameClicked);
 			Action.CanExecuteAction = FCanExecuteAction::CreateSP(this, &SSkeletonSmartNameManager::CanRename);
 			const FText Label = LOCTEXT("RenameSmartNameLabel", "Rename");
-			const FText ToolTip = LOCTEXT("RenameSmartNameToolTip", "Rename the selected item");
-			MenuBuilder.AddMenuEntry(Label, ToolTip, FSlateIcon(), Action);
+			const FText ToolTipText = LOCTEXT("RenameSmartNameToolTip", "Rename the selected item");
+			MenuBuilder.AddMenuEntry(Label, ToolTipText, FSlateIcon(), Action);
 		}
 		// Delete a name
 		{
 			Action.ExecuteAction = FExecuteAction::CreateSP(this, &SSkeletonSmartNameManager::OnDeleteNameClicked);
 			Action.CanExecuteAction = FCanExecuteAction::CreateSP(this, &SSkeletonSmartNameManager::CanDelete);
 			const FText Label = LOCTEXT("DeleteSmartNameLabel", "Delete");
-			const FText ToolTip = LOCTEXT("DeleteSmartNameToolTip", "Delete the selected item");
-			MenuBuilder.AddMenuEntry(Label, ToolTip, FSlateIcon(), Action);
+			const FText ToolTipText = LOCTEXT("DeleteSmartNameToolTip", "Delete the selected item");
+			MenuBuilder.AddMenuEntry(Label, ToolTipText, FSlateIcon(), Action);
 		}
 		// Add a name
 		MenuBuilder.AddMenuSeparator();
@@ -223,8 +223,8 @@ TSharedPtr<SWidget> SSkeletonSmartNameManager::OnGetContextMenu() const
 			Action.ExecuteAction = FExecuteAction::CreateSP(this, &SSkeletonSmartNameManager::OnAddClicked);
 			Action.CanExecuteAction = nullptr;
 			const FText Label = LOCTEXT("AddSmartNameLabel", "Add...");
-			const FText ToolTip = LOCTEXT("AddSmartNameToolTip", "Add an entry to the skeleton");
-			MenuBuilder.AddMenuEntry(Label, ToolTip, FSlateIcon(), Action);
+			const FText ToolTipText = LOCTEXT("AddSmartNameToolTip", "Add an entry to the skeleton");
+			MenuBuilder.AddMenuEntry(Label, ToolTipText, FSlateIcon(), Action);
 		}
 	}
 	MenuBuilder.EndSection();
@@ -404,9 +404,9 @@ void SCurveNameManager::OnDeleteNameClicked()
 				TArray<UObject::FAssetRegistryTag> Tags;
 				Asset->GetAssetRegistryTags(Tags);
 				
-				UObject::FAssetRegistryTag* CurveTag = Tags.FindByPredicate([](const UObject::FAssetRegistryTag& Tag)
+				UObject::FAssetRegistryTag* CurveTag = Tags.FindByPredicate([](const UObject::FAssetRegistryTag& InTag)
 				{
-					return Tag.Name == USkeleton::CurveTag;
+					return InTag.Name == USkeleton::CurveTag;
 				});
 				CurveDataCopy = CurveTag->Value;
 			}
diff --git a/Engine/Source/Editor/Persona/Private/SSkeletonTree.cpp b/Engine/Source/Editor/Persona/Private/SSkeletonTree.cpp
--- a/Engine/Source/Editor/Persona/Private/SSkeletonTree.cpp
+++ b/Engine/Source/Editor/Persona/Private/SSkeletonTree.cpp
@@ -165,19 +165,19 @@ TSharedRef< SWidget > SSkeletonTreeRow::GenerateWidgetForColumn( const FName& Co
 {
 	if ( ColumnName == ColumnID_BoneLabel )
 	{
-		TSharedPtr< SHorizontalBox > Box;
+		TSharedPtr< SHorizontalBox > RowBox;
 
-		SAssignNew( Box, SHorizontalBox );
+		SAssignNew( RowBox, SHorizontalBox );
 
-		Box->AddSlot()
+		RowBox->AddSlot()
 			.AutoWidth()
 			[
 				SNew( SExpanderArrow, SharedThis(this) )
 			];
 
-		Item->GenerateWidgetForNameColumn( Box, FilterText, FIsSelected::CreateSP(this, &STableRow::IsSelectedExclusively ) );
+		Item->GenerateWidgetForNameColumn( RowBox, FilterText, FIsSelected::CreateSP(this, &STableRow::IsSelectedExclusively ) );
 
-		return Box.ToSharedRef();
+		return RowBox.ToSharedRef();
 	}
 	else
 	{
diff --git a/Engine/Source/Editor/PhAT/Private/PhATRender.cpp b/Engine/Source/Editor/PhAT/Private/PhATRender.cpp
--- a/Engine/Source/Editor/PhAT/Private/PhATRender.cpp
+++ b/Engine/Source/Editor/PhAT/Private/PhATRender.cpp
@@ -371,7 +371,7 @@ void UPhATEdSkeletalMeshComponent::DrawCurrentInfluences(FPrimitiveDrawInterface
 
 FTransform UPhATEdSkeletalMeshComponent::GetPrimitiveTransform(FTransform& BoneTM, int32 BodyIndex, EKCollisionPrimitiveType PrimType, int32 PrimIndex, float Scale)
 {
-	UBodySetup* BodySetup = SharedData->PhysicsAsset->BodySetup[BodyIndex];
+	UBodySetup* SharedBodySetup = SharedData->PhysicsAsset->BodySetup[BodyIndex];
 	FVector Scale3D(Scale);
 
 	FTransform ManTM = FTransform::Identity;
@@ -393,25 +393,25 @@ FTransform UPhATEdSkeletalMeshComponent::GetPrimitiveTransform(FTransform& BoneT
 
 	if (PrimType == KPT_Sphere)
 	{
-		FTransform PrimTM = ManTM * BodySetup->AggGeom.SphereElems[PrimIndex].GetTransform();
+		FTransform PrimTM = ManTM * SharedBodySetup->AggGeom.SphereElems[PrimIndex].GetTransform();
 		PrimTM.ScaleTranslation(Scale3D);
 		return PrimTM * BoneTM;
 	}
 	else if (PrimType == KPT_Box)
 	{
-		FTransform PrimTM = ManTM * BodySetup->AggGeom.BoxElems[PrimIndex].GetTransform();
+		FTransform PrimTM = ManTM * SharedBodySetup->AggGeom.BoxElems[PrimIndex].GetTransform();
 		PrimTM.ScaleTranslation(Scale3D);
 		return PrimTM * BoneTM;
 	}
 	else if (PrimType == KPT_Sphyl)
 	{
-		FTransform PrimTM = ManTM * BodySetup->AggGeom.SphylElems[PrimIndex].GetTransform();
+		FTransform PrimTM = ManTM * SharedBodySetup->AggGeom.SphylElems[PrimIndex].GetTransform();
 		PrimTM.ScaleTranslation(Scale3D);
 		return PrimTM * BoneTM;
 	}
 	else if (PrimType == KPT_Convex)
 	{
-		FTransform PrimTM = ManTM * BodySetup->AggGeom.ConvexElems[PrimIndex].GetTransform();
+		FTransform PrimTM = ManTM * SharedBodySetup->AggGeom.ConvexElems[PrimIndex].GetTransform();
 		PrimTM.ScaleTranslation(Scale3D);
 		return PrimTM * BoneTM;
 	}
@@ -423,17 +423,17 @@ FTransform UPhATEdSkeletalMeshComponent::GetPrimitiveTransform(FTransform& BoneT
 
 FColor UPhATEdSkeletalMeshComponent::GetPrimitiveColor(int32 BodyIndex, EKCollisionPrimitiveType PrimitiveType, int32 PrimitiveIndex)
 {
-	UBodySetup* BodySetup = SharedData->PhysicsAsset->BodySetup[ BodyIndex ];
+	UBodySetup* SharedBodySetup = SharedData->PhysicsAsset->BodySetup[ BodyIndex ];
 
 	if (!SharedData->bRunningSimulation && SharedData->EditingMode == FPhATSharedData::PEM_ConstraintEdit && SharedData->GetSelectedConstraint())
 	{
 		UPhysicsConstraintTemplate* cs = SharedData->PhysicsAsset->ConstraintSetup[ SharedData->GetSelectedConstraint()->Index ];
 
-		if (cs->DefaultInstance.ConstraintBone1 == BodySetup->BoneName)
+		if (cs->DefaultInstance.ConstraintBone1 == SharedBodySetup->BoneName)
 		{
 			return ConstraintBone1Color;
 		}
-		else if (cs->DefaultInstance.ConstraintBone2 == BodySetup->BoneName)
+		else if (cs->DefaultInstance.ConstraintBone2 == SharedBodySetup->BoneName)
 		{
 			return ConstraintBone2Color;
 		}
@@ -467,7 +467,7 @@ FColor UPhATEdSkeletalMeshComponent::GetPrimitiveColor(int32 BodyIndex, EKCollis
 	
 	if (SharedData->bShowFixedStatus && SharedData->bRunningSimulation)
 	{
-		const bool bIsSimulatedAtAll = BodySetup->PhysicsType == PhysType_Simulated || (BodySetup->PhysicsType == PhysType_Default && SharedData->EditorSimOptions->PhysicsBlend > 0.f);
+		const bool bIsSimulatedAtAll = SharedBodySetup->PhysicsType == PhysType_Simulated || (SharedBodySetup->PhysicsType == PhysType_Default && SharedData->EditorSimOptions->PhysicsBlend > 0.f);
 		if (!bIsSimulatedAtAll)
 		{
 			return FixedColor;
diff --git a/Engine/Source/Editor/PlacementMode/Private/SPlacementModeTools.cpp b/Engine/Source/Editor/PlacementMode/Private/SPlacementModeTools.cpp
--- a/Engine/Source/Editor/PlacementMode/Private/SPlacementModeTools.cpp
+++ b/Engine/Source/Editor/PlacementMode/Private/SPlacementModeTools.cpp
@@ -133,11 +133,11 @@ void SPlacementAssetEntry::Construct(const FArguments& InArgs, const TSharedPtr<
 	}
 	
 	UClass* DocClass = nullptr;
-	TSharedPtr<IToolTip> ToolTip;
+	TSharedPtr<IToolTip> AssetEntryToolTip;
 	if(DefaultActor != nullptr)
 	{
 		DocClass = DefaultActor->GetClass();
-		ToolTip = FEditorClassUtils::GetTooltip(DefaultActor->GetClass());
+		AssetEntryToolTip = FEditorClassUtils::GetTooltip(DefaultActor->GetClass());
 	}
 
 	if (IsClass && !IsVolume && !ActorTypeDisplayName.IsEmpty())
@@ -145,9 +145,9 @@ void SPlacementAssetEntry::Construct(const FArguments& InArgs, const TSharedPtr<
 		Item->DisplayName = ActorTypeDisplayName;
 	}
 
-	if (!ToolTip.IsValid())
+	if (!AssetEntryToolTip.IsValid())
 	{
-		ToolTip = FSlateApplicationBase::Get().MakeToolTip(Item->DisplayName);
+		AssetEntryToolTip = FSlateApplicationBase::Get().MakeToolTip(Item->DisplayName);
 	}
 	
 	const FButtonStyle& ButtonStyle = FEditorStyle::GetWidgetStyle<FButtonStyle>( "PlacementBrowser.Asset" );
@@ -169,7 +169,7 @@ void SPlacementAssetEntry::Construct(const FArguments& InArgs, const TSharedPtr<
 		SNew( SBorder )
 		.BorderImage( this, &SPlacementAssetEntry::GetBorder )
 		.Cursor( EMouseCursor::GrabHand )
-		.ToolTip( ToolTip )
+		.ToolTip( AssetEntryToolTip )
 		[
 			SNew( SHorizontalBox )
 
diff --git a/Engine/Source/Editor/PropertyEditor/Private/ItemPropertyNode.cpp b/Engine/Source/Editor/PropertyEditor/Private/ItemPropertyNode.cpp
--- a/Engine/Source/Editor/PropertyEditor/Private/ItemPropertyNode.cpp
+++ b/Engine/Source/Editor/PropertyEditor/Private/ItemPropertyNode.cpp
@@ -96,7 +96,7 @@ void FItemPropertyNode::InitExpansionFlags (void)
 	if(	Cast<UStructProperty>(MyProperty)
 		||	( Cast<UArrayProperty>(MyProperty) && GetReadAddress(false,Addresses) )
 		||  HasNodeFlags(EPropertyNodeFlags::EditInline)
-		||	( Property->ArrayDim > 1 && ArrayIndex == -1 ) )
+		||	( MyProperty->ArrayDim > 1 && ArrayIndex == -1 ) )
 	{
 		SetNodeFlags(EPropertyNodeFlags::CanBeExpanded, true);
 	}
@@ -107,27 +107,27 @@ void FItemPropertyNode::InitExpansionFlags (void)
 void FItemPropertyNode::InitChildNodes()
 {
 	//NOTE - this is only turned off as to not invalidate child object nodes.
-	UProperty* Property = GetProperty();
-	UStructProperty* StructProperty = Cast<UStructProperty>(Property);
-	UArrayProperty* ArrayProperty = Cast<UArrayProperty>(Property);
-	UObjectPropertyBase* ObjectProperty = Cast<UObjectPropertyBase>(Property);
+	UProperty* MyProperty = GetProperty();
+	UStructProperty* StructProperty = Cast<UStructProperty>(MyProperty);
+	UArrayProperty* ArrayProperty = Cast<UArrayProperty>(MyProperty);
+	UObjectPropertyBase* ObjectProperty = Cast<UObjectPropertyBase>(MyProperty);
 
 	const bool bShouldShowHiddenProperties = !!HasNodeFlags(EPropertyNodeFlags::ShouldShowHiddenProperties);
 	const bool bShouldShowDisableEditOnInstance = !!HasNodeFlags(EPropertyNodeFlags::ShouldShowDisableEditOnInstance);
 
-	if( Property->ArrayDim > 1 && ArrayIndex == -1 )
+	if( MyProperty->ArrayDim > 1 && ArrayIndex == -1 )
 	{
 		// Do not add array children which are defined by an enum but the enum at the array index is hidden
 		// This only applies to static arrays
 		static const FName NAME_ArraySizeEnum("ArraySizeEnum");
 		UEnum* ArraySizeEnum = NULL; 
-		if (Property->HasMetaData(NAME_ArraySizeEnum))
+		if (MyProperty->HasMetaData(NAME_ArraySizeEnum))
 		{
-			ArraySizeEnum	= FindObject<UEnum>(NULL, *Property->GetMetaData(NAME_ArraySizeEnum));
+			ArraySizeEnum	= FindObject<UEnum>(NULL, *MyProperty->GetMetaData(NAME_ArraySizeEnum));
 		}
 
 		// Expand array.
-		for( int32 Index = 0 ; Index < Property->ArrayDim ; Index++ )
+		for( int32 Index = 0 ; Index < MyProperty->ArrayDim ; Index++ )
 		{
 			bool bShouldBeHidden = false;
 			if( ArraySizeEnum )
@@ -141,8 +141,8 @@ void FItemPropertyNode::InitChildNodes()
 				TSharedPtr<FItemPropertyNode> NewItemNode( new FItemPropertyNode);
 				FPropertyNodeInitParams InitParams;
 				InitParams.ParentNode = SharedThis(this);
-				InitParams.Property = Property;
-				InitParams.ArrayOffset = Index*Property->ElementSize;
+				InitParams.Property = MyProperty;
+				InitParams.ArrayOffset = Index*MyProperty->ElementSize;
 				InitParams.ArrayIndex = Index;
 				InitParams.bAllowChildren = true;
 				InitParams.bForceHiddenPropertyVisibility = bShouldShowHiddenProperties;
@@ -223,7 +223,7 @@ void FItemPropertyNode::InitChildNodes()
 			}
 		}
 	}
-	else if( ObjectProperty || Property->IsA(UInterfaceProperty::StaticClass()))
+	else if( ObjectProperty || MyProperty->IsA(UInterfaceProperty::StaticClass()))
 	{
 		uint8* ReadValue = NULL;
 
@@ -260,7 +260,7 @@ void FItemPropertyNode::InitChildNodes()
 
 				FPropertyNodeInitParams InitParams;
 				InitParams.ParentNode = SharedThis(this);
-				InitParams.Property = Property;
+				InitParams.Property = MyProperty;
 				InitParams.ArrayOffset = 0;
 				InitParams.ArrayIndex = INDEX_NONE;
 				InitParams.bAllowChildren = true;
diff --git a/Engine/Source/Editor/PropertyEditor/Private/SResetToDefaultMenu.cpp b/Engine/Source/Editor/PropertyEditor/Private/SResetToDefaultMenu.cpp
--- a/Engine/Source/Editor/PropertyEditor/Private/SResetToDefaultMenu.cpp
+++ b/Engine/Source/Editor/PropertyEditor/Private/SResetToDefaultMenu.cpp
@@ -56,8 +56,6 @@ void SResetToDefaultMenu::Tick( const FGeometry& AllottedGeometry, const double
 
 EVisibility SResetToDefaultMenu::GetResetToDefaultVisibility() const
 {
-	EVisibility Visibility = VisibilityWhenDefault;
-
 	return bShouldBeVisible ? EVisibility::Visible : VisibilityWhenDefault;
 }
 
diff --git a/Engine/Source/Editor/PropertyEditor/Private/SStructureDetailsView.cpp b/Engine/Source/Editor/PropertyEditor/Private/SStructureDetailsView.cpp
--- a/Engine/Source/Editor/PropertyEditor/Private/SStructureDetailsView.cpp
+++ b/Engine/Source/Editor/PropertyEditor/Private/SStructureDetailsView.cpp
@@ -98,7 +98,7 @@ void SStructureDetailsView::Construct(const FArguments& InArgs)
 			FSlateIcon(),
 			FUIAction(FExecuteAction::CreateSP(this, &SStructureDetailsView::SetRootExpansionStates, /*bExpanded=*/true, /*bRecurse=*/false)));
 
-	TSharedRef<SHorizontalBox> FilterRow = SNew( SHorizontalBox )
+	TSharedRef<SHorizontalBox> FilterBoxRow = SNew( SHorizontalBox )
 		.Visibility(this, &SStructureDetailsView::GetFilterBoxVisibility)
 		+SHorizontalBox::Slot()
 		.FillWidth( 1 )
@@ -111,7 +111,7 @@ void SStructureDetailsView::Construct(const FArguments& InArgs)
 
 	if (DetailsViewArgs.bShowOptions)
 	{
-		FilterRow->AddSlot()
+		FilterBoxRow->AddSlot()
 			.HAlign(HAlign_Right)
 			.AutoWidth()
 			[
@@ -151,7 +151,7 @@ void SStructureDetailsView::Construct(const FArguments& InArgs)
 			.AutoHeight()
 			.Padding( 0.0f, 0.0f, 0.0f, 2.0f )
 			[
-				FilterRow
+				FilterBoxRow
 			]
 			+ SVerticalBox::Slot()
 			.FillHeight(1)
diff --git a/Engine/Source/Editor/PropertyEditor/Private/UserInterface/PropertyEditor/SPropertyComboBox.cpp b/Engine/Source/Editor/PropertyEditor/Private/UserInterface/PropertyEditor/SPropertyComboBox.cpp
--- a/Engine/Source/Editor/PropertyEditor/Private/UserInterface/PropertyEditor/SPropertyComboBox.cpp
+++ b/Engine/Source/Editor/PropertyEditor/Private/UserInterface/PropertyEditor/SPropertyComboBox.cpp
@@ -101,7 +101,7 @@ void SPropertyComboBox::OnSelectionChangedInternal( TSharedPtr<FString> InSelect
 TSharedRef<SWidget> SPropertyComboBox::OnGenerateComboWidget( TSharedPtr<FString> InComboString )
 {
 	//Find the corresponding tool tip for this combo entry if any
-	FText ToolTip;
+	FText ToolTipText;
 	bool bEnabled = true;
 	if (ToolTipList.Num() > 0)
 	{
@@ -110,7 +110,7 @@ TSharedRef<SWidget> SPropertyComboBox::OnGenerateComboWidget( TSharedPtr<FString
 		{
 			//A list of tool tips should have been populated in a 1 to 1 correspondance
 			check(ComboItemList.Num() == ToolTipList.Num());
-			ToolTip = ToolTipList[Index];
+			ToolTipText = ToolTipList[Index];
 
 			if( RestrictedList.Num() > 0 )
 			{
@@ -123,7 +123,7 @@ TSharedRef<SWidget> SPropertyComboBox::OnGenerateComboWidget( TSharedPtr<FString
 		SNew( STextBlock )
 		.Text( FText::FromString(*InComboString) )
 		.Font( Font )
-		.ToolTipText(ToolTip)
+		.ToolTipText(ToolTipText)
 		.IsEnabled(bEnabled);
 }
 
diff --git a/Engine/Source/Editor/PropertyEditor/Private/UserInterface/PropertyEditor/SPropertyEditorAsset.cpp b/Engine/Source/Editor/PropertyEditor/Private/UserInterface/PropertyEditor/SPropertyEditorAsset.cpp
--- a/Engine/Source/Editor/PropertyEditor/Private/UserInterface/PropertyEditor/SPropertyEditorAsset.cpp
+++ b/Engine/Source/Editor/PropertyEditor/Private/UserInterface/PropertyEditor/SPropertyEditorAsset.cpp
@@ -515,31 +515,31 @@ FText SPropertyEditorAsset::OnGetToolTip() const
 	FObjectOrAssetData Value; 
 	FPropertyAccess::Result Result = GetValue( Value );
 
-	FText ToolTip = FText::GetEmpty();
+	FText ToolTipText = FText::GetEmpty();
 
 	if( Result == FPropertyAccess::Success )
 	{
 		if(Value.Object != NULL && !bIsActor )
 		{
 			// Display the package name which is a valid path to the object without redundant information
-			ToolTip = FText::FromString(Value.Object->GetOutermost()->GetName());
+			ToolTipText = FText::FromString(Value.Object->GetOutermost()->GetName());
 		}
 		else if( Value.AssetData.IsValid() )
 		{
-			ToolTip = FText::FromName(Value.AssetData.PackageName);
+			ToolTipText = FText::FromName(Value.AssetData.PackageName);
 		}
 	}
 	else if( Result == FPropertyAccess::MultipleValues )
 	{
-		ToolTip = LOCTEXT("MultipleValues", "Multiple Values");
+		ToolTipText = LOCTEXT("MultipleValues", "Multiple Values");
 	}
 
-	if( ToolTip.IsEmpty() )
+	if( ToolTipText.IsEmpty() )
 	{
-		ToolTip = FText::FromString(ObjectPath.Get());
+		ToolTipText = FText::FromString(ObjectPath.Get());
 	}
 
-	return ToolTip;
+	return ToolTipText;
 }
 
 void SPropertyEditorAsset::SetValue( const FAssetData& AssetData )
diff --git a/Engine/Source/Editor/PropertyEditor/Private/UserInterface/PropertyEditor/SPropertyEditorColor.cpp b/Engine/Source/Editor/PropertyEditor/Private/UserInterface/PropertyEditor/SPropertyEditorColor.cpp
--- a/Engine/Source/Editor/PropertyEditor/Private/UserInterface/PropertyEditor/SPropertyEditorColor.cpp
+++ b/Engine/Source/Editor/PropertyEditor/Private/UserInterface/PropertyEditor/SPropertyEditorColor.cpp
@@ -84,19 +84,19 @@ EVisibility SPropertyEditorColor::GetVisibilityForOpaqueDisplay() const
 	// the color even if it's mostly/entirely transparent.  But if the color is rendered as completely opaque,
 	// we might as well collapse the extra opaque display rather than drawing two separate boxes.
 
-	EVisibility Visibility = EVisibility::Collapsed;
+	EVisibility OpaqueDisplayVisibility = EVisibility::Collapsed;
 	if( !bIgnoreAlpha )
 	{
 		const bool bColorIsAlreadyOpaque = (OnGetColor().A == 1.0);
 		if (bColorIsAlreadyOpaque)
 		{
-			Visibility = EVisibility::Collapsed;
+			OpaqueDisplayVisibility = EVisibility::Collapsed;
 		}
 
-		Visibility = EVisibility::Visible;
+		OpaqueDisplayVisibility = EVisibility::Visible;
 	}
 
-	return Visibility;
+	return OpaqueDisplayVisibility;
 }
 
 FReply SPropertyEditorColor::ColorBlock_OnMouseButtonDown(const FGeometry& MyGeometry, const FPointerEvent& MouseEvent)
diff --git a/Engine/Source/Editor/PropertyEditor/Private/UserInterface/PropertyEditor/SPropertyEditorTableRow.cpp b/Engine/Source/Editor/PropertyEditor/Private/UserInterface/PropertyEditor/SPropertyEditorTableRow.cpp
--- a/Engine/Source/Editor/PropertyEditor/Private/UserInterface/PropertyEditor/SPropertyEditorTableRow.cpp
+++ b/Engine/Source/Editor/PropertyEditor/Private/UserInterface/PropertyEditor/SPropertyEditorTableRow.cpp
@@ -59,7 +59,7 @@ TSharedRef<SWidget> SPropertyEditorTableRow::GenerateWidgetForColumn( const FNam
 
 TSharedRef< SWidget > SPropertyEditorTableRow::ConstructNameColumnWidget()
 {
-	TSharedRef< SHorizontalBox > Box =
+	TSharedRef< SHorizontalBox > NameColumnWidget =
 		SNew(SHorizontalBox)
 		+SHorizontalBox::Slot()
 		.AutoWidth()
@@ -82,7 +82,7 @@ TSharedRef< SWidget > SPropertyEditorTableRow::ConstructNameColumnWidget()
 			.OnDoubleClicked( this, &SPropertyEditorTableRow::OnNameDoubleClicked )
 		];
 
-		return Box;
+		return NameColumnWidget;
 }
 
 
diff --git a/Engine/Source/Editor/Sequencer/Private/DisplayNodes/SequencerFolderNode.cpp b/Engine/Source/Editor/Sequencer/Private/DisplayNodes/SequencerFolderNode.cpp
--- a/Engine/Source/Editor/Sequencer/Private/DisplayNodes/SequencerFolderNode.cpp
+++ b/Engine/Source/Editor/Sequencer/Private/DisplayNodes/SequencerFolderNode.cpp
@@ -53,10 +53,10 @@ void FSequencerFolderNode::SetDisplayName( const FText& NewDisplayName )
 	if ( MovieSceneFolder.GetFolderName() != FName( *NewDisplayName.ToString() ) )
 	{
 		TArray<FName> SiblingNames;
-		TSharedPtr<FSequencerDisplayNode> ParentNode = GetParent();
-		if ( ParentNode.IsValid() )
+		TSharedPtr<FSequencerDisplayNode> ParentSeqNode = GetParent();
+		if ( ParentSeqNode.IsValid() )
 		{
-			for ( TSharedRef<FSequencerDisplayNode> SiblingNode : ParentNode->GetChildNodes() )
+			for ( TSharedRef<FSequencerDisplayNode> SiblingNode : ParentSeqNode->GetChildNodes() )
 			{
 				if ( SiblingNode != AsShared() )
 				{
@@ -138,8 +138,8 @@ TOptional<EItemDropZone> FSequencerFolderNode::CanDrop( FSequencerDisplayNodeDra
 
 		for ( TSharedRef<FSequencerDisplayNode> DraggedNode : DragDropOp.GetDraggedNodes() )
 		{
-			TSharedPtr<FSequencerDisplayNode> ParentNode = DraggedNode->GetParent();
-			if ( ParentNode.IsValid() && ParentNode.Get() == this )
+			TSharedPtr<FSequencerDisplayNode> ParentSeqNode = DraggedNode->GetParent();
+			if ( ParentSeqNode.IsValid() && ParentSeqNode.Get() == this )
 			{
 				DragDropOp.CurrentHoverText = NSLOCTEXT( "SeqeuencerFolderNode", "SameParentDragError", "Can't drag a node onto the same parent." );
 				return TOptional<EItemDropZone>();
@@ -179,7 +179,7 @@ void FSequencerFolderNode::Drop( const TArray<TSharedRef<FSequencerDisplayNode>>
 	GetFolder().Modify();
 	for ( TSharedRef<FSequencerDisplayNode> DraggedNode : DraggedNodes )
 	{
-		TSharedPtr<FSequencerDisplayNode> ParentNode = DraggedNode->GetParent();
+		TSharedPtr<FSequencerDisplayNode> ParentSeqNode = DraggedNode->GetParent();
 		switch ( DraggedNode->GetType() )
 		{
 			case ESequencerNode::Folder:
@@ -197,10 +197,10 @@ void FSequencerFolderNode::Drop( const TArray<TSharedRef<FSequencerDisplayNode>>
 					FocusedMovieScene->GetRootFolders().Add( &DraggedFolderNode->GetFolder() );
 				}
 				
-				if ( ParentNode.IsValid() )
+				if ( ParentSeqNode.IsValid() )
 				{
-					checkf( ParentNode->GetType() == ESequencerNode::Folder, TEXT( "Can not remove from unsupported parent node." ) );
-					TSharedPtr<FSequencerFolderNode> ParentFolder = StaticCastSharedPtr<FSequencerFolderNode>( ParentNode );
+					checkf( ParentSeqNode->GetType() == ESequencerNode::Folder, TEXT( "Can not remove from unsupported parent node." ) );
+					TSharedPtr<FSequencerFolderNode> ParentFolder = StaticCastSharedPtr<FSequencerFolderNode>( ParentSeqNode );
 					ParentFolder->GetFolder().Modify();
 					ParentFolder->GetFolder().RemoveChildFolder( &DraggedFolderNode->GetFolder() );
 				}
@@ -221,10 +221,10 @@ void FSequencerFolderNode::Drop( const TArray<TSharedRef<FSequencerDisplayNode>>
 					GetFolder().AddChildMasterTrack( DraggedTrackNode->GetTrack() );
 				}
 				
-				if ( ParentNode.IsValid() )
+				if ( ParentSeqNode.IsValid() )
 				{
-					checkf( ParentNode->GetType() == ESequencerNode::Folder, TEXT( "Can not remove from unsupported parent node." ) );
-					TSharedPtr<FSequencerFolderNode> ParentFolder = StaticCastSharedPtr<FSequencerFolderNode>( ParentNode );
+					checkf( ParentSeqNode->GetType() == ESequencerNode::Folder, TEXT( "Can not remove from unsupported parent node." ) );
+					TSharedPtr<FSequencerFolderNode> ParentFolder = StaticCastSharedPtr<FSequencerFolderNode>( ParentSeqNode );
 					ParentFolder->GetFolder().Modify();
 					ParentFolder->GetFolder().RemoveChildMasterTrack( DraggedTrackNode->GetTrack() );
 				}
@@ -240,10 +240,10 @@ void FSequencerFolderNode::Drop( const TArray<TSharedRef<FSequencerDisplayNode>>
 					GetFolder().AddChildObjectBinding( DraggedObjectBindingNode->GetObjectBinding() );
 				}
 
-				if ( ParentNode.IsValid() )
+				if ( ParentSeqNode.IsValid() )
 				{
-					checkf( ParentNode->GetType() == ESequencerNode::Folder, TEXT( "Can not remove from unsupported parent node." ) );
-					TSharedPtr<FSequencerFolderNode> ParentFolder = StaticCastSharedPtr<FSequencerFolderNode>( ParentNode );
+					checkf( ParentSeqNode->GetType() == ESequencerNode::Folder, TEXT( "Can not remove from unsupported parent node." ) );
+					TSharedPtr<FSequencerFolderNode> ParentFolder = StaticCastSharedPtr<FSequencerFolderNode>( ParentSeqNode );
 					ParentFolder->GetFolder().Modify();
 					ParentFolder->GetFolder().RemoveChildObjectBinding( DraggedObjectBindingNode->GetObjectBinding() );
 				}
diff --git a/Engine/Source/Editor/Sequencer/Private/DisplayNodes/SequencerObjectBindingNode.cpp b/Engine/Source/Editor/Sequencer/Private/DisplayNodes/SequencerObjectBindingNode.cpp
--- a/Engine/Source/Editor/Sequencer/Private/DisplayNodes/SequencerObjectBindingNode.cpp
+++ b/Engine/Source/Editor/Sequencer/Private/DisplayNodes/SequencerObjectBindingNode.cpp
@@ -330,8 +330,8 @@ void FSequencerObjectBindingNode::SetDisplayName(const FText& NewDisplayName)
 
 bool FSequencerObjectBindingNode::CanDrag() const
 {
-	TSharedPtr<FSequencerDisplayNode> ParentNode = GetParent();
-	return ParentNode.IsValid() == false || ParentNode->GetType() != ESequencerNode::Object;
+	TSharedPtr<FSequencerDisplayNode> ParentSeqNode = GetParent();
+	return ParentSeqNode.IsValid() == false || ParentSeqNode->GetType() != ESequencerNode::Object;
 }
 
 
diff --git a/Engine/Source/Editor/Sequencer/Private/DisplayNodes/SequencerTrackNode.cpp b/Engine/Source/Editor/Sequencer/Private/DisplayNodes/SequencerTrackNode.cpp
--- a/Engine/Source/Editor/Sequencer/Private/DisplayNodes/SequencerTrackNode.cpp
+++ b/Engine/Source/Editor/Sequencer/Private/DisplayNodes/SequencerTrackNode.cpp
@@ -173,11 +173,11 @@ TSharedRef<SWidget> FSequencerTrackNode::GetCustomOutlinerContent()
 	else
 	{
 		FGuid ObjectBinding;
-		TSharedPtr<FSequencerDisplayNode> ParentNode = GetParent();
+		TSharedPtr<FSequencerDisplayNode> ParentSeqNode = GetParent();
 
-		if (ParentNode.IsValid() && (ParentNode->GetType() == ESequencerNode::Object))
+		if (ParentSeqNode.IsValid() && (ParentSeqNode->GetType() == ESequencerNode::Object))
 		{
-			ObjectBinding = StaticCastSharedPtr<FSequencerObjectBindingNode>(ParentNode)->GetObjectBinding();
+			ObjectBinding = StaticCastSharedPtr<FSequencerObjectBindingNode>(ParentSeqNode)->GetObjectBinding();
 		}
 		FBuildEditWidgetParams Params;
 		Params.NodeIsHovered = TAttribute<bool>::Create(TAttribute<bool>::FGetter::CreateSP(this, &FSequencerDisplayNode::IsHovered));
diff --git a/Engine/Source/Editor/Sequencer/Private/SSequencerSectionAreaView.cpp b/Engine/Source/Editor/Sequencer/Private/SSequencerSectionAreaView.cpp
--- a/Engine/Source/Editor/Sequencer/Private/SSequencerSectionAreaView.cpp
+++ b/Engine/Source/Editor/Sequencer/Private/SSequencerSectionAreaView.cpp
@@ -179,13 +179,13 @@ void SSequencerSectionAreaView::OnArrangeChildren( const FGeometry& AllottedGeom
 
 		int32 RowIndex = SectionInterface->GetSectionObject()->GetRowIndex();
 
-		EVisibility Visibility = Widget->GetVisibility();
-		if( ArrangedChildren.Accepts( Visibility ) )
+		EVisibility WidgetVisibility = Widget->GetVisibility();
+		if( ArrangedChildren.Accepts( WidgetVisibility ) )
 		{
 			FGeometry SectionGeometry = SequencerSectionUtils::GetSectionGeometry( AllottedGeometry, RowIndex, MaxTracks, Widget->GetDesiredSize().Y, SectionInterface, TimeToPixelConverter );
 			
 			ArrangedChildren.AddWidget( 
-				Visibility, 
+				WidgetVisibility, 
 				AllottedGeometry.MakeChild( Widget, SectionGeometry.Position, SectionGeometry.GetDrawSize() )
 				);
 		}
diff --git a/Engine/Source/Editor/Sequencer/Private/SSequencerTreeView.cpp b/Engine/Source/Editor/Sequencer/Private/SSequencerTreeView.cpp
--- a/Engine/Source/Editor/Sequencer/Private/SSequencerTreeView.cpp
+++ b/Engine/Source/Editor/Sequencer/Private/SSequencerTreeView.cpp
@@ -101,13 +101,13 @@ FReply SSequencerTreeViewRow::OnDragDetected( const FGeometry& InGeometry, const
 	return FReply::Unhandled();
 }
 
-TOptional<EItemDropZone> SSequencerTreeViewRow::OnCanAcceptDrop( const FDragDropEvent& DragDropEvent, EItemDropZone ItemDropZone, FDisplayNodeRef DisplayNode )
+TOptional<EItemDropZone> SSequencerTreeViewRow::OnCanAcceptDrop( const FDragDropEvent& DragDropEvent, EItemDropZone InItemDropZone, FDisplayNodeRef DisplayNode )
 {
 	TSharedPtr<FSequencerDisplayNodeDragDropOp> DragDropOp = DragDropEvent.GetOperationAs<FSequencerDisplayNodeDragDropOp>();
 	if ( DragDropOp.IsValid() )
 	{
 		DragDropOp->ResetToDefaultToolTip();
-		TOptional<EItemDropZone> AllowedDropZone = DisplayNode->CanDrop( *DragDropOp, ItemDropZone );
+		TOptional<EItemDropZone> AllowedDropZone = DisplayNode->CanDrop( *DragDropOp, InItemDropZone );
 		if ( AllowedDropZone.IsSet() == false )
 		{
 			DragDropOp->CurrentIconBrush = FEditorStyle::GetBrush( TEXT( "Graph.ConnectorFeedback.Error" ) );
@@ -117,12 +117,12 @@ TOptional<EItemDropZone> SSequencerTreeViewRow::OnCanAcceptDrop( const FDragDrop
 	return TOptional<EItemDropZone>();
 }
 
-FReply SSequencerTreeViewRow::OnAcceptDrop( const FDragDropEvent& DragDropEvent, EItemDropZone ItemDropZone, FDisplayNodeRef DisplayNode )
+FReply SSequencerTreeViewRow::OnAcceptDrop( const FDragDropEvent& DragDropEvent, EItemDropZone InItemDropZone, FDisplayNodeRef DisplayNode )
 {
 	TSharedPtr<FSequencerDisplayNodeDragDropOp> DragDropOp = DragDropEvent.GetOperationAs<FSequencerDisplayNodeDragDropOp>();
 	if ( DragDropOp.IsValid())
 	{
-		DisplayNode->Drop( DragDropOp->GetDraggedNodes(), ItemDropZone );
+		DisplayNode->Drop( DragDropOp->GetDraggedNodes(), InItemDropZone );
 		return FReply::Handled();
 	}
 	return FReply::Unhandled();
diff --git a/Engine/Source/Editor/TranslationEditor/Private/TranslationPickerFloatingWindow.cpp b/Engine/Source/Editor/TranslationEditor/Private/TranslationPickerFloatingWindow.cpp
--- a/Engine/Source/Editor/TranslationEditor/Private/TranslationPickerFloatingWindow.cpp
+++ b/Engine/Source/Editor/TranslationEditor/Private/TranslationPickerFloatingWindow.cpp
@@ -163,13 +163,13 @@ FText STranslationPickerFloatingWindow::GetTextFromWidget(TSharedRef<SWidget> Wi
 	}
 	else if (Widget->GetTypeAsString() == "SToolTip")
 	{
-		SToolTip* ToolTip = ((SToolTip*)&Widget.Get());
-		if (ToolTip != nullptr)
+		SToolTip* ToolTipWidget = ((SToolTip*)&Widget.Get());
+		if (ToolTipWidget != nullptr)
 		{
-			OriginalText = GetTextFromWidget(ToolTip->GetContentWidget());
+			OriginalText = GetTextFromWidget(ToolTipWidget->GetContentWidget());
 			if (OriginalText.IsEmpty())
 			{
-				OriginalText = ToolTip->GetTextTooltip();
+				OriginalText = ToolTipWidget->GetTextTooltip();
 			}
 		}
 	}
diff --git a/Engine/Source/Editor/TranslationEditor/Private/TranslationPickerWidget.cpp b/Engine/Source/Editor/TranslationEditor/Private/TranslationPickerWidget.cpp
--- a/Engine/Source/Editor/TranslationEditor/Private/TranslationPickerWidget.cpp
+++ b/Engine/Source/Editor/TranslationEditor/Private/TranslationPickerWidget.cpp
@@ -20,7 +20,7 @@ void STranslationWidgetPicker::Construct(const FArguments& InArgs)
 	FName StyleSet = FEditorStyle::GetStyleSetName();
 	FName StyleName = "Toolbar";
 
-	FText ToolTip = LOCTEXT("TranslationPickerTooltip", "Open the Translation Picker");
+	FText ToolTipText = LOCTEXT("TranslationPickerTooltip", "Open the Translation Picker");
 
 	// Create the content for our button
 	TSharedRef< SWidget > ButtonContent =
@@ -66,7 +66,7 @@ void STranslationWidgetPicker::Construct(const FArguments& InArgs)
 			.IsFocusable(false)
 
 			// Pass along the block's tool-tip string
-			.ToolTip(SNew(SToolTip).Text(ToolTip))
+			.ToolTip(SNew(SToolTip).Text(ToolTipText))
 			[
 				ButtonContent
 			]
diff --git a/Engine/Source/Editor/UMGEditor/Private/Designer/SDesignSurface.cpp b/Engine/Source/Editor/UMGEditor/Private/Designer/SDesignSurface.cpp
--- a/Engine/Source/Editor/UMGEditor/Private/Designer/SDesignSurface.cpp
+++ b/Engine/Source/Editor/UMGEditor/Private/Designer/SDesignSurface.cpp
@@ -422,7 +422,7 @@ bool SDesignSurface::ScrollToLocation(const FGeometry& MyGeometry, FVector2D Des
 	return ( ( NewPosition - DesiredCenterPosition ).SizeSquared() < 1.f );
 }
 
-bool SDesignSurface::ZoomToLocation(const FVector2D& CurrentSizeWithoutZoom, const FVector2D& DesiredSize, bool bDoneScrolling)
+bool SDesignSurface::ZoomToLocation(const FVector2D& CurrentSizeWithoutZoom, const FVector2D& InDesiredSize, bool bDoneScrolling)
 {
 	if ( bAllowContinousZoomInterpolation && ZoomLevelGraphFade.IsPlaying() )
 	{
@@ -437,9 +437,9 @@ bool SDesignSurface::ZoomToLocation(const FVector2D& CurrentSizeWithoutZoom, con
 	for ( int32 Zoom = 0; Zoom < DefaultZoomLevel; ++Zoom )
 	{
 		const FVector2D SizeWithZoom = (CurrentSizeWithoutZoom - ZoomToFitPadding) / ZoomLevels->GetZoomAmount(Zoom);
-		const FVector2D LeftOverSize = SizeWithZoom - DesiredSize;
+		const FVector2D LeftOverSize = SizeWithZoom - InDesiredSize;
 
-		if ( ( DesiredSize.X > SizeWithZoom.X ) || ( DesiredSize.Y > SizeWithZoom.Y ) )
+		if ( ( InDesiredSize.X > SizeWithZoom.X ) || ( InDesiredSize.Y > SizeWithZoom.Y ) )
 		{
 			// Use the previous zoom level, this one is too tight
 			DesiredZoom = FMath::Max<int32>(0, Zoom - 1);
diff --git a/Engine/Source/Editor/UMGEditor/Private/Designer/SDesignerView.cpp b/Engine/Source/Editor/UMGEditor/Private/Designer/SDesignerView.cpp
--- a/Engine/Source/Editor/UMGEditor/Private/Designer/SDesignerView.cpp
+++ b/Engine/Source/Editor/UMGEditor/Private/Designer/SDesignerView.cpp
@@ -1410,13 +1410,13 @@ FReply SDesignerView::OnMouseMove(const FGeometry& MyGeometry, const FPointerEve
 						{
 							const FSlateRenderTransform& AbsoluteToLocalTransform = Inverse(ParentGeometry.GetAccumulatedRenderTransform());
 
-							FWidgetTransform RenderTransform = WidgetPreview->RenderTransform;
-							RenderTransform.Translation += AbsoluteToLocalTransform.TransformVector(MouseEvent.GetCursorDelta());
+							FWidgetTransform WidgetRenderTransform = WidgetPreview->RenderTransform;
+							WidgetRenderTransform.Translation += AbsoluteToLocalTransform.TransformVector(MouseEvent.GetCursorDelta());
 
 							static const FName RenderTransformName(TEXT("RenderTransform"));
 
-							FObjectEditorUtils::SetPropertyValue<UWidget, FWidgetTransform>(WidgetPreview, RenderTransformName, RenderTransform);
-							FObjectEditorUtils::SetPropertyValue<UWidget, FWidgetTransform>(SelectedWidget.GetTemplate(), RenderTransformName, RenderTransform);
+							FObjectEditorUtils::SetPropertyValue<UWidget, FWidgetTransform>(WidgetPreview, RenderTransformName, WidgetRenderTransform);
+							FObjectEditorUtils::SetPropertyValue<UWidget, FWidgetTransform>(SelectedWidget.GetTemplate(), RenderTransformName, WidgetRenderTransform);
 						}
 					}
 				}
diff --git a/Engine/Source/Editor/UMGEditor/Private/Designer/STransformHandle.cpp b/Engine/Source/Editor/UMGEditor/Private/Designer/STransformHandle.cpp
--- a/Engine/Source/Editor/UMGEditor/Private/Designer/STransformHandle.cpp
+++ b/Engine/Source/Editor/UMGEditor/Private/Designer/STransformHandle.cpp
@@ -136,12 +136,12 @@ FReply STransformHandle::OnMouseMove(const FGeometry& MyGeometry, const FPointer
 		ETransformMode::Type TransformMode = Designer->GetTransformMode();
 		if ( TransformMode == ETransformMode::Render )
 		{
-			FWidgetTransform RenderTransform = Preview->RenderTransform;
+			FWidgetTransform PreviewRenderTransform = Preview->RenderTransform;
 
 			static const FName RenderTransformName(TEXT("RenderTransform"));
 
-			FObjectEditorUtils::SetPropertyValue<UWidget, FWidgetTransform>(Preview, RenderTransformName, RenderTransform);
-			FObjectEditorUtils::SetPropertyValue<UWidget, FWidgetTransform>(Template, RenderTransformName, RenderTransform);
+			FObjectEditorUtils::SetPropertyValue<UWidget, FWidgetTransform>(Preview, RenderTransformName, PreviewRenderTransform);
+			FObjectEditorUtils::SetPropertyValue<UWidget, FWidgetTransform>(Template, RenderTransformName, PreviewRenderTransform);
 		}
 	}
 
diff --git a/Engine/Source/Editor/UMGEditor/Private/WidgetBlueprintCompiler.cpp b/Engine/Source/Editor/UMGEditor/Private/WidgetBlueprintCompiler.cpp
--- a/Engine/Source/Editor/UMGEditor/Private/WidgetBlueprintCompiler.cpp
+++ b/Engine/Source/Editor/UMGEditor/Private/WidgetBlueprintCompiler.cpp
@@ -97,19 +97,19 @@ void FWidgetBlueprintCompiler::CreateFunctionList()
 
 void FWidgetBlueprintCompiler::ValidateWidgetNames()
 {
-	UWidgetBlueprint* Blueprint = WidgetBlueprint();
+	UWidgetBlueprint* WidgetBP = WidgetBlueprint();
 
 	TSharedPtr<FKismetNameValidator> ParentBPNameValidator;
-	if ( Blueprint->ParentClass != NULL )
+	if ( WidgetBP->ParentClass != nullptr )
 	{
-		UBlueprint* ParentBP = Cast<UBlueprint>(Blueprint->ParentClass->ClassGeneratedBy);
-		if ( ParentBP != NULL )
+		UBlueprint* ParentBP = Cast<UBlueprint>(WidgetBP->ParentClass->ClassGeneratedBy);
+		if ( ParentBP != nullptr )
 		{
 			ParentBPNameValidator = MakeShareable(new FKismetNameValidator(ParentBP));
 		}
 	}
 
-	Blueprint->WidgetTree->ForEachWidget([&] (UWidget* Widget) {
+	WidgetBP->WidgetTree->ForEachWidget([&] (UWidget* Widget) {
 		if ( ParentBPNameValidator.IsValid() && ParentBPNameValidator->IsValid(Widget->GetName()) != EValidatorResult::Ok )
 		{
 			// TODO Support renaming items, similar to timelines.
@@ -117,9 +117,9 @@ void FWidgetBlueprintCompiler::ValidateWidgetNames()
 			//// Use the viewer displayed Timeline name (without the _Template suffix) because it will be added later for appropriate checks.
 			//FString TimelineName = UTimelineTemplate::TimelineTemplateNameToVariableName(TimelineTemplate->GetFName());
 
-			//FName NewName = FBlueprintEditorUtils::FindUniqueKismetName(Blueprint, TimelineName);
+			//FName NewName = FBlueprintEditorUtils::FindUniqueKismetName(WidgetBP, TimelineName);
 			//MessageLog.Warning(*FString::Printf(*LOCTEXT("TimelineConflictWarning", "Found a timeline with a conflicting name (%s) - changed to %s.").ToString(), *TimelineTemplate->GetName(), *NewName.ToString()));
-			//FBlueprintEditorUtils::RenameTimeline(Blueprint, FName(*TimelineName), NewName);
+			//FBlueprintEditorUtils::RenameTimeline(WidgetBP, FName(*TimelineName), NewName);
 		}
 	});
 }
@@ -174,11 +174,11 @@ void FWidgetBlueprintCompiler::SaveSubObjectsFromCleanAndSanitizeClass(FSubobjec
 	check(ClassToClean == NewClass);
 	NewWidgetBlueprintClass = CastChecked<UWidgetBlueprintGeneratedClass>((UObject*)NewClass);
 
-	UWidgetBlueprint* Blueprint = WidgetBlueprint();
+	UWidgetBlueprint* WidgetBP = WidgetBlueprint();
 
 	// We need to save the widget tree to survive the initial sub-object clean blitz, 
 	// otherwise they all get renamed, and it causes early loading errors.
-	SubObjectsToSave.AddObject(Blueprint->WidgetTree);
+	SubObjectsToSave.AddObject(WidgetBP->WidgetTree);
 
 	// We need to save all the animations to survive the initial sub-object clean blitz, 
 	// otherwise they all get renamed, and it causes early loading errors.
@@ -193,14 +193,14 @@ void FWidgetBlueprintCompiler::CreateClassVariablesFromBlueprint()
 {
 	Super::CreateClassVariablesFromBlueprint();
 
-	UWidgetBlueprint* Blueprint = WidgetBlueprint();
-	UClass* ParentClass = Blueprint->ParentClass;
+	UWidgetBlueprint* WidgetBP = WidgetBlueprint();
+	UClass* ParentClass = WidgetBP->ParentClass;
 
 	ValidateWidgetNames();
 
 	// Build the set of variables based on the variable widgets in the widget tree.
 	TArray<UWidget*> Widgets;
-	Blueprint->WidgetTree->GetAllWidgets(Widgets);
+	WidgetBP->WidgetTree->GetAllWidgets(Widgets);
 
 	// Sort the widgets alphabetically
 	Widgets.Sort( []( const UWidget& Lhs, const UWidget& Rhs ) { return Rhs.GetFName() < Lhs.GetFName(); } );
@@ -237,7 +237,7 @@ void FWidgetBlueprintCompiler::CreateClassVariablesFromBlueprint()
 		{
 			const FString VariableName = Widget->IsGeneratedName() ? Widget->GetName() : Widget->GetLabelText().ToString();
 			WidgetProperty->SetMetaData(TEXT("DisplayName"), *VariableName);
-			WidgetProperty->SetMetaData(TEXT("Category"), *Blueprint->GetName());
+			WidgetProperty->SetMetaData(TEXT("Category"), *WidgetBP->GetName());
 			
 			WidgetProperty->SetPropertyFlags(CPF_BlueprintVisible);
 			WidgetProperty->SetPropertyFlags(CPF_Transient);
@@ -248,7 +248,7 @@ void FWidgetBlueprintCompiler::CreateClassVariablesFromBlueprint()
 	}
 
 	// Add movie scenes variables here
-	for(UWidgetAnimation* Animation : Blueprint->Animations)
+	for(UWidgetAnimation* Animation : WidgetBP->Animations)
 	{
 		FEdGraphPinType WidgetPinType(Schema->PC_Object, TEXT(""), Animation->GetClass(), false, true);
 		UProperty* AnimationProperty = CreateVariable(Animation->GetFName(), WidgetPinType);
@@ -266,19 +266,19 @@ void FWidgetBlueprintCompiler::CreateClassVariablesFromBlueprint()
 
 void FWidgetBlueprintCompiler::FinishCompilingClass(UClass* Class)
 {
-	UWidgetBlueprint* Blueprint = WidgetBlueprint();
+	UWidgetBlueprint* WidgetBP = WidgetBlueprint();
 
 	// Don't do a bunch of extra work on the skeleton generated class
-	if ( Blueprint->SkeletonGeneratedClass != Class )
+	if ( WidgetBP->SkeletonGeneratedClass != Class )
 	{
 		UWidgetBlueprintGeneratedClass* BPGClass = CastChecked<UWidgetBlueprintGeneratedClass>(Class);
-		if( !Blueprint->bHasBeenRegenerated )
+		if( !WidgetBP->bHasBeenRegenerated )
 		{
-			FBlueprintEditorUtils::ForceLoadMembers(Blueprint->WidgetTree);
+			FBlueprintEditorUtils::ForceLoadMembers(WidgetBP->WidgetTree);
 		}
-		BPGClass->WidgetTree = BPGClass->DesignerWidgetTree = DuplicateObject<UWidgetTree>(Blueprint->WidgetTree, BPGClass);
+		BPGClass->WidgetTree = BPGClass->DesignerWidgetTree = DuplicateObject<UWidgetTree>(WidgetBP->WidgetTree, BPGClass);
 
-		for ( const UWidgetAnimation* Animation : Blueprint->Animations )
+		for ( const UWidgetAnimation* Animation : WidgetBP->Animations )
 		{
 			UWidgetAnimation* ClonedAnimation = DuplicateObject<UWidgetAnimation>(Animation, BPGClass, *( Animation->GetName() + TEXT("_INST") ));
 
@@ -288,16 +288,16 @@ void FWidgetBlueprintCompiler::FinishCompilingClass(UClass* Class)
 		// Only check bindings on a full compile.  Also don't check them if we're regenerating on load,
 		// that has a nasty tendency to fail because the other dependent classes that may also be blueprints
 		// might not be loaded yet.
-		const bool bIsLoading = Blueprint->bIsRegeneratingOnLoad;
+		const bool bIsLoading = WidgetBP->bIsRegeneratingOnLoad;
 		if ( bIsFullCompile )
 		{
 			// Convert all editor time property bindings into a list of bindings
 			// that will be applied at runtime.  Ensure all bindings are still valid.
-			for ( const FDelegateEditorBinding& EditorBinding : Blueprint->Bindings )
+			for ( const FDelegateEditorBinding& EditorBinding : WidgetBP->Bindings )
 			{
-				if ( bIsLoading || EditorBinding.IsBindingValid(Class, Blueprint, MessageLog) )
+				if ( bIsLoading || EditorBinding.IsBindingValid(Class, WidgetBP, MessageLog) )
 				{
-					BPGClass->Bindings.Add(EditorBinding.ToRuntimeBinding(Blueprint));
+					BPGClass->Bindings.Add(EditorBinding.ToRuntimeBinding(WidgetBP));
 				}
 			}
 		}
@@ -312,7 +312,7 @@ void FWidgetBlueprintCompiler::FinishCompilingClass(UClass* Class)
 		});
 	}
 
-	UClass* ParentClass = Blueprint->ParentClass;
+	UClass* ParentClass = WidgetBP->ParentClass;
 	for ( TUObjectPropertyBase<UWidget*>* WidgetProperty : TFieldRange<TUObjectPropertyBase<UWidget*>>( ParentClass ) )
 	{
 		if ( WidgetProperty->HasMetaData( "BindWidget" ) && ( !WidgetProperty->HasMetaData( "OptionalWidget" ) || !WidgetProperty->GetBoolMetaData( "OptionalWidget" ) ) )
diff --git a/Engine/Source/Editor/UnrealEd/Private/Animation/DebugSkelMeshComponent.cpp b/Engine/Source/Editor/UnrealEd/Private/Animation/DebugSkelMeshComponent.cpp
--- a/Engine/Source/Editor/UnrealEd/Private/Animation/DebugSkelMeshComponent.cpp
+++ b/Engine/Source/Editor/UnrealEd/Private/Animation/DebugSkelMeshComponent.cpp
@@ -43,8 +43,7 @@ public:
 			// We don't want to draw the mesh geometry to the hit testing render target
 			// so that we can get to triangle strips that are partially obscured by other
 			// triangle strips easier.
-			const bool bSelectable = false;
-			GetMeshElementsConditionallySelectable(Views, ViewFamily, bSelectable, VisibilityMap, Collector);
+			GetMeshElementsConditionallySelectable(Views, ViewFamily, /*bSelectable=*/false, VisibilityMap, Collector);
 		}
 
 		//@todo - the rendering thread should never read from UObjects directly!  These are race conditions, the properties should be mirrored on the proxy
@@ -366,15 +365,15 @@ bool UDebugSkelMeshComponent::ShouldCPUSkin()
 }
 
 
-void UDebugSkelMeshComponent::PostInitMeshObject(FSkeletalMeshObject* MeshObject)
+void UDebugSkelMeshComponent::PostInitMeshObject(FSkeletalMeshObject* InMeshObject)
 {
-	Super::PostInitMeshObject( MeshObject );
+	Super::PostInitMeshObject( InMeshObject );
 
-	if (MeshObject)
+	if (InMeshObject)
 	{
 		if(bDrawBoneInfluences)
 		{
-			MeshObject->EnableBlendWeightRendering(true, BonesOfInterest);
+			InMeshObject->EnableBlendWeightRendering(true, BonesOfInterest);
 		}
 	}
 }
diff --git a/Engine/Source/Editor/UnrealEd/Private/AutoReimport/ReimportFeedbackContext.cpp b/Engine/Source/Editor/UnrealEd/Private/AutoReimport/ReimportFeedbackContext.cpp
--- a/Engine/Source/Editor/UnrealEd/Private/AutoReimport/ReimportFeedbackContext.cpp
+++ b/Engine/Source/Editor/UnrealEd/Private/AutoReimport/ReimportFeedbackContext.cpp
@@ -44,8 +44,7 @@ class SWidgetStack : public SVerticalBox
 	virtual FVector2D ComputeDesiredSize(float) const override
 	{
 		const float Lerp = SizeCurve.GetLerp();
-		FVector2D DesiredSize = ComputeTotalSize() * Lerp + StartSizeOffset * (1.f-Lerp);
-		return DesiredSize;
+		return ComputeTotalSize() * Lerp + StartSizeOffset * (1.f-Lerp);
 	}
 
 	virtual void OnArrangeChildren( const FGeometry& AllottedGeometry, FArrangedChildren& ArrangedChildren ) const override
