VR Editor: Disabled VR Editor when simulating in editor is on

diff --git a/Engine/Source/Editor/LevelEditor/Private/LevelEditorActions.cpp b/Engine/Source/Editor/LevelEditor/Private/LevelEditorActions.cpp
--- a/Engine/Source/Editor/LevelEditor/Private/LevelEditorActions.cpp
+++ b/Engine/Source/Editor/LevelEditor/Private/LevelEditorActions.cpp
@@ -2111,7 +2111,7 @@ void FLevelEditorActionCallbacks::ToggleVR()
 bool FLevelEditorActionCallbacks::ToggleVR_CanExecute()
 {
 	IVREditorModule& VREditorModule = IVREditorModule::Get();
-	return VREditorModule.IsVREditorAvailable();
+	return VREditorModule.IsVREditorAvailable() && !GEditor->bIsSimulatingInEditor;
 }
 
 
