Changing window closing to tab closing when Sequencer is invoked in VR Editor

diff --git a/Engine/Source/Editor/LevelEditor/Private/SLevelEditor.cpp b/Engine/Source/Editor/LevelEditor/Private/SLevelEditor.cpp
--- a/Engine/Source/Editor/LevelEditor/Private/SLevelEditor.cpp
+++ b/Engine/Source/Editor/LevelEditor/Private/SLevelEditor.cpp
@@ -537,8 +537,8 @@ void SLevelEditor::AttachSequencer( TSharedPtr<SWidget> SequencerWidget, TShared
 			if (FGlobalTabmanager::Get()->OnOverrideDockableAreaRestore_Handler.IsBound())
 			{
 				// @todo vreditor: more general vr editor tab manager should handle windows instead
-				// Destroy the original window so we just work with the override window
-				Tab->GetParentWindow().Get()->RequestDestroyWindow();
+				// Close the original tab so we just work with the override window
+				Tab->RequestCloseTab();
 			}
 		}
 		else
