Fix non-unity compile
#rnx

diff --git a/Engine/Source/Editor/LevelEditor/Public/ILevelViewport.h b/Engine/Source/Editor/LevelEditor/Public/ILevelViewport.h
--- a/Engine/Source/Editor/LevelEditor/Public/ILevelViewport.h
+++ b/Engine/Source/Editor/LevelEditor/Public/ILevelViewport.h
@@ -8,6 +8,7 @@
 
 class FLevelEditorViewportClient;
 class FViewport;
+class FSceneViewport;
 
 /**
  * Public interface to SLevelViewport
