Sequencer now initiates a slow task when opening via the Cinematic's dropdown in the UI. This helps with first time opening of large cinematics that require compressing audio and animation data.

diff --git a/Engine/Source/Editor/LevelEditor/Private/LevelEditorToolBar.cpp b/Engine/Source/Editor/LevelEditor/Private/LevelEditorToolBar.cpp
--- a/Engine/Source/Editor/LevelEditor/Private/LevelEditorToolBar.cpp
+++ b/Engine/Source/Editor/LevelEditor/Private/LevelEditorToolBar.cpp
@@ -55,6 +55,7 @@
 #include "Misc/ConfigCacheIni.h"
 #include "ILauncherPlatform.h"
 #include "LauncherPlatformModule.h"
+#include "Misc/ScopedSlowTask.h"
 
 namespace LevelEditorActionHelpers
 {
@@ -2233,6 +2234,9 @@ void FLevelEditorToolBar::OnCinematicsActorPicked( AActor* Actor )
 	}
 	else if (ALevelSequenceActor* LevelSequenceActor = Cast<ALevelSequenceActor>(Actor))
 	{
+		FScopedSlowTask SlowTask(1.f, NSLOCTEXT("LevelToolBarCinematicsMenu", "LoadSequenceSlowTask", "Loading Level Sequence..."));
+		SlowTask.MakeDialog();
+		SlowTask.EnterProgressFrame();
 		UObject* Asset = LevelSequenceActor->LevelSequence.TryLoad();
 
 		if (Asset != nullptr)
