Merging //UE4/Dev-Main to Dev-Build (//UE4/Dev-Build)

#rnx

diff --git a/Engine/Plugins/Online/OnlineFramework/Source/Party/Private/SocialToolkit.cpp b/Engine/Plugins/Online/OnlineFramework/Source/Party/Private/SocialToolkit.cpp
--- a/Engine/Plugins/Online/OnlineFramework/Source/Party/Private/SocialToolkit.cpp
+++ b/Engine/Plugins/Online/OnlineFramework/Source/Party/Private/SocialToolkit.cpp
@@ -637,6 +637,7 @@ void USocialToolkit::HandleReadFriendsListComplete(int32 LocalUserNum, bool bWas
 
 		//@todo DanH: This isn't actually quite correct - some of these could actually just be friend info for pending invites, not fully accepted friends. Should piece out the list into respective categories and process each separately (or make the associated event determination more complex)
 		ProcessUserList(FriendsList, SubsystemType, OnFriendshipEstablished());
+		OnQueryFriendsListSuccess(SubsystemType, FriendsList);
 	}
 	else
 	{
@@ -659,6 +660,7 @@ void USocialToolkit::HandleQueryBlockedPlayersComplete(const FUniqueNetId& UserI
 			TArray<TSharedRef<FOnlineBlockedPlayer>> BlockedPlayers;
 			FriendsInterface->GetBlockedPlayers(UserId, BlockedPlayers);
 			ProcessUserList(BlockedPlayers, SubsystemType, OnUserBlocked());
+			OnQueryBlockedPlayersSuccess(SubsystemType, BlockedPlayers);
 		}
 		else
 		{
@@ -681,6 +683,7 @@ void USocialToolkit::HandleQueryRecentPlayersComplete(const FUniqueNetId& UserId
 			TArray<TSharedRef<FOnlineRecentPlayer>> RecentPlayers;
 			FriendsInterface->GetRecentPlayers(UserId, Namespace, RecentPlayers);
 			ProcessUserList(RecentPlayers, SubsystemType, OnRecentPlayerAdded());
+			OnQueryRecentPlayersSuccess(SubsystemType, RecentPlayers);
 		}
 		else
 		{
diff --git a/Engine/Plugins/Online/OnlineFramework/Source/Party/Private/User/SocialUser.cpp b/Engine/Plugins/Online/OnlineFramework/Source/Party/Private/User/SocialUser.cpp
--- a/Engine/Plugins/Online/OnlineFramework/Source/Party/Private/User/SocialUser.cpp
+++ b/Engine/Plugins/Online/OnlineFramework/Source/Party/Private/User/SocialUser.cpp
@@ -871,7 +871,6 @@ bool USocialUser::HasSentPartyInvite(const FOnlinePartyTypeId& PartyTypeId) cons
 		const FUniqueNetIdRepl UserId = GetUserId(ESocialSubsystem::Primary);
 		if (ensure(LocalUserId.IsValid()) && UserId.IsValid())
 		{
-			// No advertised party info, check to see if this user has sent an invite
 			TArray<TSharedRef<IOnlinePartyJoinInfo>> AllPendingInvites;
 			if (PartyInterface->GetPendingInvites(*LocalUserId, AllPendingInvites))
 			{
diff --git a/Engine/Plugins/Online/OnlineFramework/Source/Party/Public/SocialToolkit.h b/Engine/Plugins/Online/OnlineFramework/Source/Party/Public/SocialToolkit.h
--- a/Engine/Plugins/Online/OnlineFramework/Source/Party/Public/SocialToolkit.h
+++ b/Engine/Plugins/Online/OnlineFramework/Source/Party/Public/SocialToolkit.h
@@ -108,14 +108,17 @@ protected:
 
 	virtual void OnOwnerLoggedIn();
 	virtual void OnOwnerLoggedOut();
-	virtual void OnAllUsersInitialized() {}
-
+	
 	virtual void NotifyFriendInviteFailed(const FUniqueNetId& InvitedUserId, const FString& InvitedUserName, ESendFriendInviteFailureReason FailureReason, bool bCanShow = true) {}
 
 	void QueryFriendsLists();
 	void QueryBlockedPlayers();
 	void QueryRecentPlayers();
 
+	virtual void OnQueryFriendsListSuccess(ESocialSubsystem SubsystemType, const TArray<TSharedRef<FOnlineFriend>>& FriendsList) {}
+	virtual void OnQueryBlockedPlayersSuccess(ESocialSubsystem SubsystemType, const TArray<TSharedRef<FOnlineBlockedPlayer>>& BlockedPlayers) {}
+	virtual void OnQueryRecentPlayersSuccess(ESocialSubsystem SubsystemType, const TArray<TSharedRef<FOnlineRecentPlayer>>& FriendsList) {}
+	
 	/** The type of SocialUser to create to represent known users */
 	TSubclassOf<USocialUser> SocialUserClass;
 
@@ -215,6 +218,5 @@ private:
 	mutable FOnRelationshipEstablished OnRecentPlayerAddedEvent;
 	
 	mutable FOnKnownUserInitialized OnKnownUserInitializedEvent;
-	mutable FBasicToolkitEvent OnFinishedStartupQueriesEvent;
 	mutable FBasicToolkitEvent OnToolkitResetEvent;
 };
diff --git a/Engine/Plugins/Online/OnlineSubsystem/Source/Private/OnlineError.cpp b/Engine/Plugins/Online/OnlineSubsystem/Source/Private/OnlineError.cpp
--- a/Engine/Plugins/Online/OnlineSubsystem/Source/Private/OnlineError.cpp
+++ b/Engine/Plugins/Online/OnlineSubsystem/Source/Private/OnlineError.cpp
@@ -18,6 +18,7 @@ FOnlineError::FOnlineError(const TCHAR* const ErrorCodeIn)
 
 FOnlineError::FOnlineError(const int32 ErrorCodeIn)
 	: bSucceeded(false)
+	, Result(EOnlineErrorResult::FailExtended)
 {
 	SetFromErrorCode(ErrorCodeIn);
 }
@@ -96,37 +97,43 @@ void FOnlineError::SetFromErrorCode(EOnlineErrorResult InResult, const FString&
 
 FOnlineError::FOnlineError(bool bSucceededIn)
 	: bSucceeded(bSucceededIn)
+	, Result(bSucceededIn ? EOnlineErrorResult::Success : EOnlineErrorResult::Unknown)
 {
 }
 
 FOnlineError::FOnlineError(const FString& ErrorCodeIn)
 	: bSucceeded(false)
+	, Result(EOnlineErrorResult::FailExtended)
 {
 	SetFromErrorCode(ErrorCodeIn);
 }
 
 FOnlineError::FOnlineError(FString&& ErrorCodeIn)
 	: bSucceeded(false)
+	, Result(EOnlineErrorResult::FailExtended)
 {
 	SetFromErrorCode(MoveTemp(ErrorCodeIn));
 }
 
+FOnlineError::FOnlineError(const FText& ErrorMessageIn)
+	: bSucceeded(false)
+	, Result(EOnlineErrorResult::FailExtended)
+{
+	SetFromErrorMessage(ErrorMessageIn);
+}
+
 void FOnlineError::SetFromErrorCode(const FString& ErrorCodeIn)
 {
 	ErrorCode = ErrorCodeIn;
 	ErrorRaw = ErrorCodeIn;
+	Result = EOnlineErrorResult::FailExtended;
 }
 
 void FOnlineError::SetFromErrorCode(FString&& ErrorCodeIn)
 {
 	ErrorCode = MoveTemp(ErrorCodeIn);
 	ErrorRaw = ErrorCode;
-}
-
-FOnlineError::FOnlineError(const FText& ErrorMessageIn)
-	: bSucceeded(false)
-{
-	SetFromErrorMessage(ErrorMessageIn);
+	Result = EOnlineErrorResult::FailExtended;
 }
 
 void FOnlineError::SetFromErrorMessage(const FText& ErrorMessageIn)
@@ -134,6 +141,7 @@ void FOnlineError::SetFromErrorMessage(const FText& ErrorMessageIn)
 	ErrorMessage = ErrorMessageIn;
 	ErrorCode = FTextInspector::GetKey(ErrorMessageIn).Get(GenericErrorCode);
 	ErrorRaw = ErrorMessageIn.ToString();
+	Result = EOnlineErrorResult::FailExtended;
 }
 
 FString FOnlineError::ToLogString() const
diff --git a/Engine/Source/Editor/AdvancedPreviewScene/Private/AdvancedPreviewScene.cpp b/Engine/Source/Editor/AdvancedPreviewScene/Private/AdvancedPreviewScene.cpp
--- a/Engine/Source/Editor/AdvancedPreviewScene/Private/AdvancedPreviewScene.cpp
+++ b/Engine/Source/Editor/AdvancedPreviewScene/Private/AdvancedPreviewScene.cpp
@@ -23,6 +23,7 @@
 #include "AdvancedPreviewSceneCommands.h"
 #include "Framework/Commands/UICommandList.h"
 #include "Framework/Application/SlateApplication.h"
+#include "Editor.h"
 
 FAdvancedPreviewScene::FAdvancedPreviewScene(ConstructionValues CVS, float InFloorOffset)
 	: FPreviewScene(CVS)
@@ -91,6 +92,12 @@ FAdvancedPreviewScene::FAdvancedPreviewScene(ConstructionValues CVS, float InFlo
 	bUseSkylight = true;
 
 	BindCommands();
+
+	// since advance preview scenes are used in conjunction with material/mesh editors we should match their feature level with the global one
+	if (GIsEditor && GEditor != nullptr)
+	{
+		PreviewWorld->ChangeFeatureLevel(GEditor->DefaultWorldFeatureLevel);
+	}
 }
 
 FAdvancedPreviewScene::~FAdvancedPreviewScene()
diff --git a/Engine/Source/Editor/EditorStyle/Private/SlateEditorStyle.cpp b/Engine/Source/Editor/EditorStyle/Private/SlateEditorStyle.cpp
--- a/Engine/Source/Editor/EditorStyle/Private/SlateEditorStyle.cpp
+++ b/Engine/Source/Editor/EditorStyle/Private/SlateEditorStyle.cpp
@@ -4760,6 +4760,23 @@ void FSlateEditorStyle::FStyle::SetupLevelEditorStyle()
 		Set("LevelEditor.SourceControl.Problem", new IMAGE_BRUSH("Icons/icon_source_control_40x_problem", Icon40x40));
 		Set("LevelEditor.SourceControl.Problem.Small", new IMAGE_BRUSH("Icons/icon_source_control_40x_problem", Icon20x20));
 
+		Set("LevelEditor.PreviewMode.Enabled", new IMAGE_BRUSH("Icons/icon_PreviewMode_SM5_Enabled_40x", Icon40x40));
+		Set("LevelEditor.PreviewMode.Disabled", new IMAGE_BRUSH("Icons/icon_PreviewMode_SM5_Disabled_40x", Icon40x40));
+		Set("LevelEditor.PreviewMode.SM5.Enabled", new IMAGE_BRUSH("Icons/icon_PreviewMode_SM5_Enabled_40x", Icon40x40));
+		Set("LevelEditor.PreviewMode.SM5.Disabled", new IMAGE_BRUSH("Icons/icon_PreviewMode_SM5_Enabled_40x", Icon40x40));
+		Set("LevelEditor.PreviewMode.SM4.Enabled", new IMAGE_BRUSH("Icons/icon_PreviewMode_SM4_Enabled_40x", Icon40x40));
+		Set("LevelEditor.PreviewMode.SM4.Disabled", new IMAGE_BRUSH("Icons/icon_PreviewMode_SM4_Enabled_40x", Icon40x40));
+		Set("LevelEditor.PreviewMode.AndroidES2.Enabled", new IMAGE_BRUSH("Icons/icon_PreviewMode_AndroidES2_Enabled_40x", Icon40x40));
+		Set("LevelEditor.PreviewMode.AndroidES2.Disabled", new IMAGE_BRUSH("Icons/icon_PreviewMode_AndroidES2_Disabled_40x", Icon40x40));
+		Set("LevelEditor.PreviewMode.AndroidES31.Enabled", new IMAGE_BRUSH("Icons/icon_PreviewMode_AndroidES31_Enabled_40x", Icon40x40));
+		Set("LevelEditor.PreviewMode.AndroidES31.Disabled", new IMAGE_BRUSH("Icons/icon_PreviewMode_AndroidES31_Disabled_40x", Icon40x40));
+		Set("LevelEditor.PreviewMode.AndroidVulkan.Enabled", new IMAGE_BRUSH("Icons/icon_PreviewMode_AndroidVulkan_Enabled_40x", Icon40x40));
+		Set("LevelEditor.PreviewMode.AndroidVulkan.Disabled", new IMAGE_BRUSH("Icons/icon_PreviewMode_AndroidVulkan_Disabled_40x", Icon40x40));
+		Set("LevelEditor.PreviewMode.iOS.Enabled", new IMAGE_BRUSH("Icons/icon_PreviewMode_iOS_Enabled_40x", Icon40x40));
+		Set("LevelEditor.PreviewMode.iOS.Disabled", new IMAGE_BRUSH("Icons/icon_PreviewMode_iOS_Disabled_40x", Icon40x40));
+		Set("LevelEditor.PreviewMode.HTML5.Enabled", new IMAGE_BRUSH("Icons/icon_PreviewMode_HTML5_Enabled_40x", Icon40x40));
+		Set("LevelEditor.PreviewMode.HTML5.Disabled", new IMAGE_BRUSH("Icons/icon_PreviewMode_HTML5_Disabled_40x", Icon40x40));
+
 		Set("LevelEditor.ViewOptions", new IMAGE_BRUSH("Icons/icon_view_40x", Icon40x40));
 		Set( "LevelEditor.ViewOptions.Small", new IMAGE_BRUSH( "Icons/icon_view_40x", Icon20x20 ) );
 
diff --git a/Engine/Source/Editor/LandscapeEditor/Private/LandscapeEdModePaintTools.cpp b/Engine/Source/Editor/LandscapeEditor/Private/LandscapeEdModePaintTools.cpp
--- a/Engine/Source/Editor/LandscapeEditor/Private/LandscapeEdModePaintTools.cpp
+++ b/Engine/Source/Editor/LandscapeEditor/Private/LandscapeEdModePaintTools.cpp
@@ -818,7 +818,7 @@ public:
 						{
 						case ELandscapeToolFlattenMode::Terrace:
 							{
-								const FTransform& LocalToWorld = Target.LandscapeInfo->GetLandscapeProxy()->ActorToWorld();
+								const FTransform& LocalToWorld = this->Target.LandscapeInfo->GetLandscapeProxy()->ActorToWorld();
 								float ScaleZ = LocalToWorld.GetScale3D().Z;
 								float TranslateZ = LocalToWorld.GetTranslation().Z;
 								float TerraceInterval = UISettings->TerraceInterval;
diff --git a/Engine/Source/Editor/LevelEditor/Private/LevelEditor.cpp b/Engine/Source/Editor/LevelEditor/Private/LevelEditor.cpp
--- a/Engine/Source/Editor/LevelEditor/Private/LevelEditor.cpp
+++ b/Engine/Source/Editor/LevelEditor/Private/LevelEditor.cpp
@@ -1701,6 +1701,12 @@ void FLevelEditorModule::BindGlobalLevelEditorCommands()
 		FCanExecuteAction(),
 		FIsActionChecked::CreateStatic( &FLevelEditorActionCallbacks::IsMaterialQualityLevelChecked, (EMaterialQualityLevel::Type)EMaterialQualityLevel::High ) );
 
+	ActionList.MapAction(
+		Commands.ToggleFeatureLevelPreview,
+		FExecuteAction::CreateStatic(&FLevelEditorActionCallbacks::ToggleFeatureLevelPreview),
+		FIsActionChecked::CreateStatic(&FLevelEditorActionCallbacks::IsFeatureLevelPreviewEnabled),
+		FIsActionChecked::CreateStatic(&FLevelEditorActionCallbacks::IsFeatureLevelPreviewActive));
+
 	ActionList.MapAction(
 		Commands.PreviewPlatformOverride_DefaultES2,
 		FExecuteAction::CreateStatic(&FLevelEditorActionCallbacks::SetPreviewPlatform, FName(), ERHIFeatureLevel::ES2),
diff --git a/Engine/Source/Editor/LevelEditor/Private/LevelEditorActions.cpp b/Engine/Source/Editor/LevelEditor/Private/LevelEditorActions.cpp
--- a/Engine/Source/Editor/LevelEditor/Private/LevelEditorActions.cpp
+++ b/Engine/Source/Editor/LevelEditor/Private/LevelEditorActions.cpp
@@ -567,15 +567,34 @@ bool FLevelEditorActionCallbacks::IsMaterialQualityLevelChecked( EMaterialQualit
 	return TestQualityLevel == MaterialQualityLevel;
 }
 
+void FLevelEditorActionCallbacks::ToggleFeatureLevelPreview()
+{
+	GEditor->ToggleFeatureLevelPreview();
+}
+
+bool FLevelEditorActionCallbacks::IsFeatureLevelPreviewEnabled()
+{
+	if (GEditor->PlayWorld || GUnrealEd->bIsSimulatingInEditor)
+	{
+		return false;
+	}
+	return GEditor->IsFeatureLevelPreviewEnabled();
+}
+
+bool FLevelEditorActionCallbacks::IsFeatureLevelPreviewActive()
+{
+	return GEditor->IsFeatureLevelPreviewEnabled() && GEditor->IsFeatureLevelPreviewActive();
+}
+
 void FLevelEditorActionCallbacks::SetPreviewPlatform(FName MaterialQualityPlatform, ERHIFeatureLevel::Type PreviewFeatureLevel)
 {
 	GEditor->SetPreviewPlatform(MaterialQualityPlatform, PreviewFeatureLevel);
 }
 
-bool FLevelEditorActionCallbacks::IsPreviewPlatformChecked(FName MaterialQualityPlatform, ERHIFeatureLevel::Type PreviewFeatureLevel)
+bool FLevelEditorActionCallbacks::IsPreviewPlatformChecked(FName InMaterialQualityPlatform, ERHIFeatureLevel::Type InPreviewFeatureLevel)
 {
 	const FName& PreviewPlatform = UMaterialShaderQualitySettings::Get()->GetPreviewPlatform();
-	return PreviewPlatform == MaterialQualityPlatform && PreviewFeatureLevel == GetWorld()->FeatureLevel;
+	return PreviewPlatform == InMaterialQualityPlatform && InPreviewFeatureLevel == GEditor->PreviewFeatureLevel;
 }
 
 void FLevelEditorActionCallbacks::SetFeatureLevelPreview(ERHIFeatureLevel::Type InPreviewFeatureLevel)
@@ -585,7 +604,7 @@ void FLevelEditorActionCallbacks::SetFeatureLevelPreview(ERHIFeatureLevel::Type
 
 bool FLevelEditorActionCallbacks::IsFeatureLevelPreviewChecked(ERHIFeatureLevel::Type InPreviewFeatureLevel)
 {
-	return InPreviewFeatureLevel == GetWorld()->FeatureLevel;
+	return GEditor->PreviewFeatureLevel == InPreviewFeatureLevel;
 }
 
 bool FLevelEditorActionCallbacks::IsFeatureLevelPreviewAvailable(ERHIFeatureLevel::Type InPreviewFeatureLevel)
@@ -3277,13 +3296,15 @@ void FLevelEditorCommands::RegisterCommands()
 	UI_COMMAND(MaterialQualityLevel_Medium, "Medium", "Sets material quality in the scene to medium.", EUserInterfaceActionType::RadioButton, FInputChord());
 	UI_COMMAND(MaterialQualityLevel_High, "High", "Sets material quality in the scene to high.", EUserInterfaceActionType::RadioButton, FInputChord());
 
-	UI_COMMAND(PreviewPlatformOverride_DefaultES2, "Default Mobile / HTML5 Preview", "Use default mobile settings (no quality overrides).", EUserInterfaceActionType::RadioButton, FInputChord());
-	UI_COMMAND(PreviewPlatformOverride_AndroidGLES2, "Android Preview", "Mobile preview using Android's quality settings.", EUserInterfaceActionType::RadioButton, FInputChord());
+	UI_COMMAND(ToggleFeatureLevelPreview, "Preview Platform", "Preview Platform", EUserInterfaceActionType::ToggleButton, FInputChord());
+
+	UI_COMMAND(PreviewPlatformOverride_AndroidGLES2, "Preview as Android ES2", "Mobile preview using Android's quality settings.", EUserInterfaceActionType::RadioButton, FInputChord());
+	UI_COMMAND(PreviewPlatformOverride_DefaultES2, "Preview as HTML5", "HTML5 preview.", EUserInterfaceActionType::RadioButton, FInputChord());
 
 	UI_COMMAND(PreviewPlatformOverride_DefaultES31, "Default High-End Mobile", "Use default mobile settings (no quality overrides).", EUserInterfaceActionType::RadioButton, FInputChord());
-	UI_COMMAND(PreviewPlatformOverride_AndroidGLES31, "Android GLES3.1 Preview", "Mobile preview using Android ES3.1 quality settings.", EUserInterfaceActionType::RadioButton, FInputChord());
-	UI_COMMAND(PreviewPlatformOverride_AndroidVulkanES31, "Android Vulkan Preview", "Mobile preview using Android Vulkan quality settings.", EUserInterfaceActionType::RadioButton, FInputChord());
-	UI_COMMAND(PreviewPlatformOverride_IOSMetalES31, "iOS Preview", "Mobile preview using iOS material quality settings.", EUserInterfaceActionType::RadioButton, FInputChord());
+	UI_COMMAND(PreviewPlatformOverride_AndroidGLES31, "Preview as Android ES31", "Mobile preview using Android ES3.1 quality settings.", EUserInterfaceActionType::RadioButton, FInputChord());
+	UI_COMMAND(PreviewPlatformOverride_AndroidVulkanES31, "Preview as Android Vulkan", "Mobile preview using Android Vulkan quality settings.", EUserInterfaceActionType::RadioButton, FInputChord());
+	UI_COMMAND(PreviewPlatformOverride_IOSMetalES31, "Preview as iOS", "Mobile preview using iOS material quality settings.", EUserInterfaceActionType::RadioButton, FInputChord());
 
 
 	UI_COMMAND( ConnectToSourceControl, "Connect to Source Control...", "Opens a dialog to connect to source control.", EUserInterfaceActionType::Button, FInputChord());
@@ -3299,8 +3320,8 @@ void FLevelEditorCommands::RegisterCommands()
 	{
 		NSLOCTEXT("LevelEditorCommands", "FeatureLevelPreviewType_ES2", "Mobile / HTML5"),
 		NSLOCTEXT("LevelEditorCommands", "FeatureLevelPreviewType_ES31", "High-End Mobile"),
-		NSLOCTEXT("LevelEditorCommands", "FeatureLevelPreviewType_SM4", "Shader Model 4"),
-		NSLOCTEXT("LevelEditorCommands", "FeatureLevelPreviewType_SM5", "Shader Model 5"),
+		NSLOCTEXT("LevelEditorCommands", "FeatureLevelPreviewType_SM4", "Preview as SM4"),
+		NSLOCTEXT("LevelEditorCommands", "FeatureLevelPreviewType_SM5", "No Preview Device. (View as SM5)"),
 	};
 
 	static const FText FeatureLevelToolTips[ERHIFeatureLevel::Num] = 
diff --git a/Engine/Source/Editor/LevelEditor/Private/LevelEditorToolBar.cpp b/Engine/Source/Editor/LevelEditor/Private/LevelEditorToolBar.cpp
--- a/Engine/Source/Editor/LevelEditor/Private/LevelEditorToolBar.cpp
+++ b/Engine/Source/Editor/LevelEditor/Private/LevelEditorToolBar.cpp
@@ -56,6 +56,11 @@
 #include "ILauncherPlatform.h"
 #include "LauncherPlatformModule.h"
 #include "Misc/ScopedSlowTask.h"
+#include "MaterialShaderQualitySettings.h"
+
+#define LOCTEXT_NAMESPACE "EditorViewport"
+#include "RHIShaderPlatformDefinitions.inl"
+#undef LOCTEXT_NAMESPACE
 
 namespace LevelEditorActionHelpers
 {
@@ -1274,6 +1279,120 @@ TSharedRef< SWidget > FLevelEditorToolBar::MakeLevelEditorToolBar( const TShared
 	}
 	ToolbarBuilder.EndSection();
 
+	ToolbarBuilder.BeginSection(NAME_None);
+	{
+		struct FPreviewModeFunctionality
+		{
+			static FText GetPreviewModeTooltip()
+			{
+				UMaterialShaderQualitySettings* MaterialShaderQualitySettings = UMaterialShaderQualitySettings::Get();
+				const FName& PreviewPlatform = MaterialShaderQualitySettings->GetPreviewPlatform();
+
+				EShaderPlatform PreviewShaderPlatform = ShaderFormatToLegacyShaderPlatform(PreviewPlatform);
+				if (PreviewShaderPlatform == SP_NumPlatforms)
+				{
+					PreviewShaderPlatform = GetFeatureLevelShaderPlatform(GEditor->PreviewFeatureLevel);
+				}
+
+				EShaderPlatform MaxRHIFeatureLevelPlatform = GetFeatureLevelShaderPlatform(GMaxRHIFeatureLevel);
+
+				if (GEditor->PreviewFeatureLevel == GMaxRHIFeatureLevel)
+				{
+                    const FText& RenderingAsPlatformName = GetFriendlyShaderPlatformName(MaxRHIFeatureLevelPlatform);
+                    return FText::Format(LOCTEXT("PreviewModeRenderingAs", "Viewing as {0}."), RenderingAsPlatformName);
+				}
+				else
+				{
+                    const FText& RenderingAsPlatformName = GetFriendlyShaderPlatformName(GWorld->FeatureLevel == GMaxRHIFeatureLevel ? MaxRHIFeatureLevelPlatform : PreviewShaderPlatform);
+                    const FText& SwitchToPlatformName = GetFriendlyShaderPlatformName(GWorld->FeatureLevel == GMaxRHIFeatureLevel ? PreviewShaderPlatform : MaxRHIFeatureLevelPlatform);
+                    if (GWorld->FeatureLevel == GMaxRHIFeatureLevel)
+                    {
+                        return FText::Format(LOCTEXT("PreviewModeViewingAsSwitchTo", "Viewing as {0}. Click to preview as {1}."), RenderingAsPlatformName, SwitchToPlatformName);
+                    }
+                    else
+                    {
+                        return FText::Format(LOCTEXT("PreviewModePreviewingAsSwitchTo", "Previewing as {0}. Click to view as {1}."), RenderingAsPlatformName, SwitchToPlatformName);
+                    }
+				}
+			}
+
+			static FSlateIcon GetPreviewModeIcon()
+			{
+				UMaterialShaderQualitySettings* MaterialShaderQualitySettings = UMaterialShaderQualitySettings::Get();
+				const FName& PreviewPlatform = MaterialShaderQualitySettings->GetPreviewPlatform();
+
+				EShaderPlatform ShaderPlatform = ShaderFormatToLegacyShaderPlatform(PreviewPlatform);
+				if (ShaderPlatform == SP_NumPlatforms)
+				{
+					ShaderPlatform = GetFeatureLevelShaderPlatform(GEditor->PreviewFeatureLevel);
+				}
+				switch (ShaderPlatform)
+				{
+					case SP_OPENGL_ES3_1_ANDROID:
+					{
+						return FSlateIcon(FEditorStyle::GetStyleSetName(), GEditor->IsFeatureLevelPreviewActive() ? "LevelEditor.PreviewMode.AndroidES31.Enabled" : "LevelEditor.PreviewMode.AndroidES31.Disabled");
+					}
+					case SP_VULKAN_ES3_1_ANDROID:
+					{
+						return FSlateIcon(FEditorStyle::GetStyleSetName(), GEditor->IsFeatureLevelPreviewActive() ? "LevelEditor.PreviewMode.AndroidVulkan.Enabled" : "LevelEditor.PreviewMode.AndroidVulkan.Disabled");
+					}
+					case SP_METAL:
+					{
+						return FSlateIcon(FEditorStyle::GetStyleSetName(), GEditor->IsFeatureLevelPreviewActive() ? "LevelEditor.PreviewMode.iOS.Enabled" : "LevelEditor.PreviewMode.iOS.Disabled");
+					}
+                	case SP_OPENGL_ES2_WEBGL:
+					{
+						return FSlateIcon(FEditorStyle::GetStyleSetName(), GEditor->IsFeatureLevelPreviewActive() ? "LevelEditor.PreviewMode.HTML5.Enabled" : "LevelEditor.PreviewMode.HTML5.Disabled");
+					}
+				}
+				switch (GEditor->PreviewFeatureLevel)
+				{
+					case ERHIFeatureLevel::SM5:
+					{
+						return FSlateIcon(FEditorStyle::GetStyleSetName(), GEditor->IsFeatureLevelPreviewActive() ? "LevelEditor.PreviewMode.SM5.Enabled" : "LevelEditor.PreviewMode.SM5.Disabled");
+					}
+					case ERHIFeatureLevel::SM4:
+					{
+						return FSlateIcon(FEditorStyle::GetStyleSetName(), GEditor->IsFeatureLevelPreviewActive() ? "LevelEditor.PreviewMode.SM4.Enabled" : "LevelEditor.PreviewMode.SM4.Enabled");
+					}
+					case ERHIFeatureLevel::ES2:
+					{
+						return FSlateIcon(FEditorStyle::GetStyleSetName(), GEditor->IsFeatureLevelPreviewActive() ? "LevelEditor.PreviewMode.AndroidES2.Enabled" : "LevelEditor.PreviewMode.AndroidES2.Disabled");
+					}
+					case ERHIFeatureLevel::ES3_1:
+					{
+						return FSlateIcon(FEditorStyle::GetStyleSetName(), GEditor->IsFeatureLevelPreviewActive() ? "LevelEditor.PreviewMode.Enabled" : "LevelEditor.PreviewMode.Disabled");
+					}
+					default:
+					{
+						return FSlateIcon(FEditorStyle::GetStyleSetName(), GEditor->IsFeatureLevelPreviewActive() ? "LevelEditor.PreviewMode.Enabled" : "LevelEditor.PreviewMode.Disabled");
+					}
+				}
+			}
+		};
+
+		ToolbarBuilder.AddToolBarButton(
+			FLevelEditorCommands::Get().ToggleFeatureLevelPreview,
+			NAME_None,
+			LOCTEXT("PreviewModeActivate", "Preview Mode"),
+        	TAttribute<FText>::Create(&FPreviewModeFunctionality::GetPreviewModeTooltip),
+        	TAttribute<FSlateIcon>::Create(&FPreviewModeFunctionality::GetPreviewModeIcon)
+			);
+
+		FUIAction PreviewModeMenuCanExecute;
+		PreviewModeMenuCanExecute.CanExecuteAction = FIsActionChecked::CreateStatic(&FLevelEditorActionCallbacks::IsFeatureLevelPreviewEnabled);
+
+		ToolbarBuilder.AddComboButton(
+			PreviewModeMenuCanExecute,
+			FOnGetContent::CreateStatic(&FLevelEditorToolBar::GeneratePreviewModeMenu, InCommandList),
+			LOCTEXT("PreviewModeCombo", "Preview Mode"),
+			LOCTEXT("PreviewModeCombo_ToolTip", "Preview Mode settings"),
+			FSlateIcon(FEditorStyle::GetStyleSetName(), "LevelEditor.GameSettings"),
+            true
+			);
+	}
+	ToolbarBuilder.EndSection();
+
 	ToolbarBuilder.BeginSection( NAME_None );
 	{
 		ToolbarBuilder.AddComboButton(
@@ -1658,45 +1777,6 @@ TSharedRef< SWidget > FLevelEditorToolBar::GenerateBuildMenuContent( TSharedRef<
 	return MenuBuilder.MakeWidget();
 }
 
-static void MakeES2PreviewPlatformOverrideMenu(FMenuBuilder& MenuBuilder)
-{
-	MenuBuilder.BeginSection("LevelEditorShaderModelPreview", NSLOCTEXT("LevelToolBarViewMenu", "ES2PreviewPlatformOverrideHeading", "Preview Platform"));
-	{
-		MenuBuilder.AddMenuEntry(FLevelEditorCommands::Get().PreviewPlatformOverride_DefaultES2);
-		MenuBuilder.AddMenuEntry(FLevelEditorCommands::Get().PreviewPlatformOverride_AndroidGLES2);
-	}
-	MenuBuilder.EndSection();
-}
-
-static void MakeES31PreviewPlatformOverrideMenu(FMenuBuilder& MenuBuilder)
-{
-	MenuBuilder.BeginSection("LevelEditorShaderModelPreview", NSLOCTEXT("LevelToolBarViewMenu", "ES31PreviewPlatformOverrideHeading", "Preview Platform"));
-	{
-		MenuBuilder.AddMenuEntry(FLevelEditorCommands::Get().PreviewPlatformOverride_DefaultES31);
-
-		bool bAndroidBuildForES31 = false;
-		bool bAndroidSupportsVulkan = false;
-		bool bIOSSupportsMetal = false;
-		GConfig->GetBool(TEXT("/Script/AndroidRuntimeSettings.AndroidRuntimeSettings"), TEXT("bBuildForES31"), bAndroidBuildForES31, GEngineIni);
-		GConfig->GetBool(TEXT("/Script/AndroidRuntimeSettings.AndroidRuntimeSettings"), TEXT("bSupportsVulkan"), bAndroidSupportsVulkan, GEngineIni);
-		GConfig->GetBool(TEXT("/Script/IOSRuntimeSettings.IOSRuntimeSettings"), TEXT("bSupportsMetal"), bIOSSupportsMetal, GEngineIni);
-
-		if(bAndroidBuildForES31)
-		{
-			MenuBuilder.AddMenuEntry(FLevelEditorCommands::Get().PreviewPlatformOverride_AndroidGLES31);
-		}
-		if(bAndroidSupportsVulkan)
-		{
-			MenuBuilder.AddMenuEntry(FLevelEditorCommands::Get().PreviewPlatformOverride_AndroidVulkanES31);
-		}
-		if(bIOSSupportsMetal)
-		{
-			MenuBuilder.AddMenuEntry(FLevelEditorCommands::Get().PreviewPlatformOverride_IOSMetalES31);
-		}
-	}
-	MenuBuilder.EndSection();
-}
-
 static void MakeMaterialQualityLevelMenu( FMenuBuilder& MenuBuilder )
 {
 	MenuBuilder.BeginSection("LevelEditorMaterialQualityLevel", NSLOCTEXT( "LevelToolBarViewMenu", "MaterialQualityLevelHeading", "Material Quality Level" ) );
@@ -1708,34 +1788,6 @@ static void MakeMaterialQualityLevelMenu( FMenuBuilder& MenuBuilder )
 	MenuBuilder.EndSection();
 }
 
-static void MakeShaderModelPreviewMenu(FMenuBuilder& MenuBuilder)
-{
-	MenuBuilder.BeginSection("LevelEditorShaderModelPreview", NSLOCTEXT("LevelToolBarViewMenu", "FeatureLevelPreviewHeading", "Preview Rendering Level"));
-	{
-		for (int32 i = GMaxRHIFeatureLevel; i >= 0; --i)
-		{
-			switch (i)
-			{
-				case ERHIFeatureLevel::ES2:
-					MenuBuilder.AddSubMenu(
-						FLevelEditorCommands::Get().FeatureLevelPreview[i]->GetLabel(),
-						FLevelEditorCommands::Get().FeatureLevelPreview[i]->GetDescription(),
-						FNewMenuDelegate::CreateStatic(&MakeES2PreviewPlatformOverrideMenu));
-					break;
-				case ERHIFeatureLevel::ES3_1:
-					MenuBuilder.AddSubMenu(
-						FLevelEditorCommands::Get().FeatureLevelPreview[i]->GetLabel(),
-						FLevelEditorCommands::Get().FeatureLevelPreview[i]->GetDescription(),
-						FNewMenuDelegate::CreateStatic(&MakeES31PreviewPlatformOverrideMenu));
-					break;
-				default:
-					MenuBuilder.AddMenuEntry(FLevelEditorCommands::Get().FeatureLevelPreview[i]);
-			}
-		}
-	}
-	MenuBuilder.EndSection();
-}
-
 static void MakeScalabilityMenu( FMenuBuilder& MenuBuilder )
 {
 	MenuBuilder.AddWidget(SNew(SScalabilitySettings), FText(), true);
@@ -1828,11 +1880,6 @@ TSharedRef< SWidget > FLevelEditorToolBar::GenerateQuickSettingsMenu( TSharedRef
 			LOCTEXT( "MaterialQualityLevelSubMenu", "Material Quality Level" ),
 			LOCTEXT( "MaterialQualityLevelSubMenu_ToolTip", "Sets the value of the CVar \"r.MaterialQualityLevel\" (low=0, high=1, medium=2). This affects materials via the QualitySwitch material expression." ),
 			FNewMenuDelegate::CreateStatic( &MakeMaterialQualityLevelMenu ) );
-
-		MenuBuilder.AddSubMenu(
-			LOCTEXT("FeatureLevelPreviewSubMenu", "Preview Rendering Level"),
-			LOCTEXT("FeatureLevelPreviewSubMenu_ToolTip", "Sets the rendering level used by the main editor"),
-			FNewMenuDelegate::CreateStatic(&MakeShaderModelPreviewMenu));
 	}
 	MenuBuilder.EndSection();
 
@@ -1893,6 +1940,79 @@ TSharedRef< SWidget > FLevelEditorToolBar::GenerateQuickSettingsMenu( TSharedRef
 }
 
 
+TSharedRef< SWidget > FLevelEditorToolBar::GeneratePreviewModeMenu( TSharedRef<FUICommandList> InCommandList )
+{
+#define LOCTEXT_NAMESPACE "LevelToolBarViewMenu"
+
+	// Get all menu extenders for this context menu from the level editor module
+	FLevelEditorModule& LevelEditorModule = FModuleManager::GetModuleChecked<FLevelEditorModule>( TEXT("LevelEditor") );
+	TArray<FLevelEditorModule::FLevelEditorMenuExtender> MenuExtenderDelegates = LevelEditorModule.GetAllLevelEditorToolbarViewMenuExtenders();
+
+	TArray<TSharedPtr<FExtender>> Extenders;
+	for (int32 i = 0; i < MenuExtenderDelegates.Num(); ++i)
+	{
+		if (MenuExtenderDelegates[i].IsBound())
+		{
+			Extenders.Add(MenuExtenderDelegates[i].Execute(InCommandList));
+		}
+	}
+	TSharedPtr<FExtender> MenuExtender = FExtender::Combine(Extenders);
+
+	const bool bShouldCloseWindowAfterMenuSelection = true;
+	FMenuBuilder MenuBuilder(bShouldCloseWindowAfterMenuSelection, InCommandList, MenuExtender);
+
+	MenuBuilder.BeginSection("EditorPreviewMode", LOCTEXT("EditorPreviewModeDevices", "Preview Devices"));
+
+	for (int32 i = GMaxRHIFeatureLevel; i >= 0; --i)
+	{
+		switch (i)
+		{
+			case ERHIFeatureLevel::ES2:
+				MenuBuilder.AddMenuEntry(FLevelEditorCommands::Get().PreviewPlatformOverride_AndroidGLES2);
+				MenuBuilder.AddMenuEntry(FLevelEditorCommands::Get().PreviewPlatformOverride_DefaultES2);
+				break;
+
+			case ERHIFeatureLevel::ES3_1:
+			{
+				//MenuBuilder.AddMenuEntry(FLevelEditorCommands::Get().PreviewPlatformOverride_DefaultES31);
+
+				bool bAndroidBuildForES31 = false;
+				GConfig->GetBool(TEXT("/Script/AndroidRuntimeSettings.AndroidRuntimeSettings"), TEXT("bBuildForES31"), bAndroidBuildForES31, GEngineIni);
+				if(bAndroidBuildForES31)
+				{
+					MenuBuilder.AddMenuEntry(FLevelEditorCommands::Get().PreviewPlatformOverride_AndroidGLES31);
+				}
+
+				bool bAndroidSupportsVulkan = false;
+				GConfig->GetBool(TEXT("/Script/AndroidRuntimeSettings.AndroidRuntimeSettings"), TEXT("bSupportsVulkan"), bAndroidSupportsVulkan, GEngineIni);
+				if(bAndroidSupportsVulkan)
+				{
+					MenuBuilder.AddMenuEntry(FLevelEditorCommands::Get().PreviewPlatformOverride_AndroidVulkanES31);
+				}
+
+				bool bIOSSupportsMetal = false;
+				GConfig->GetBool(TEXT("/Script/IOSRuntimeSettings.IOSRuntimeSettings"), TEXT("bSupportsMetal"), bIOSSupportsMetal, GEngineIni);
+				if(bIOSSupportsMetal)
+				{
+					MenuBuilder.AddMenuEntry(FLevelEditorCommands::Get().PreviewPlatformOverride_IOSMetalES31);
+				}
+
+				break;
+			}
+
+			default:
+				MenuBuilder.AddMenuEntry(FLevelEditorCommands::Get().FeatureLevelPreview[i]);
+		}
+	}
+
+    MenuBuilder.EndSection();
+
+#undef LOCTEXT_NAMESPACE
+
+	return MenuBuilder.MakeWidget();
+}
+
+
 TSharedRef< SWidget > FLevelEditorToolBar::GenerateSourceControlMenu(TSharedRef<FUICommandList> InCommandList)
 {
 #define LOCTEXT_NAMESPACE "LevelToolBarSourceControlMenu"
diff --git a/Engine/Source/Editor/LevelEditor/Private/LevelEditorToolBar.h b/Engine/Source/Editor/LevelEditor/Private/LevelEditorToolBar.h
--- a/Engine/Source/Editor/LevelEditor/Private/LevelEditorToolBar.h
+++ b/Engine/Source/Editor/LevelEditor/Private/LevelEditorToolBar.h
@@ -47,6 +47,13 @@ protected:
 	 */
 	static TSharedRef< SWidget > GenerateQuickSettingsMenu( TSharedRef<FUICommandList> InCommandList );
 
+	/**
+	 * Generates menu content for the preview mode combo button drop down menu
+	 *
+	 * @return	Menu content widget
+	 */
+	static TSharedRef< SWidget > GeneratePreviewModeMenu( TSharedRef<FUICommandList> InCommandList );
+
 	/**
 	 * Generates menu content for the source control combo button drop down menu
 	 *
diff --git a/Engine/Source/Editor/LevelEditor/Private/SLevelEditor.cpp b/Engine/Source/Editor/LevelEditor/Private/SLevelEditor.cpp
--- a/Engine/Source/Editor/LevelEditor/Private/SLevelEditor.cpp
+++ b/Engine/Source/Editor/LevelEditor/Private/SLevelEditor.cpp
@@ -160,6 +160,16 @@ void SLevelEditor::Construct( const SLevelEditor::FArguments& InArgs)
 	// @todo This is a hack to get this working for now. This won't work with multiple worlds
 	GEditor->GetEditorWorldContext(true).AddRef(World);
 
+	// Set the initial preview feature level.
+	UEditorEngine* Editor = (UEditorEngine*)GEngine;
+	World->ChangeFeatureLevel(Editor->PreviewFeatureLevel);
+
+	// Patch into the OnPreviewFeatureLevelChanged() delegate to swap out the current feature level with a user selection.
+	PreviewFeatureLevelChangedHandle = Editor->OnPreviewFeatureLevelChanged().AddLambda([this](ERHIFeatureLevel::Type NewFeatureLevel)
+		{
+			World->ChangeFeatureLevel(NewFeatureLevel);
+		});
+
 	FEditorDelegates::MapChange.AddRaw(this, &SLevelEditor::HandleEditorMapChange);
 	HandleEditorMapChange(MapChangeEventFlags::NewMap);
 }
@@ -278,6 +288,8 @@ SLevelEditor::~SLevelEditor()
 
 	FEditorDelegates::MapChange.RemoveAll(this);
 
+	CastChecked<UEditorEngine>(GEngine)->OnPreviewFeatureLevelChanged().Remove(PreviewFeatureLevelChangedHandle);
+	
 	if (GEditor)
 	{
 		GEditor->GetEditorWorldContext(true).RemoveRef(World);
diff --git a/Engine/Source/Editor/LevelEditor/Private/SLevelEditor.h b/Engine/Source/Editor/LevelEditor/Private/SLevelEditor.h
--- a/Engine/Source/Editor/LevelEditor/Private/SLevelEditor.h
+++ b/Engine/Source/Editor/LevelEditor/Private/SLevelEditor.h
@@ -252,4 +252,7 @@ private:
 
 	/** Weak pointer to the level editor's scene outliner */
 	TWeakPtr<ISceneOutliner> SceneOutlinerPtr;
+
+	/** Handle to the registered OnPreviewFeatureLevelChanged delegate. */
+	FDelegateHandle PreviewFeatureLevelChangedHandle;
 };
diff --git a/Engine/Source/Editor/LevelEditor/Private/SLevelViewport.cpp b/Engine/Source/Editor/LevelEditor/Private/SLevelViewport.cpp
--- a/Engine/Source/Editor/LevelEditor/Private/SLevelViewport.cpp
+++ b/Engine/Source/Editor/LevelEditor/Private/SLevelViewport.cpp
@@ -278,34 +278,12 @@ void SLevelViewport::ConstructViewportOverlayContent()
 					.ShadowOffset(FVector2D(1, 1))
 				]
 		]
+		// add feature level widget
 		+ SVerticalBox::Slot()
 		.AutoHeight()
 		.Padding(2.0f, 1.0f, 2.0f, 1.0f)
 		[
-			SNew(SHorizontalBox)
-			.Visibility(this, &SLevelViewport::GetCurrentFeatureLevelPreviewTextVisibility)
-			// Current level label
-			+ SHorizontalBox::Slot()
-			.AutoWidth()
-			.Padding(2.0f, 1.0f, 2.0f, 1.0f)
-			[
-				SNew(STextBlock)
-				.Text(this, &SLevelViewport::GetCurrentFeatureLevelPreviewText, true)
-				.Font(FEditorStyle::GetFontStyle(TEXT("MenuItem.Font")))
-				.ShadowOffset(FVector2D(1, 1))
-			]
-
-			// Current level
-			+ SHorizontalBox::Slot()
-				.AutoWidth()
-				.Padding(4.0f, 1.0f, 2.0f, 1.0f)
-				[
-					SNew(STextBlock)
-					.Text(this, &SLevelViewport::GetCurrentFeatureLevelPreviewText, false)
-					.Font(FEditorStyle::GetFontStyle(TEXT("MenuItem.Font")))
-					.ColorAndOpacity(FLinearColor(0.4f, 1.0f, 1.0f))
-					.ShadowOffset(FVector2D(1, 1))
-				]
+			BuildFeatureLevelWidget()
 		]
 		+ SVerticalBox::Slot()
 		.AutoHeight()
@@ -1045,6 +1023,8 @@ void SLevelViewport::OnMapChanged( UWorld* World, EMapChangeType MapChangeType )
 		}
 		World->EditorViews[LevelViewportClient->ViewportType].CamUpdated = false;
 
+		World->ChangeFeatureLevel(GWorld->FeatureLevel);
+
 		RedrawViewport(true);
 	}
 }
@@ -3324,38 +3304,6 @@ FText SLevelViewport::GetCurrentScreenPercentageText(bool bDrawOnlyLabel) const
 	return FText::FromString(FString::Printf(TEXT("%3d%%"), int32(GetLevelViewportClient().GetPreviewScreenPercentage())));
 }
 
-FText SLevelViewport::GetCurrentFeatureLevelPreviewText( bool bDrawOnlyLabel ) const
-{
-	FText LabelName;
-	FText FeatureLevelText;
-
-	if (bDrawOnlyLabel)
-	{
-		LabelName = LOCTEXT("FeatureLevelLabel", "Feature Level:");
-	}
-	else
-	{
-		auto* World = GetWorld();
-		if (World != nullptr)
-		{
-			const auto FeatureLevel = World->FeatureLevel;
-			if (FeatureLevel != GMaxRHIFeatureLevel)
-			{
-				FName FeatureLevelName;
-				GetFeatureLevelName(FeatureLevel, FeatureLevelName);
-				FeatureLevelText = FText::Format(LOCTEXT("FeatureLevel", "{0}"), FText::FromName(FeatureLevelName));
-			}
-		}
-	}
-	
-	if (bDrawOnlyLabel)
-	{
-		return LabelName;
-	}
-
-	return FeatureLevelText;
-}
-
 FText SLevelViewport::GetCurrentLevelText( bool bDrawOnlyLabel ) const
 {
 	// Display the current level and current level grid volume in the status bar
@@ -3408,18 +3356,6 @@ EVisibility SLevelViewport::GetCurrentLevelTextVisibility() const
 	return (&GetLevelViewportClient() == GCurrentLevelEditingViewportClient) ? ContentVisibility : EVisibility::Collapsed;
 }
 
-EVisibility SLevelViewport::GetCurrentFeatureLevelPreviewTextVisibility() const
-{
-	if (GetWorld())
-	{
-		return (GetWorld()->FeatureLevel != GMaxRHIFeatureLevel) ? EVisibility::SelfHitTestInvisible : EVisibility::Collapsed;
-	}
-	else
-	{
-		return EVisibility::Collapsed;
-	}
-}
-
 EVisibility SLevelViewport::GetCurrentScreenPercentageVisibility() const
 {
 	bool Visible = !IsPlayInEditorViewportActive() &&
diff --git a/Engine/Source/Editor/LevelEditor/Public/LevelEditorActions.h b/Engine/Source/Editor/LevelEditor/Public/LevelEditorActions.h
--- a/Engine/Source/Editor/LevelEditor/Public/LevelEditorActions.h
+++ b/Engine/Source/Editor/LevelEditor/Public/LevelEditorActions.h
@@ -579,6 +579,8 @@ public:
 
 	TSharedPtr< FUICommandInfo > FeatureLevelPreview[ERHIFeatureLevel::Num];
 	
+	TSharedPtr< FUICommandInfo > ToggleFeatureLevelPreview;
+
 	TSharedPtr< FUICommandInfo > PreviewPlatformOverride_DefaultES2;
 	TSharedPtr< FUICommandInfo > PreviewPlatformOverride_AndroidGLES2;
 
@@ -797,6 +799,9 @@ public:
 	static void AttachToSocketSelection(FName SocketName, AActor* ParentActorPtr);
 	static void SetMaterialQualityLevel( EMaterialQualityLevel::Type NewQualityLevel );
 	static bool IsMaterialQualityLevelChecked( EMaterialQualityLevel::Type TestQualityLevel );
+	static void ToggleFeatureLevelPreview();
+	static bool IsFeatureLevelPreviewEnabled();
+	static bool IsFeatureLevelPreviewActive();
 	static void SetPreviewPlatform(FName MaterialQualityPlatform,ERHIFeatureLevel::Type PreviewFeatureLevel);
 	static bool IsPreviewPlatformChecked(FName MaterialQualityPlatform, ERHIFeatureLevel::Type PreviewFeatureLevel);
 	static void SetFeatureLevelPreview(ERHIFeatureLevel::Type InFeatureLevel);
diff --git a/Engine/Source/Editor/LevelEditor/Public/SLevelViewport.h b/Engine/Source/Editor/LevelEditor/Public/SLevelViewport.h
--- a/Engine/Source/Editor/LevelEditor/Public/SLevelViewport.h
+++ b/Engine/Source/Editor/LevelEditor/Public/SLevelViewport.h
@@ -289,15 +289,9 @@ public:
 	/** Called to get the screen percentage preview text */
 	FText GetCurrentScreenPercentageText(bool bDrawOnlyLabel) const;
 
-	/** Called to get the feature level preview text */
-	FText GetCurrentFeatureLevelPreviewText( bool bDrawOnlyLabel ) const;
-	
 	/** @return The visibility of the current level text display */
 	virtual EVisibility GetCurrentLevelTextVisibility() const;
 
-	/** @return The visibility of the current feature level preview text display */
-	EVisibility GetCurrentFeatureLevelPreviewTextVisibility() const;
-
 	/** @return The visibility of the current screen percentage text display */
 	EVisibility GetCurrentScreenPercentageVisibility() const;
 
diff --git a/Engine/Source/Editor/MaterialEditor/Private/SMaterialEditorViewport.cpp b/Engine/Source/Editor/MaterialEditor/Private/SMaterialEditorViewport.cpp
--- a/Engine/Source/Editor/MaterialEditor/Private/SMaterialEditorViewport.cpp
+++ b/Engine/Source/Editor/MaterialEditor/Private/SMaterialEditorViewport.cpp
@@ -235,6 +235,19 @@ void SMaterialEditor3DPreviewViewport::Construct(const FArguments& InArgs)
 	MaterialEditorPtr = InArgs._MaterialEditor;
 	AdvancedPreviewScene = MakeShareable(new FAdvancedPreviewScene(FPreviewScene::ConstructionValues()));
 
+	// restore last used feature level
+	UWorld* PreviewWorld = AdvancedPreviewScene->GetWorld();
+	if (PreviewWorld != nullptr)
+	{
+		PreviewWorld->ChangeFeatureLevel(GWorld->FeatureLevel);
+	}	
+
+	UEditorEngine* Editor = CastChecked<UEditorEngine>(GEngine);
+	PreviewFeatureLevelChangedHandle = Editor->OnPreviewFeatureLevelChanged().AddLambda([this](ERHIFeatureLevel::Type NewFeatureLevel)
+		{
+			AdvancedPreviewScene->GetWorld()->ChangeFeatureLevel(NewFeatureLevel);
+		});
+
 	bShowGrid = false;
 
 	PreviewPrimType = TPT_None;
@@ -266,6 +279,8 @@ void SMaterialEditor3DPreviewViewport::Construct(const FArguments& InArgs)
 
 SMaterialEditor3DPreviewViewport::~SMaterialEditor3DPreviewViewport()
 {
+	CastChecked<UEditorEngine>(GEngine)->OnPreviewFeatureLevelChanged().Remove(PreviewFeatureLevelChangedHandle);
+	
 	UAssetViewerSettings::Get()->OnAssetViewerSettingsChanged().RemoveAll(this);
 	if (PreviewMeshComponent != nullptr)
 	{
@@ -378,6 +393,10 @@ bool SMaterialEditor3DPreviewViewport::SetPreviewAsset(UObject* InAsset)
 	// Add the new component to the scene
 	if (PreviewMeshComponent != nullptr)
 	{
+		if (GEditor->PreviewFeatureLevel <= ERHIFeatureLevel::ES3_1)
+		{
+			PreviewMeshComponent->SetMobility(EComponentMobility::Static);
+		}
 		AdvancedPreviewScene->AddComponent(PreviewMeshComponent, Transform);
 		AdvancedPreviewScene->SetFloorOffset(-PreviewMeshComponent->Bounds.Origin.Z + PreviewMeshComponent->Bounds.BoxExtent.Z);
 
@@ -707,6 +726,15 @@ void SMaterialEditor3DPreviewViewport::PopulateViewportOverlays(TSharedRef<class
 		[
 			SNew(SMaterialEditorViewportPreviewShapeToolBar, SharedThis(this))
 		];
+
+	// add the feature level display widget
+	Overlay->AddSlot()
+		.VAlign(VAlign_Top)
+		.HAlign(HAlign_Right)
+		.Padding(5.0f)
+		[
+			BuildFeatureLevelWidget()
+		];
 }
 
 EVisibility SMaterialEditor3DPreviewViewport::OnGetViewportContentVisibility() const
diff --git a/Engine/Source/Editor/MaterialEditor/Private/SMaterialEditorViewport.h b/Engine/Source/Editor/MaterialEditor/Private/SMaterialEditorViewport.h
--- a/Engine/Source/Editor/MaterialEditor/Private/SMaterialEditorViewport.h
+++ b/Engine/Source/Editor/MaterialEditor/Private/SMaterialEditorViewport.h
@@ -124,6 +124,9 @@ private:
 
 	/** Handle to the registered OnPropertyChangedHandle delegate. */
 	FDelegateHandle OnPropertyChangedHandleDelegateHandle;
+
+	/** Handle to the registered OnPreviewFeatureLevelChanged delegate. */
+	FDelegateHandle PreviewFeatureLevelChangedHandle;
 };
 
 /**
diff --git a/Engine/Source/Editor/PIEPreviewDeviceProfileSelector/Private/PIEPreviewDevice.cpp b/Engine/Source/Editor/PIEPreviewDeviceProfileSelector/Private/PIEPreviewDevice.cpp
--- a/Engine/Source/Editor/PIEPreviewDeviceProfileSelector/Private/PIEPreviewDevice.cpp
+++ b/Engine/Source/Editor/PIEPreviewDeviceProfileSelector/Private/PIEPreviewDevice.cpp
@@ -410,7 +410,11 @@ void FPIEPreviewDevice::ApplyRHIOverrides() const
 
 	if (PreviewPlatform != SP_NumPlatforms)
 	{
+		UMaterialShaderQualitySettings* MaterialShaderQualitySettings = UMaterialShaderQualitySettings::Get();
 		FName QualityPreviewShaderPlatform = LegacyShaderPlatformToShaderFormat(PreviewPlatform);
+		MaterialShaderQualitySettings->GetShaderPlatformQualitySettings(QualityPreviewShaderPlatform);
+		ERHIFeatureLevel::Type FeatureLevel = GetMaxSupportedFeatureLevel(PreviewPlatform);
+
 		UMaterialShaderQualitySettings::Get()->SetPreviewPlatform(QualityPreviewShaderPlatform);
 	}
 
@@ -490,4 +494,4 @@ FString FPIEPreviewDevice::GetProfile() const
 int32 FPIEPreviewDevice::GetWindowClientHeight() const
 {
 	return WindowHeight - WindowTitleBarSize * DPIScaleFactor;
-}
\ No newline at end of file
+}
diff --git a/Engine/Source/Editor/Persona/Private/PersonaToolkit.cpp b/Engine/Source/Editor/Persona/Private/PersonaToolkit.cpp
--- a/Engine/Source/Editor/Persona/Private/PersonaToolkit.cpp
+++ b/Engine/Source/Editor/Persona/Private/PersonaToolkit.cpp
@@ -119,6 +119,10 @@ void FPersonaToolkit::CreatePreviewScene(const FPersonaToolkitArgs& PersonaToolk
 
 			// Create the preview component
 			UDebugSkelMeshComponent* SkeletalMeshComponent = NewObject<UDebugSkelMeshComponent>(Actor);
+			if (GEditor->PreviewFeatureLevel <= ERHIFeatureLevel::ES3_1)
+			{
+				SkeletalMeshComponent->SetMobility(EComponentMobility::Static);
+			}
 			PreviewScene->AddComponent(SkeletalMeshComponent, FTransform::Identity);
 			PreviewScene->SetPreviewMeshComponent(SkeletalMeshComponent);
 
@@ -347,4 +351,4 @@ FName FPersonaToolkit::GetContext() const
 	}
 
 	return NAME_None;
-}
\ No newline at end of file
+}
diff --git a/Engine/Source/Editor/Persona/Private/SAnimationEditorViewport.cpp b/Engine/Source/Editor/Persona/Private/SAnimationEditorViewport.cpp
--- a/Engine/Source/Editor/Persona/Private/SAnimationEditorViewport.cpp
+++ b/Engine/Source/Editor/Persona/Private/SAnimationEditorViewport.cpp
@@ -65,6 +65,51 @@ void SAnimationEditorViewport::Construct(const FArguments& InArgs, const FAnimat
 		);
 
 	Client->VisibilityDelegate.BindSP(this, &SAnimationEditorViewport::IsVisible);
+
+	// restore last used feature level
+	auto ScenePtr = PreviewScenePtr.Pin();
+	if (ScenePtr.IsValid())
+	{
+		UWorld* World = ScenePtr->GetWorld();
+		if (World != nullptr)
+		{
+			World->ChangeFeatureLevel(GWorld->FeatureLevel);
+		}
+	}
+
+	UEditorEngine* Editor = (UEditorEngine*)GEngine;
+	PreviewFeatureLevelChangedHandle = Editor->OnPreviewFeatureLevelChanged().AddLambda([this](ERHIFeatureLevel::Type NewFeatureLevel)
+		{
+			auto ScenePtr = PreviewScenePtr.Pin();
+			if (ScenePtr.IsValid())
+			{
+				UWorld* World = ScenePtr->GetWorld();
+				if (World != nullptr)
+				{
+					World->ChangeFeatureLevel(NewFeatureLevel);
+				}
+			}
+		});
+}
+
+SAnimationEditorViewport::~SAnimationEditorViewport()
+{
+	UEditorEngine* Editor = (UEditorEngine*)GEngine;
+	Editor->OnPreviewFeatureLevelChanged().Remove(PreviewFeatureLevelChangedHandle);
+}
+
+void SAnimationEditorViewport::PopulateViewportOverlays(TSharedRef<SOverlay> Overlay)
+{
+	SEditorViewport::PopulateViewportOverlays(Overlay);
+
+	// add the feature level display widget
+	Overlay->AddSlot()
+		.VAlign(VAlign_Bottom)
+		.HAlign(HAlign_Right)
+		.Padding(5.0f)
+		[
+			BuildFeatureLevelWidget()
+		];
 }
 
 TSharedRef<FEditorViewportClient> SAnimationEditorViewport::MakeEditorViewportClient()
diff --git a/Engine/Source/Editor/Persona/Private/SAnimationEditorViewport.h b/Engine/Source/Editor/Persona/Private/SAnimationEditorViewport.h
--- a/Engine/Source/Editor/Persona/Private/SAnimationEditorViewport.h
+++ b/Engine/Source/Editor/Persona/Private/SAnimationEditorViewport.h
@@ -88,6 +88,7 @@ public:
 	SLATE_END_ARGS()
 
 	void Construct(const FArguments& InArgs, const FAnimationEditorViewportRequiredArgs& InRequiredArgs);
+	virtual ~SAnimationEditorViewport();
 
 	/** Get the viewport toolbar widget */
 	TSharedPtr<SAnimViewportToolBar> GetViewportToolbar() const { return ViewportToolbar; }
@@ -97,6 +98,7 @@ protected:
 	virtual TSharedRef<FEditorViewportClient> MakeEditorViewportClient() override;
 	virtual TSharedPtr<SWidget> MakeViewportToolbar() override;
 	virtual void OnFocusViewportToSelection() override;
+	virtual void PopulateViewportOverlays(TSharedRef<SOverlay> Overlay) override;
 	// End of SEditorViewport interface
 
 	/**  Handle undo/redo by refreshing the viewport */
@@ -115,6 +117,9 @@ protected:
 	// The preview scene that we are viewing
 	TWeakPtr<class IPersonaPreviewScene> PreviewScenePtr;
 
+	// Handle to the registered OnPreviewFeatureLevelChanged delegate.
+	FDelegateHandle PreviewFeatureLevelChangedHandle;
+
 	// The asset editor we are embedded in
 	TWeakPtr<class FAssetEditorToolkit> AssetEditorToolkitPtr;
 
diff --git a/Engine/Source/Editor/StaticMeshEditor/Private/SStaticMeshEditorViewport.cpp b/Engine/Source/Editor/StaticMeshEditor/Private/SStaticMeshEditorViewport.cpp
--- a/Engine/Source/Editor/StaticMeshEditor/Private/SStaticMeshEditorViewport.cpp
+++ b/Engine/Source/Editor/StaticMeshEditor/Private/SStaticMeshEditorViewport.cpp
@@ -19,6 +19,7 @@
 #include "Widgets/Docking/SDockTab.h"
 #include "Engine/StaticMeshSocket.h"
 #include "Editor/UnrealEd/Public/SEditorViewportToolBarMenu.h"
+#include "Editor.h"
 
 #define HITPROXY_SOCKET	1
 
@@ -33,6 +34,19 @@ void SStaticMeshEditorViewport::Construct(const FArguments& InArgs)
 
 	PreviewScene->SetFloorOffset(-InArgs._ObjectToEdit->ExtendedBounds.Origin.Z + InArgs._ObjectToEdit->ExtendedBounds.BoxExtent.Z);
 
+	// restore last used feature level
+	UWorld* World = PreviewScene->GetWorld();
+	if (World != nullptr)
+	{
+		World->ChangeFeatureLevel(GWorld->FeatureLevel);
+	}
+
+	UEditorEngine* Editor = CastChecked<UEditorEngine>(GEngine);
+	PreviewFeatureLevelChangedHandle = Editor->OnPreviewFeatureLevelChanged().AddLambda([this](ERHIFeatureLevel::Type NewFeatureLevel)
+		{
+			PreviewScene->GetWorld()->ChangeFeatureLevel(NewFeatureLevel);
+		});
+
 	StaticMeshEditorPtr = InArgs._StaticMeshEditor;
 
 	StaticMesh = InArgs._ObjectToEdit;
@@ -44,10 +58,21 @@ void SStaticMeshEditorViewport::Construct(const FArguments& InArgs)
 	SEditorViewport::Construct( SEditorViewport::FArguments() );
 
 	PreviewMeshComponent = NewObject<UStaticMeshComponent>(GetTransientPackage(), NAME_None, RF_Transient );
+	if (GEditor->PreviewFeatureLevel <= ERHIFeatureLevel::ES3_1)
+	{
+		PreviewMeshComponent->SetMobility(EComponentMobility::Static);
+	}
 
 	SetPreviewMesh(StaticMesh);
 
-	ViewportOverlay->AddSlot()
+	FCoreUObjectDelegates::OnObjectPropertyChanged.AddRaw(this, &SStaticMeshEditorViewport::OnObjectPropertyChanged);
+}
+
+void SStaticMeshEditorViewport::PopulateViewportOverlays(TSharedRef<SOverlay> Overlay)
+{
+	SEditorViewport::PopulateViewportOverlays(Overlay);
+
+	Overlay->AddSlot()
 		.VAlign(VAlign_Top)
 		.HAlign(HAlign_Left)
 		.Padding(FMargin(10.0f, 40.0f, 10.0f, 10.0f))
@@ -55,7 +80,14 @@ void SStaticMeshEditorViewport::Construct(const FArguments& InArgs)
 			SAssignNew(OverlayTextVerticalBox, SVerticalBox)
 		];
 
-	FCoreUObjectDelegates::OnObjectPropertyChanged.AddRaw(this, &SStaticMeshEditorViewport::OnObjectPropertyChanged);
+	// this widget will display the current viewed feature level
+	Overlay->AddSlot()
+		.VAlign(VAlign_Bottom)
+		.HAlign(HAlign_Right)
+		.Padding(5.0f)
+		[
+			BuildFeatureLevelWidget()
+		];
 }
 
 SStaticMeshEditorViewport::SStaticMeshEditorViewport()
@@ -66,6 +98,8 @@ SStaticMeshEditorViewport::SStaticMeshEditorViewport()
 
 SStaticMeshEditorViewport::~SStaticMeshEditorViewport()
 {
+	CastChecked<UEditorEngine>(GEngine)->OnPreviewFeatureLevelChanged().Remove(PreviewFeatureLevelChangedHandle);
+
 	FCoreUObjectDelegates::OnObjectPropertyChanged.RemoveAll(this);
 	if (EditorViewportClient.IsValid())
 	{
@@ -233,6 +267,10 @@ void SStaticMeshEditorViewport::UpdatePreviewMesh(UStaticMesh* InStaticMesh, boo
 	}
 
 	PreviewMeshComponent = NewObject<UStaticMeshComponent>();
+	if (GEditor->PreviewFeatureLevel <= ERHIFeatureLevel::ES3_1)
+	{
+		PreviewMeshComponent->SetMobility(EComponentMobility::Static);
+	}
 
 	PreviewMeshComponent->SetStaticMesh(InStaticMesh);
 
diff --git a/Engine/Source/Editor/StaticMeshEditor/Private/SStaticMeshEditorViewport.h b/Engine/Source/Editor/StaticMeshEditor/Private/SStaticMeshEditorViewport.h
--- a/Engine/Source/Editor/StaticMeshEditor/Private/SStaticMeshEditorViewport.h
+++ b/Engine/Source/Editor/StaticMeshEditor/Private/SStaticMeshEditorViewport.h
@@ -127,6 +127,7 @@ protected:
 	virtual void BindCommands() override;
 	virtual void OnFocusViewportToSelection() override;
 	virtual TSharedPtr<SWidget> MakeViewportToolbar() override;
+	virtual void PopulateViewportOverlays(TSharedRef<SOverlay> Overlay) override;
 
 private:
 	/** Determines the visibility of the viewport. */
@@ -174,4 +175,7 @@ private:
 
 	/** Current LOD Selection where 0 is Auto */
 	int32 LODSelection;
+
+	/** Handle to the registered OnPreviewFeatureLevelChanged delegate. */
+	FDelegateHandle PreviewFeatureLevelChangedHandle;
 };
diff --git a/Engine/Source/Editor/UMGEditor/Private/Designer/SDesignerView.cpp b/Engine/Source/Editor/UMGEditor/Private/Designer/SDesignerView.cpp
--- a/Engine/Source/Editor/UMGEditor/Private/Designer/SDesignerView.cpp
+++ b/Engine/Source/Editor/UMGEditor/Private/Designer/SDesignerView.cpp
@@ -2023,7 +2023,7 @@ void SDesignerView::PopulateWidgetGeometryCache_Loop(FArrangedWidget& CurrentWid
 	{
 		bool bRespectLocks = IsRespectingLocks();
 
-		if (!CandidateUWidget->GetVisibilityInDesigner().IsVisible())
+		if (!CandidateUWidget->IsVisibleInDesigner())
 		{
 			bIncludeInHitTestGrid = false;
 		}
diff --git a/Engine/Source/Editor/UnrealEd/Classes/Editor/EditorEngine.h b/Engine/Source/Editor/UnrealEd/Classes/Editor/EditorEngine.h
--- a/Engine/Source/Editor/UnrealEd/Classes/Editor/EditorEngine.h
+++ b/Engine/Source/Editor/UnrealEd/Classes/Editor/EditorEngine.h
@@ -634,6 +634,13 @@ public:
 	/** The feature level we should use when loading or creating a new world */
 	ERHIFeatureLevel::Type DefaultWorldFeatureLevel;
 
+	/** The feature level we should use when loading or creating a new world */
+	ERHIFeatureLevel::Type PreviewFeatureLevel;
+
+	/** A delegate that is called when the preview feature level changes. Primarily used to switch a viewport's feature level. */
+	DECLARE_MULTICAST_DELEGATE_OneParam(FPreviewFeatureLevelChanged, ERHIFeatureLevel::Type);
+	FPreviewFeatureLevelChanged PreviewFeatureLevelChanged;
+
 	/** Whether or not the editor is currently compiling */
 	bool bIsCompiling;
 
@@ -2993,7 +3000,19 @@ public:
 	void OnSceneMaterialsModified();
 
 	/** Call this function to change the feature level and to override the material quality platform of the editor and PIE worlds */
-	void SetPreviewPlatform(const FName MaterialQualityPlatform, const ERHIFeatureLevel::Type PreviewFeatureLevel, const bool bSaveSettings = true);
+	void SetPreviewPlatform(const FName MaterialQualityPlatform, const ERHIFeatureLevel::Type InPreviewFeatureLevel, const bool bSaveSettings = true);
+
+	/** Toggle the feature level preview */
+	void ToggleFeatureLevelPreview();
+
+	/** Return whether the feature level preview is able to be enabled */
+	bool IsFeatureLevelPreviewEnabled() const;
+
+	/** Return whether the feature level preview is currently active */
+	bool IsFeatureLevelPreviewActive() const;
+
+	/** Return the delegate that is called when the preview feature level changes */
+	FPreviewFeatureLevelChanged& OnPreviewFeatureLevelChanged() { return PreviewFeatureLevelChanged; }
 
 protected:
 	/** Call this function to change the feature level of the editor and PIE worlds */
diff --git a/Engine/Source/Editor/UnrealEd/Private/EditorEngine.cpp b/Engine/Source/Editor/UnrealEd/Private/EditorEngine.cpp
--- a/Engine/Source/Editor/UnrealEd/Private/EditorEngine.cpp
+++ b/Engine/Source/Editor/UnrealEd/Private/EditorEngine.cpp
@@ -346,6 +346,7 @@ UEditorEngine::UEditorEngine(const FObjectInitializer& ObjectInitializer)
 	bIsEndingPlay = false;
 	NumOnlinePIEInstances = 0;
 	DefaultWorldFeatureLevel = GMaxRHIFeatureLevel;
+	PreviewFeatureLevel = DefaultWorldFeatureLevel;
 
 	bNotifyUndoRedoSelectionChange = true;
 
@@ -7414,6 +7415,7 @@ void UEditorEngine::OnSceneMaterialsModified()
 void UEditorEngine::SetMaterialsFeatureLevel(const ERHIFeatureLevel::Type InFeatureLevel)
 {
 	FScopedSlowTask SlowTask(100.f, NSLOCTEXT("Engine", "UpdatingMaterialsMessage", "Updating Materials"), true);
+	SlowTask.Visibility = ESlowTaskVisibility::ForceVisible;
 	SlowTask.MakeDialog();
 
 	//invalidate global bound shader states so they will be created with the new shaders the next time they are set (in SetGlobalBoundShaderState)
@@ -7425,59 +7427,36 @@ void UEditorEngine::SetMaterialsFeatureLevel(const ERHIFeatureLevel::Type InFeat
 	FGlobalComponentReregisterContext RecreateComponents;
 	FlushRenderingCommands();
 
-	SlowTask.EnterProgressFrame(5.0f);
+	// Clear all required global feature levels, we only require the preview feature level.
+	for (uint32 i = (uint32)ERHIFeatureLevel::ES2; i < (uint32)ERHIFeatureLevel::Num; i++)
+	{
+		UMaterialInterface::SetGlobalRequiredFeatureLevel((ERHIFeatureLevel::Type)i, false);
+	}
 
-	// Decrement refcount on old feature level
 	UMaterialInterface::SetGlobalRequiredFeatureLevel(InFeatureLevel, true);
 
-	SlowTask.EnterProgressFrame(50.0f);
-	UMaterial::AllMaterialsCacheResourceShadersForRendering();
+	SlowTask.EnterProgressFrame(35.0f);
+	UMaterial::AllMaterialsCacheResourceShadersForRendering(true);
 	
-	SlowTask.EnterProgressFrame(40.0f);
-	UMaterialInstance::AllMaterialsCacheResourceShadersForRendering();
+	SlowTask.EnterProgressFrame(35.0f);
+	UMaterialInstance::AllMaterialsCacheResourceShadersForRendering(true);
 
+	SlowTask.EnterProgressFrame(15.0f, NSLOCTEXT("Engine", "SlowTaskGlobalShaderMapMessage", "Compiling global shaders"));
 	CompileGlobalShaderMap(InFeatureLevel);
+
+	SlowTask.EnterProgressFrame(15.0f, NSLOCTEXT("Engine", "SlowTaskFinalizingMessage", "Finalizing"));
 	GShaderCompilingManager->ProcessAsyncResults(false, true);
-	SlowTask.EnterProgressFrame(5.0f);
 }
 
 void UEditorEngine::SetFeatureLevelPreview(const ERHIFeatureLevel::Type InPreviewFeatureLevel)
 {
-	if (DefaultWorldFeatureLevel != InPreviewFeatureLevel)
+	if (PreviewFeatureLevel != InPreviewFeatureLevel)
 	{
 		// Record this feature level as we want to use it for all subsequent level creation and loading
-		DefaultWorldFeatureLevel = InPreviewFeatureLevel;
-
-		FScopedSlowTask SlowTask(100.f, NSLOCTEXT("Engine", "ChangingPreviewRenderingLevelMessage", "Changing Preview Rendering Level"), true);
-		SlowTask.MakeDialog();
+		PreviewFeatureLevel = InPreviewFeatureLevel;
 
 		// first change the feature level for global/shared resources
 		SetMaterialsFeatureLevel(InPreviewFeatureLevel);
-
-		SlowTask.EnterProgressFrame(50.0f);
-
-		UWorld* MainWorld = GetEditorWorldContext().World();
-		if (MainWorld != nullptr)
-		{
-			MainWorld->ChangeFeatureLevel(InPreviewFeatureLevel, false);
-		}
-
-		SlowTask.EnterProgressFrame(25.0f);
-
-		// Update any currently running PIE sessions.
-		for (TObjectIterator<UWorld> It; It; ++It)
-		{
-			UWorld* ItWorld = *It;
-			if (ItWorld->WorldType == EWorldType::PIE)
-			{
-				ItWorld->ChangeFeatureLevel(InPreviewFeatureLevel, false);
-			}
-		}
-
-		SlowTask.EnterProgressFrame(25.0f);
-
-		GUnrealEd->OnSceneMaterialsModified();
-		GUnrealEd->RedrawAllViewports();
 	}
 }
 
@@ -7485,23 +7464,23 @@ void UEditorEngine::AllMaterialsCacheResourceShadersForRendering()
 {
 	FGlobalComponentRecreateRenderStateContext Recreate;
 	FlushRenderingCommands();
-	UMaterial::AllMaterialsCacheResourceShadersForRendering();
-	UMaterialInstance::AllMaterialsCacheResourceShadersForRendering();
+	UMaterial::AllMaterialsCacheResourceShadersForRendering(true);
+	UMaterialInstance::AllMaterialsCacheResourceShadersForRendering(true);
 }
 
-void UEditorEngine::SetPreviewPlatform(const FName MaterialQualityPlatform, const ERHIFeatureLevel::Type PreviewFeatureLevel, const bool bSaveSettings/* = true*/)
+void UEditorEngine::SetPreviewPlatform(const FName MaterialQualityPlatform, const ERHIFeatureLevel::Type InPreviewFeatureLevel, const bool bSaveSettings/* = true*/)
 {
 	// If we have specified a MaterialQualityPlatform ensure its feature level matches the requested feature level.
-	check(MaterialQualityPlatform.IsNone() || GetMaxSupportedFeatureLevel(ShaderFormatToLegacyShaderPlatform(MaterialQualityPlatform)) == PreviewFeatureLevel);
+	check(MaterialQualityPlatform.IsNone() || GetMaxSupportedFeatureLevel(ShaderFormatToLegacyShaderPlatform(MaterialQualityPlatform)) == InPreviewFeatureLevel);
 
 	UMaterialShaderQualitySettings* MaterialShaderQualitySettings = UMaterialShaderQualitySettings::Get();
 	const FName InitialPreviewPlatform = MaterialShaderQualitySettings->GetPreviewPlatform();
 	MaterialShaderQualitySettings->SetPreviewPlatform(MaterialQualityPlatform);
 
-	if (DefaultWorldFeatureLevel != PreviewFeatureLevel)
+	if (PreviewFeatureLevel != InPreviewFeatureLevel)
 	{
 		// a new feature level will recompile the materials and apply the effect of any 'material quality platform'
-		SetFeatureLevelPreview(PreviewFeatureLevel);
+		SetFeatureLevelPreview(InPreviewFeatureLevel);
 	}
 	else if (InitialPreviewPlatform != MaterialQualityPlatform)
 	{
@@ -7509,12 +7488,34 @@ void UEditorEngine::SetPreviewPlatform(const FName MaterialQualityPlatform, cons
 		AllMaterialsCacheResourceShadersForRendering();
 	}
 
+	PreviewFeatureLevelChanged.Broadcast(InPreviewFeatureLevel);
+
 	if (bSaveSettings)
 	{
 		SaveEditorFeatureLevel();
 	}
 }
 
+void UEditorEngine::ToggleFeatureLevelPreview()
+{
+	ERHIFeatureLevel::Type NewPreviewFeatureLevel = GWorld->FeatureLevel == GMaxRHIFeatureLevel ? PreviewFeatureLevel : GMaxRHIFeatureLevel;
+
+	PreviewFeatureLevelChanged.Broadcast(NewPreviewFeatureLevel);
+	
+	GEditor->OnSceneMaterialsModified();
+	GEditor->RedrawAllViewports();
+}
+
+bool UEditorEngine::IsFeatureLevelPreviewEnabled() const
+{
+	return PreviewFeatureLevel != GMaxRHIFeatureLevel;
+}
+
+bool UEditorEngine::IsFeatureLevelPreviewActive() const
+{
+	return GWorld->FeatureLevel != GMaxRHIFeatureLevel;
+}
+
 void UEditorEngine::LoadEditorFeatureLevel()
 {
  	auto* Settings = GetMutableDefault<UEditorPerProjectUserSettings>();
@@ -7524,8 +7525,9 @@ void UEditorEngine::LoadEditorFeatureLevel()
 	if (ShaderPlatform != SP_NumPlatforms)
 	{
 		const FName MaterialQualityPlatform = Settings->bIsMaterialQualityOverridePlatform ? Settings->PreviewShaderPlatformName : NAME_None;
-		const ERHIFeatureLevel::Type PreviewFeatureLevel = GetMaxSupportedFeatureLevel(ShaderPlatform);
-		SetPreviewPlatform(MaterialQualityPlatform, PreviewFeatureLevel, false);
+		const ERHIFeatureLevel::Type FeatureLevel = GetMaxSupportedFeatureLevel(ShaderPlatform);
+
+		SetPreviewPlatform(MaterialQualityPlatform, FeatureLevel, false);
 	}
 }
 
@@ -7538,7 +7540,7 @@ void UEditorEngine::SaveEditorFeatureLevel()
 	{
 		Settings->bIsMaterialQualityOverridePlatform = false;
 
-		const EShaderPlatform ShaderPlatform = GetFeatureLevelShaderPlatform(DefaultWorldFeatureLevel);
+		const EShaderPlatform ShaderPlatform = GetFeatureLevelShaderPlatform(PreviewFeatureLevel);
 		Settings->PreviewShaderPlatformName = LegacyShaderPlatformToShaderFormat(ShaderPlatform);
 	}
 	else
diff --git a/Engine/Source/Editor/UnrealEd/Private/PlayLevel.cpp b/Engine/Source/Editor/UnrealEd/Private/PlayLevel.cpp
--- a/Engine/Source/Editor/UnrealEd/Private/PlayLevel.cpp
+++ b/Engine/Source/Editor/UnrealEd/Private/PlayLevel.cpp
@@ -2944,8 +2944,8 @@ UGameInstance* UEditorEngine::CreatePIEGameInstance(int32 InPIEInstance, bool bI
 	GameInstanceParams.bSimulateInEditor = bInSimulateInEditor;
 	GameInstanceParams.bStartInSpectatorMode = bStartInSpectatorMode;
 	GameInstanceParams.bRunAsDedicated = bRunAsDedicated;
-	GameInstanceParams.WorldFeatureLevel = DefaultWorldFeatureLevel;
-	
+	GameInstanceParams.WorldFeatureLevel = PreviewFeatureLevel;
+
 	const FGameInstancePIEResult InitializeResult = GameInstance->InitializeForPlayInEditor(InPIEInstance, GameInstanceParams);
 	if (!InitializeResult.IsSuccess())
 	{
diff --git a/Engine/Source/Editor/UnrealEd/Private/SEditorViewport.cpp b/Engine/Source/Editor/UnrealEd/Private/SEditorViewport.cpp
--- a/Engine/Source/Editor/UnrealEd/Private/SEditorViewport.cpp
+++ b/Engine/Source/Editor/UnrealEd/Private/SEditorViewport.cpp
@@ -20,9 +20,12 @@
 #include "Widgets/Layout/SBox.h"
 #include "Widgets/Input/SSpinBox.h"
 #include "Widgets/Input/SCheckBox.h"
+#include "MaterialShaderQualitySettings.h"
 
 #define LOCTEXT_NAMESPACE "EditorViewport"
 
+#include "RHIShaderPlatformDefinitions.inl"
+
 SEditorViewport::SEditorViewport()
 	: LastTickTime(0)
 	, bInvalidated(false)
@@ -683,4 +686,87 @@ EActiveTimerReturnType SEditorViewport::EnsureTick( double InCurrentTime, float
 	return bShouldContinue ? EActiveTimerReturnType::Continue : EActiveTimerReturnType::Stop;
 }
 
+///////////////////////////////////////////////////////////////////////////////
+// begin feature level control functions block
+///////////////////////////////////////////////////////////////////////////////
+EShaderPlatform SEditorViewport::GetShaderPlatformHelper(const ERHIFeatureLevel::Type FeatureLevel) const
+{
+	UMaterialShaderQualitySettings* MaterialShaderQualitySettings = UMaterialShaderQualitySettings::Get();
+	const FName& PreviewPlatform = MaterialShaderQualitySettings->GetPreviewPlatform();
+
+	EShaderPlatform ShaderPlatform = ShaderFormatToLegacyShaderPlatform(PreviewPlatform);
+	if (ShaderPlatform == SP_NumPlatforms)
+	{
+		ShaderPlatform = GetFeatureLevelShaderPlatform(FeatureLevel);
+	}
+
+	return ShaderPlatform;
+}
+
+TSharedRef<SWidget> SEditorViewport::BuildFeatureLevelWidget() const
+{
+	TSharedRef<SWidget> BoxWidget = SNew(SHorizontalBox)
+		.Visibility(this, &SEditorViewport::GetCurrentFeatureLevelPreviewTextVisibility)
+		+ SHorizontalBox::Slot()
+		.AutoWidth()
+		.Padding(2.0f, 1.0f, 2.0f, 1.0f)
+		[
+			SNew(STextBlock)
+			.Text(this, &SEditorViewport::GetCurrentFeatureLevelPreviewText, true)
+		.Font(FEditorStyle::GetFontStyle(TEXT("MenuItem.Font")))
+		.ShadowOffset(FVector2D(1, 1))
+		]
+		+ SHorizontalBox::Slot()
+		.AutoWidth()
+		.Padding(4.0f, 1.0f, 2.0f, 1.0f)
+		[
+			SNew(STextBlock)
+			.Text(this, &SEditorViewport::GetCurrentFeatureLevelPreviewText, false)
+			.Font(FEditorStyle::GetFontStyle(TEXT("MenuItem.Font")))
+			.ColorAndOpacity(FLinearColor(0.4f, 1.0f, 1.0f))
+			.ShadowOffset(FVector2D(1, 1))
+		];
+
+	return BoxWidget;
+}
+
+EVisibility SEditorViewport::GetCurrentFeatureLevelPreviewTextVisibility() const
+{
+	if (GetWorld())
+	{
+		return (GetWorld()->FeatureLevel != GMaxRHIFeatureLevel) ? EVisibility::SelfHitTestInvisible : EVisibility::Collapsed;
+	}
+	else
+	{
+		return EVisibility::Collapsed;
+	}
+}
+
+FText SEditorViewport::GetCurrentFeatureLevelPreviewText(bool bDrawOnlyLabel) const
+{
+	FText LabelName;
+
+	if (bDrawOnlyLabel)
+	{
+		LabelName = LOCTEXT("FeatureLevelLabel", "Feature Level:");
+	}
+	else
+	{
+		UWorld* World = GetWorld();
+		if (World != nullptr)
+		{
+			ERHIFeatureLevel::Type TargetFeatureLevel = World->FeatureLevel;
+			EShaderPlatform ShaderPlatform = GetShaderPlatformHelper(TargetFeatureLevel);
+			const FText& PlatformText = GetFriendlyShaderPlatformName(ShaderPlatform);
+			LabelName = FText::Format(LOCTEXT("WorldFeatureLevel", "{0}"), PlatformText);
+		}
+	}
+
+	return LabelName;
+}
+
+///////////////////////////////////////////////////////////////////////////////
+// end feature level control functions block
+///////////////////////////////////////////////////////////////////////////////
+
 #undef LOCTEXT_NAMESPACE
diff --git a/Engine/Source/Editor/UnrealEd/Public/SEditorViewport.h b/Engine/Source/Editor/UnrealEd/Public/SEditorViewport.h
--- a/Engine/Source/Editor/UnrealEd/Public/SEditorViewport.h
+++ b/Engine/Source/Editor/UnrealEd/Public/SEditorViewport.h
@@ -84,7 +84,25 @@ public:
 	/** Build the exposure menu using EV100 settings */
 	TSharedRef<SWidget> BuildFixedEV100Menu()  const;
 
+///////////////////////////////////////////////////////////////////////////////
+// begin feature level control functions block
+///////////////////////////////////////////////////////////////////////////////
+private:
+	/** Called to get the feature level preview text */
+	FText GetCurrentFeatureLevelPreviewText(bool bDrawOnlyLabel) const;
+
+	/** @return The visibility of the current feature level preview text display */
+	EVisibility GetCurrentFeatureLevelPreviewTextVisibility() const;
+
+	/** Helper function that, for some FeatureLevel argument, will retrieve the required shader platform */
+	EShaderPlatform GetShaderPlatformHelper(const ERHIFeatureLevel::Type InFeatureLevel) const;
+
 protected:
+	/** call this function to build a 'text' widget that can display the present feature level */
+	TSharedRef<SWidget> BuildFeatureLevelWidget() const;
+///////////////////////////////////////////////////////////////////////////////
+// end feature level control functions block
+///////////////////////////////////////////////////////////////////////////////
 
 	/** Called by the fixed EV100 slider to get the fixed EV100 value */
 	float OnGetFixedEV100Value() const;
diff --git a/Engine/Source/Runtime/CoreUObject/Private/UObject/GarbageCollection.cpp b/Engine/Source/Runtime/CoreUObject/Private/UObject/GarbageCollection.cpp
--- a/Engine/Source/Runtime/CoreUObject/Private/UObject/GarbageCollection.cpp
+++ b/Engine/Source/Runtime/CoreUObject/Private/UObject/GarbageCollection.cpp
@@ -1019,7 +1019,7 @@ struct FConditionalGCLock
 void IncrementalPurgeGarbage( bool bUseTimeLimit, float TimeLimit )
 {
 	SCOPED_NAMED_EVENT(IncrementalPurgeGarbage, FColor::Red);
-	DECLARE_SCOPE_CYCLE_COUNTER( TEXT( "IncrementalPurgeGarbage" ), STAT_IncrementalPurgeGarbage, STATGROUP_GC );
+	DECLARE_SCOPE_CYCLE_COUNTER(TEXT("IncrementalPurgeGarbage"), STAT_IncrementalPurgeGarbage, STATGROUP_GC);
 
 	if (GExitPurge)
 	{
@@ -1488,6 +1488,7 @@ void CollectGarbageInternal(EObjectFlags KeepFlags, bool bPerformFullPurge)
 	SCOPE_TIME_GUARD(TEXT("Collect Garbage"));
 	SCOPED_NAMED_EVENT(CollectGarbageInternal, FColor::Red);
 	CSV_EVENT_GLOBAL(TEXT("GC"));
+	CSV_SCOPED_TIMING_STAT_EXCLUSIVE(GarbageCollection);
 
 	FGCCSyncObject::Get().ResetGCIsWaiting();
 
diff --git a/Engine/Source/Runtime/Engine/Classes/Materials/Material.h b/Engine/Source/Runtime/Engine/Classes/Materials/Material.h
--- a/Engine/Source/Runtime/Engine/Classes/Materials/Material.h
+++ b/Engine/Source/Runtime/Engine/Classes/Materials/Material.h
@@ -1683,7 +1683,7 @@ public:
 #endif // WITH_EDITOR
 
 	// For all materials, UMaterial::CacheResourceShadersForRendering
-	ENGINE_API static void AllMaterialsCacheResourceShadersForRendering();
+	ENGINE_API static void AllMaterialsCacheResourceShadersForRendering(bool bUpdateProgressDialog = false);
 
 #if WITH_EDITORONLY_DATA
 	/**
diff --git a/Engine/Source/Runtime/Engine/Classes/Materials/MaterialInstance.h b/Engine/Source/Runtime/Engine/Classes/Materials/MaterialInstance.h
--- a/Engine/Source/Runtime/Engine/Classes/Materials/MaterialInstance.h
+++ b/Engine/Source/Runtime/Engine/Classes/Materials/MaterialInstance.h
@@ -579,9 +579,9 @@ public:
 	ENGINE_API virtual bool HasOverridenBaseProperties()const;
 
 	// For all materials instances, UMaterialInstance::CacheResourceShadersForRendering
-	ENGINE_API static void AllMaterialsCacheResourceShadersForRendering();
+	ENGINE_API static void AllMaterialsCacheResourceShadersForRendering(bool bUpdateProgressDialog = false);
 
-	/** 
+	/**
 	 * Determine whether this Material Instance is a child of another Material
 	 *
 	 * @param	Material	Material to check against
diff --git a/Engine/Source/Runtime/Engine/Private/Materials/Material.cpp b/Engine/Source/Runtime/Engine/Private/Materials/Material.cpp
--- a/Engine/Source/Runtime/Engine/Private/Materials/Material.cpp
+++ b/Engine/Source/Runtime/Engine/Private/Materials/Material.cpp
@@ -72,6 +72,7 @@
 #include "Curves/CurveLinearColor.h"
 #include "Curves/CurveLinearColorAtlas.h"
 #include "HAL/ThreadHeartBeat.h"
+#include "Misc/ScopedSlowTask.h"
 
 #define LOCTEXT_NAMESPACE "Material"
 
@@ -5705,7 +5706,7 @@ UMaterial::FMaterialCompilationFinished& UMaterial::OnMaterialCompilationFinishe
 }
 #endif // WITH_EDITOR
 
-void UMaterial::AllMaterialsCacheResourceShadersForRendering()
+void UMaterial::AllMaterialsCacheResourceShadersForRendering(bool bUpdateProgressDialog)
 {
 #if STORE_ONLY_ACTIVE_SHADERMAPS
 	TArray<UMaterial*> Materials;
@@ -5720,13 +5721,33 @@ void UMaterial::AllMaterialsCacheResourceShadersForRendering()
 		FThreadHeartBeat::Get().HeartBeat();
 	}
 #else
-	for (TObjectIterator<UMaterial> It; It; ++It)
+#if WITH_EDITOR
+	FScopedSlowTask SlowTask(100.f, NSLOCTEXT("Engine", "CacheMaterialShadersMessage", "Caching material shaders"), true);
+	if (bUpdateProgressDialog)
+	{
+		SlowTask.Visibility = ESlowTaskVisibility::ForceVisible;
+		SlowTask.MakeDialog();
+	}
+#endif // WITH_EDITOR
+
+	TArray<UObject*> MaterialArray;
+	GetObjectsOfClass(UMaterial::StaticClass(), MaterialArray, true, RF_ClassDefaultObject, EInternalObjectFlags::None);
+	float TaskIncrement = (float)100.0f / MaterialArray.Num();
+
+	for (UObject* MaterialObj : MaterialArray)
 	{
-		UMaterial* Material = *It;
+		UMaterial* Material = (UMaterial*)MaterialObj;
 
 		Material->CacheResourceShadersForRendering(false);
+
+#if WITH_EDITOR
+		if (bUpdateProgressDialog)
+		{
+			SlowTask.EnterProgressFrame(TaskIncrement);
+		}
+#endif // WITH_EDITOR
 	}
-#endif
+#endif // STORE_ONLY_ACTIVE_SHADERMAPS
 }
 
 /**
diff --git a/Engine/Source/Runtime/Engine/Private/Materials/MaterialInstance.cpp b/Engine/Source/Runtime/Engine/Private/Materials/MaterialInstance.cpp
--- a/Engine/Source/Runtime/Engine/Private/Materials/MaterialInstance.cpp
+++ b/Engine/Source/Runtime/Engine/Private/Materials/MaterialInstance.cpp
@@ -35,6 +35,7 @@
 #include "Curves/CurveLinearColor.h"
 #include "Curves/CurveLinearColorAtlas.h"
 #include "HAL/ThreadHeartBeat.h"
+#include "Misc/ScopedSlowTask.h"
 
 DECLARE_CYCLE_STAT(TEXT("MaterialInstance CopyMatInstParams"), STAT_MaterialInstance_CopyMatInstParams, STATGROUP_Shaders);
 DECLARE_CYCLE_STAT(TEXT("MaterialInstance Serialize"), STAT_MaterialInstance_Serialize, STATGROUP_Shaders);
@@ -3828,7 +3829,7 @@ FPostProcessMaterialNode* IteratePostProcessMaterialNodes(const FFinalPostProces
 	}
 }
 
-void UMaterialInstance::AllMaterialsCacheResourceShadersForRendering()
+void UMaterialInstance::AllMaterialsCacheResourceShadersForRendering(bool bUpdateProgressDialog)
 {
 #if STORE_ONLY_ACTIVE_SHADERMAPS
 	TArray<UMaterialInstance*> MaterialInstances;
@@ -3843,13 +3844,33 @@ void UMaterialInstance::AllMaterialsCacheResourceShadersForRendering()
 		FThreadHeartBeat::Get().HeartBeat();
 	}
 #else
-	for (TObjectIterator<UMaterialInstance> It; It; ++It)
+#if WITH_EDITOR
+	FScopedSlowTask SlowTask(100.f, NSLOCTEXT("Engine", "CacheMaterialInstanceShadersMessage", "Caching material instance shaders"), true);
+	if (bUpdateProgressDialog)
+	{
+		SlowTask.Visibility = ESlowTaskVisibility::ForceVisible;
+		SlowTask.MakeDialog();
+	}
+#endif // WITH_EDITOR
+
+	TArray<UObject*> MaterialInstanceArray;
+	GetObjectsOfClass(UMaterialInstance::StaticClass(), MaterialInstanceArray, true, RF_ClassDefaultObject, EInternalObjectFlags::None);
+	float TaskIncrement = (float)100.0f / MaterialInstanceArray.Num();
+
+	for (UObject* MaterialInstanceObj : MaterialInstanceArray)
 	{
-		UMaterialInstance* MaterialInstance = *It;
+		UMaterialInstance* MaterialInstance = (UMaterialInstance*)MaterialInstanceObj;
 
 		MaterialInstance->CacheResourceShadersForRendering();
+
+#if WITH_EDITOR
+		if (bUpdateProgressDialog)
+		{
+			SlowTask.EnterProgressFrame(TaskIncrement);
+		}
+#endif // WITH_EDITOR
 	}
-#endif
+#endif // STORE_ONLY_ACTIVE_SHADERMAPS
 }
 
 
diff --git a/Engine/Source/Runtime/Online/Voice/Private/VoicePrivate.h b/Engine/Source/Runtime/Online/Voice/Private/VoicePrivate.h
--- a/Engine/Source/Runtime/Online/Voice/Private/VoicePrivate.h
+++ b/Engine/Source/Runtime/Online/Voice/Private/VoicePrivate.h
@@ -24,6 +24,6 @@ THIRD_PARTY_INCLUDES_END
 
 #define ANDROIDVOICE_SUPPORTED_PLATFORMS (PLATFORM_ANDROID && (PLATFORM_ANDROID_ARM || PLATFORM_ANDROID_ARM64 || PLATFORM_ANDROID_X64) && !PLATFORM_LUMIN)
 #define PLATFORM_SUPPORTS_VOICE_CAPTURE (PLATFORM_WINDOWS || PLATFORM_MAC || ANDROIDVOICE_SUPPORTED_PLATFORMS || (PLATFORM_LINUX && VOICE_MODULE_WITH_CAPTURE))
-#define PLATFORM_SUPPORTS_OPUS_CODEC (!PLATFORM_HTML5)
+#define PLATFORM_SUPPORTS_OPUS_CODEC (!PLATFORM_HTML5 && !PLATFORM_TVOS)
 
 // Module includes
diff --git a/Engine/Source/Runtime/RHI/Public/RHIShaderPlatformDefinitions.inl b/Engine/Source/Runtime/RHI/Public/RHIShaderPlatformDefinitions.inl
new file mode 100644
--- /dev/null
+++ b/Engine/Source/Runtime/RHI/Public/RHIShaderPlatformDefinitions.inl
@@ -0,0 +1,158 @@
+// Copyright 1998-2018 Epic Games, Inc. All Rights Reserved.
+
+/*=============================================================================
+	RHIShaderPlatformDefinitions.h: Localizable Friendly Names for Shader Platforms
+=============================================================================*/
+
+#pragma once
+
+static const FText& GetFriendlyShaderPlatformName(const EShaderPlatform InShaderPlatform)
+{
+	switch (InShaderPlatform)
+	{
+	case SP_PCD3D_SM5:
+	case SP_OPENGL_SM5:
+	case SP_METAL_SM5:
+	case SP_VULKAN_SM5:
+	case SP_OPENGL_ES31_EXT:
+	{
+		static FText Description = LOCTEXT("Generic_SM5_loc", "SM5");
+		return Description;
+	}
+	break;
+
+	case SP_METAL_SM5_NOTESS:
+	case SP_METAL_MRT_MAC:
+	case SP_METAL_MRT:
+	{
+		static FText Description = LOCTEXT("Metal_SM5_loc", "Metal SM5");
+		return Description;
+	}
+	break;
+
+	case SP_PCD3D_SM4:
+	case SP_OPENGL_SM4:
+	case SP_VULKAN_SM4:
+	{
+		static FText Description = LOCTEXT("Generic_SM4_loc", "SM4");
+		return Description;
+	}
+	break;
+
+	case SP_PS4:
+	{
+		static FText Description = LOCTEXT("Playstation_4_loc", "PlayStation 4");
+		return Description;
+	}
+	break;
+
+	case SP_XBOXONE_D3D12:
+	{
+		static FText Description = LOCTEXT("XBox_One_loc", "XBox One");
+		return Description;
+	}
+	break;
+
+	case SP_SWITCH:
+	{
+		static FText Description = LOCTEXT("Switch_loc", "Switch");
+		return Description;
+	}
+	break;
+
+	case SP_SWITCH_FORWARD:
+	{
+		static FText Description = LOCTEXT("Switch_Forward_loc", "Switch Forward");
+		return Description;
+	}
+	break;
+
+	case SP_PCD3D_ES3_1:
+	case SP_VULKAN_PCES3_1:
+	case SP_OPENGL_PCES3_1:
+	{
+		static FText Description = LOCTEXT("Generic_ES31_loc", "ES31");
+		return Description;
+	}
+	break;
+
+	case SP_VULKAN_ES3_1_ANDROID:
+	{
+		static FText Description = LOCTEXT("Android_Vulkan_ES31_loc", "Android Vulkan ES31");
+		return Description;
+	}
+	break;
+
+	case SP_OPENGL_ES3_1_ANDROID:
+	{
+		static FText Description = LOCTEXT("Android_ES31_loc", "Android ES31");
+		return Description;
+	}
+	break;
+
+	case SP_VULKAN_ES3_1_LUMIN:
+	{
+		static FText Description = LOCTEXT("Lumin_Vulkan_ES31_loc", "Lumin Vulkan ES31");
+		return Description;
+	}
+	break;
+
+	case SP_METAL:
+	case SP_METAL_MACES3_1:
+	{
+		static FText Description = LOCTEXT("iOS_Metal_Mace_31_loc", "Metal ES31");
+		return Description;
+	}
+	break;
+
+	case SP_VULKAN_SM5_LUMIN:
+	{
+		static FText Description = LOCTEXT("Lumin_Vulkan_SM5_loc", "Lumin Vulkan SM5");
+		return Description;
+	}
+	break;
+
+	case SP_PCD3D_ES2:
+	case SP_OPENGL_PCES2:
+	{
+		static FText Description = LOCTEXT("Generic_ES2_loc", "ES2");
+		return Description;
+	}
+	break;
+
+	case SP_OPENGL_ES2_ANDROID:
+	{
+		static FText Description = LOCTEXT("Android_ES2_loc", "Android ES2");
+		return Description;
+	}
+	break;
+
+	case SP_OPENGL_ES2_WEBGL:
+	{
+		static FText Description = LOCTEXT("WebGL_ES2_loc", "WebGL ES2");
+		return Description;
+	}
+	break;
+
+	case SP_OPENGL_ES2_IOS:
+	{
+		static FText Description = LOCTEXT("iOS_ES2_loc", "iOS ES2");
+		return Description;
+	}
+	break;
+
+	case SP_METAL_MACES2:
+	{
+		static FText Description = LOCTEXT("iOS_Metal_Mace_2_loc", "iOS Metal ES2");
+		return Description;
+	}
+	break;
+
+	default:
+		break;
+	};
+
+	static FText Description = LOCTEXT("", "");
+	return Description;
+}
+
diff --git a/Engine/Source/Runtime/Slate/Private/Widgets/Images/SImage.cpp b/Engine/Source/Runtime/Slate/Private/Widgets/Images/SImage.cpp
--- a/Engine/Source/Runtime/Slate/Private/Widgets/Images/SImage.cpp
+++ b/Engine/Source/Runtime/Slate/Private/Widgets/Images/SImage.cpp
@@ -7,7 +7,8 @@ void SImage::Construct( const FArguments& InArgs )
 {
 	Image = InArgs._Image;
 	ColorAndOpacity = InArgs._ColorAndOpacity;
-	OnMouseButtonDownHandler = InArgs._OnMouseButtonDown;
+
+	SetOnMouseButtonDown(InArgs._OnMouseButtonDown);
 }
 
 int32 SImage::OnPaint( const FPaintArgs& Args, const FGeometry& AllottedGeometry, const FSlateRect& MyCullingRect, FSlateWindowElementList& OutDrawElements, int32 LayerId, const FWidgetStyle& InWidgetStyle, bool bParentEnabled ) const
@@ -26,20 +27,6 @@ int32 SImage::OnPaint( const FPaintArgs& Args, const FGeometry& AllottedGeometry
 	return LayerId;
 }
 
-FReply SImage::OnMouseButtonDown( const FGeometry& MyGeometry, const FPointerEvent& MouseEvent )
-{
-	if ( OnMouseButtonDownHandler.IsBound() )
-	{
-		// If a handler is assigned, call it.
-		return OnMouseButtonDownHandler.Execute(MyGeometry, MouseEvent);
-	}
-	else
-	{
-		// otherwise the event is unhandled.
-		return FReply::Unhandled();
-	}	
-}
-
 FVector2D SImage::ComputeDesiredSize( float ) const
 {
 	const FSlateBrush* ImageBrush = Image.Get();
@@ -76,8 +63,3 @@ void SImage::SetImage(TAttribute<const FSlateBrush*> InImage)
 		Invalidate(EInvalidateWidget::LayoutAndVolatility);
 	}
 }
-
-void SImage::SetOnMouseButtonDown(FPointerEventHandler EventHandler)
-{
-	OnMouseButtonDownHandler = EventHandler;
-}
diff --git a/Engine/Source/Runtime/Slate/Private/Widgets/Images/SSpinningImage.cpp b/Engine/Source/Runtime/Slate/Private/Widgets/Images/SSpinningImage.cpp
--- a/Engine/Source/Runtime/Slate/Private/Widgets/Images/SSpinningImage.cpp
+++ b/Engine/Source/Runtime/Slate/Private/Widgets/Images/SSpinningImage.cpp
@@ -8,7 +8,11 @@ void SSpinningImage::Construct(const FArguments& InArgs)
 {
 	Image = InArgs._Image;
 	ColorAndOpacity = InArgs._ColorAndOpacity;
-	OnMouseButtonDownHandler = InArgs._OnMouseButtonDown;
+
+	if (InArgs._OnMouseButtonDown.IsBound())
+	{
+		SetOnMouseButtonDown(InArgs._OnMouseButtonDown);
+	}
 	
 	SpinAnimationSequence = FCurveSequence( 0.f, InArgs._Period );
 	SpinAnimationSequence.Play( this->AsShared(), true );
diff --git a/Engine/Source/Runtime/Slate/Private/Widgets/Layout/SBorder.cpp b/Engine/Source/Runtime/Slate/Private/Widgets/Layout/SBorder.cpp
--- a/Engine/Source/Runtime/Slate/Private/Widgets/Layout/SBorder.cpp
+++ b/Engine/Source/Runtime/Slate/Private/Widgets/Layout/SBorder.cpp
@@ -30,10 +30,26 @@ void SBorder::Construct( const SBorder::FArguments& InArgs )
 	BorderImage = InArgs._BorderImage;
 	BorderBackgroundColor = InArgs._BorderBackgroundColor;
 	ForegroundColor = InArgs._ForegroundColor;
-	SetOnMouseButtonDown(InArgs._OnMouseButtonDown);
-	SetOnMouseButtonUp(InArgs._OnMouseButtonUp);
-	SetOnMouseMove(InArgs._OnMouseMove);
-	SetOnMouseDoubleClick(InArgs._OnMouseDoubleClick);
+
+	if (InArgs._OnMouseButtonDown.IsBound())
+	{
+		SetOnMouseButtonDown(InArgs._OnMouseButtonDown);
+	}
+
+	if (InArgs._OnMouseButtonUp.IsBound())
+	{
+		SetOnMouseButtonUp(InArgs._OnMouseButtonUp);
+	}
+
+	if (InArgs._OnMouseMove.IsBound())
+	{
+		SetOnMouseMove(InArgs._OnMouseMove);
+	}
+
+	if (InArgs._OnMouseDoubleClick.IsBound())
+	{
+		SetOnMouseDoubleClick(InArgs._OnMouseDoubleClick);
+	}
 
 	ChildSlot
 	.HAlign(InArgs._HAlign)
diff --git a/Engine/Source/Runtime/Slate/Public/Widgets/Images/SImage.h b/Engine/Source/Runtime/Slate/Public/Widgets/Images/SImage.h
--- a/Engine/Source/Runtime/Slate/Public/Widgets/Images/SImage.h
+++ b/Engine/Source/Runtime/Slate/Public/Widgets/Images/SImage.h
@@ -20,20 +20,19 @@ class SLATE_API SImage
 	: public SLeafWidget
 {
 public:
-	SLATE_BEGIN_ARGS( SImage )
-		: _Image( FCoreStyle::Get().GetDefaultBrush() )
-		, _ColorAndOpacity( FLinearColor::White )
-		, _OnMouseButtonDown()
-		{}
+	SLATE_BEGIN_ARGS(SImage)
+		: _Image(FCoreStyle::Get().GetDefaultBrush())
+		, _ColorAndOpacity(FLinearColor::White)
+	{}
 
-		/** Image resource */
-		SLATE_ATTRIBUTE( const FSlateBrush*, Image )
+	/** Image resource */
+	SLATE_ATTRIBUTE(const FSlateBrush*, Image)
 
 		/** Color and opacity */
-		SLATE_ATTRIBUTE( FSlateColor, ColorAndOpacity )
+		SLATE_ATTRIBUTE(FSlateColor, ColorAndOpacity)
 
 		/** Invoked when the mouse is pressed in the widget. */
-		SLATE_EVENT( FPointerEventHandler, OnMouseButtonDown )
+		SLATE_EVENT(FPointerEventHandler, OnMouseButtonDown)
 
 	SLATE_END_ARGS()
 
@@ -49,27 +48,23 @@ public:
 	 *
 	 * @param	InArgs	The declaration data for this widget
 	 */
-	void Construct( const FArguments& InArgs );
+	void Construct(const FArguments& InArgs);
 
 public:
 
 	/** See the ColorAndOpacity attribute */
-	void SetColorAndOpacity( const TAttribute<FSlateColor>& InColorAndOpacity );
-	
+	void SetColorAndOpacity(const TAttribute<FSlateColor>& InColorAndOpacity);
+
 	/** See the ColorAndOpacity attribute */
-	void SetColorAndOpacity( FLinearColor InColorAndOpacity );
+	void SetColorAndOpacity(FLinearColor InColorAndOpacity);
 
 	/** See the Image attribute */
 	void SetImage(TAttribute<const FSlateBrush*> InImage);
-	
-	/** See OnMouseButtonDown event */
-	void SetOnMouseButtonDown(FPointerEventHandler EventHandler);
 
 public:
 
 	// SWidget overrides
-	virtual int32 OnPaint( const FPaintArgs& Args, const FGeometry& AllottedGeometry, const FSlateRect& MyCullingRect, FSlateWindowElementList& OutDrawElements, int32 LayerId, const FWidgetStyle& InWidgetStyle, bool bParentEnabled ) const override;
-	virtual FReply OnMouseButtonDown( const FGeometry& MyGeometry, const FPointerEvent& MouseEvent ) override;
+	virtual int32 OnPaint(const FPaintArgs& Args, const FGeometry& AllottedGeometry, const FSlateRect& MyCullingRect, FSlateWindowElementList& OutDrawElements, int32 LayerId, const FWidgetStyle& InWidgetStyle, bool bParentEnabled) const override;
 
 protected:
 	// Begin SWidget overrides.
@@ -83,7 +78,4 @@ protected:
 
 	/** Color and opacity scale for this image */
 	TAttribute<FSlateColor> ColorAndOpacity;
-
-	/** Invoked when the mouse is pressed in the image */
-	FPointerEventHandler OnMouseButtonDownHandler;
 };
diff --git a/Engine/Source/Runtime/SlateCore/Private/Widgets/SWidget.cpp b/Engine/Source/Runtime/SlateCore/Private/Widgets/SWidget.cpp
--- a/Engine/Source/Runtime/SlateCore/Private/Widgets/SWidget.cpp
+++ b/Engine/Source/Runtime/SlateCore/Private/Widgets/SWidget.cpp
@@ -266,7 +266,7 @@ FReply SWidget::OnPreviewMouseButtonDown( const FGeometry& MyGeometry, const FPo
 
 FReply SWidget::OnMouseButtonDown(const FGeometry& MyGeometry, const FPointerEvent& MouseEvent)
 {
-	if ( FPointerEventHandler* Event = PointerEvents.Find(NAME_MouseButtonDown) )
+	if (const FPointerEventHandler* Event = GetPointerEvent(NAME_MouseButtonDown))
 	{
 		if ( Event->IsBound() )
 		{
@@ -278,7 +278,7 @@ FReply SWidget::OnMouseButtonDown(const FGeometry& MyGeometry, const FPointerEve
 
 FReply SWidget::OnMouseButtonUp(const FGeometry& MyGeometry, const FPointerEvent& MouseEvent)
 {
-	if ( FPointerEventHandler* Event = PointerEvents.Find(NAME_MouseButtonUp) )
+	if (const FPointerEventHandler* Event = GetPointerEvent(NAME_MouseButtonUp) )
 	{
 		if ( Event->IsBound() )
 		{
@@ -290,7 +290,7 @@ FReply SWidget::OnMouseButtonUp(const FGeometry& MyGeometry, const FPointerEvent
 
 FReply SWidget::OnMouseMove(const FGeometry& MyGeometry, const FPointerEvent& MouseEvent)
 {
-	if ( FPointerEventHandler* Event = PointerEvents.Find(NAME_MouseMove) )
+	if (const FPointerEventHandler* Event = GetPointerEvent(NAME_MouseMove) )
 	{
 		if ( Event->IsBound() )
 		{
@@ -302,7 +302,7 @@ FReply SWidget::OnMouseMove(const FGeometry& MyGeometry, const FPointerEvent& Mo
 
 FReply SWidget::OnMouseButtonDoubleClick(const FGeometry& MyGeometry, const FPointerEvent& MouseEvent)
 {
-	if ( FPointerEventHandler* Event = PointerEvents.Find(NAME_MouseDoubleClick) )
+	if ( const FPointerEventHandler* Event = GetPointerEvent(NAME_MouseDoubleClick) )
 	{
 		if ( Event->IsBound() )
 		{
@@ -1167,24 +1167,49 @@ void SWidget::ExecuteActiveTimers(double CurrentTime, float DeltaTime)
 	}
 }
 
+const FPointerEventHandler* SWidget::GetPointerEvent(const FName EventName) const
+{
+	auto* FoundPair = PointerEvents.FindByPredicate([&EventName](const auto& TestPair) {return TestPair.Key == EventName; });
+	if (FoundPair)
+	{
+		return &FoundPair->Value;
+	}
+	return nullptr;
+}
+
+void SWidget::SetPointerEvent(const FName EventName, FPointerEventHandler& InEvent)
+{
+	// Find the event name and if found, replace the delegate
+	auto* FoundPair = PointerEvents.FindByPredicate([&EventName](const auto& TestPair) {return TestPair.Key == EventName; });
+	if (FoundPair)
+	{
+		FoundPair->Value = InEvent;
+	}
+	else
+	{
+		PointerEvents.Emplace(EventName, InEvent);
+	}
+
+}
+
 void SWidget::SetOnMouseButtonDown(FPointerEventHandler EventHandler)
 {
-	PointerEvents.Add(NAME_MouseButtonDown, EventHandler);
+	SetPointerEvent(NAME_MouseButtonDown, EventHandler);
 }
 
 void SWidget::SetOnMouseButtonUp(FPointerEventHandler EventHandler)
 {
-	PointerEvents.Add(NAME_MouseButtonUp, EventHandler);
+	SetPointerEvent(NAME_MouseButtonUp, EventHandler);
 }
 
 void SWidget::SetOnMouseMove(FPointerEventHandler EventHandler)
 {
-	PointerEvents.Add(NAME_MouseMove, EventHandler);
+	SetPointerEvent(NAME_MouseMove, EventHandler);
 }
 
 void SWidget::SetOnMouseDoubleClick(FPointerEventHandler EventHandler)
 {
-	PointerEvents.Add(NAME_MouseDoubleClick, EventHandler);
+	SetPointerEvent(NAME_MouseDoubleClick, EventHandler);
 }
 
 void SWidget::SetOnMouseEnter(FNoReplyPointerEventHandler EventHandler)
diff --git a/Engine/Source/Runtime/SlateCore/Public/Widgets/SWidget.h b/Engine/Source/Runtime/SlateCore/Public/Widgets/SWidget.h
--- a/Engine/Source/Runtime/SlateCore/Public/Widgets/SWidget.h
+++ b/Engine/Source/Runtime/SlateCore/Public/Widgets/SWidget.h
@@ -1418,6 +1418,8 @@ private:
 	/** Iterates over the active timer handles on the widget and executes them if their interval has elapsed. */
 	void ExecuteActiveTimers(double CurrentTime, float DeltaTime);
 
+	const FPointerEventHandler* GetPointerEvent(const FName EventName) const;
+	void SetPointerEvent(const FName EventName, FPointerEventHandler& InEvent);
 protected:
 	/** Dtor ensures that active timer handles are UnRegistered with the SlateApplication. */
 	virtual ~SWidget();
@@ -1553,7 +1555,7 @@ private:
 
 private:
 	// Events
-	TMap<FName, FPointerEventHandler> PointerEvents;
+	TArray<TPair<FName, FPointerEventHandler>> PointerEvents;
 
 	FNoReplyPointerEventHandler MouseEnterHandler;
 	FSimpleNoReplyPointerEventHandler MouseLeaveHandler;
diff --git a/Engine/Source/Runtime/UMG/Private/Components/Widget.cpp b/Engine/Source/Runtime/UMG/Private/Components/Widget.cpp
--- a/Engine/Source/Runtime/UMG/Private/Components/Widget.cpp
+++ b/Engine/Source/Runtime/UMG/Private/Components/Widget.cpp
@@ -980,6 +980,27 @@ EVisibility UWidget::GetVisibilityInDesigner() const
 	return bHiddenInDesigner ? EVisibility::Collapsed : EVisibility::Visible;
 }
 
+bool UWidget::IsVisibleInDesigner() const
+{
+	if (bHiddenInDesigner)
+	{
+		return false;
+	}
+
+	UWidget* Parent = GetParent();
+	while (Parent != nullptr)
+	{
+		if (Parent->bHiddenInDesigner)
+		{
+			return false;
+		}
+
+		Parent = Parent->GetParent();
+	}
+
+	return true;
+}
+
 void UWidget::PostEditChangeProperty(FPropertyChangedEvent& PropertyChangedEvent)
 {
 	Super::PostEditChangeProperty(PropertyChangedEvent);
diff --git a/Engine/Source/Runtime/UMG/Public/Components/Widget.h b/Engine/Source/Runtime/UMG/Public/Components/Widget.h
--- a/Engine/Source/Runtime/UMG/Public/Components/Widget.h
+++ b/Engine/Source/Runtime/UMG/Public/Components/Widget.h
@@ -804,8 +804,8 @@ public:
 	virtual void PostEditChangeProperty(struct FPropertyChangedEvent& PropertyChangedEvent) override;
 	// End of UObject interface
 
-	/** Gets the visibility of the widget inside the designer. */
-	EVisibility GetVisibilityInDesigner() const;
+	/** Is the widget visible in the designer?  If this widget is 'hidden in the designer' or a parent is, this widget will also return false here. */
+	bool IsVisibleInDesigner() const;
 
 	// Begin Designer contextual events
 	void SelectByDesigner();
@@ -836,6 +836,12 @@ public:
 	static FString GetDefaultFontName();
 
 protected:
+#if WITH_EDITOR
+	// This is an implementation detail that allows us to show and hide the widget in the designer
+	// regardless of the actual visibility state set by the user.
+	EVisibility GetVisibilityInDesigner() const;
+#endif
+
 	virtual void OnBindingChanged(const FName& Property);
 
 protected:
diff --git a/Engine/Source/Runtime/WebBrowser/Private/Android/AndroidCookieManager.cpp b/Engine/Source/Runtime/WebBrowser/Private/Android/AndroidCookieManager.cpp
--- a/Engine/Source/Runtime/WebBrowser/Private/Android/AndroidCookieManager.cpp
+++ b/Engine/Source/Runtime/WebBrowser/Private/Android/AndroidCookieManager.cpp
@@ -2,15 +2,14 @@
 
 #include "AndroidCookieManager.h"
 
-#if PLATFORM_ANDROID
+#if USE_ANDROID_JNI
 
 #include "Android/AndroidJNI.h"
 #include "Android/AndroidApplication.h"
 
-#if USE_ANDROID_JNI
 #include <jni.h>
 #include "Async/TaskGraphInterfaces.h"
-#endif
+
 
 FAndroidCookieManager::FAndroidCookieManager()
 {
@@ -74,4 +73,36 @@ void FAndroidCookieManager::DeleteCookies(const FString& URL, const FString& Coo
 	}
 }
 
+#else
+
+FAndroidCookieManager::FAndroidCookieManager()
+{
+}
+
+FAndroidCookieManager::~FAndroidCookieManager()
+{
+}
+
+void FAndroidCookieManager::SetCookie(const FString& URL, const FCookie& Cookie, TFunction<void(bool)> Completed)
+{
+	bool bResult = false;
+
+	// send result
+	if (Completed)
+	{
+		Completed(bResult);
+	}
+}
+
+void FAndroidCookieManager::DeleteCookies(const FString& URL, const FString& CookieName, TFunction<void(int)> Completed)
+{
+	bool bResult = false;
+
+	// send result
+	if (Completed)
+	{
+		Completed(bResult);
+	}
+}
+
 #endif
