Fix the android mobile preview icon

diff --git a/Engine/Source/Editor/LevelEditor/Private/LevelEditorToolBar.cpp b/Engine/Source/Editor/LevelEditor/Private/LevelEditorToolBar.cpp
--- a/Engine/Source/Editor/LevelEditor/Private/LevelEditorToolBar.cpp
+++ b/Engine/Source/Editor/LevelEditor/Private/LevelEditorToolBar.cpp
@@ -1281,6 +1281,23 @@ TSharedRef< SWidget > FLevelEditorToolBar::MakeLevelEditorToolBar( const TShared
 
 			static FText GetPreviewModeText()
 			{
+				UMaterialShaderQualitySettings* MaterialShaderQualitySettings = UMaterialShaderQualitySettings::Get();
+				const FName& PreviewPlatform = MaterialShaderQualitySettings->GetPreviewPlatform();
+
+				EShaderPlatform ShaderPlatform = ShaderFormatToLegacyShaderPlatform(PreviewPlatform);
+				if (ShaderPlatform == SP_NumPlatforms)
+				{
+					ShaderPlatform = GetFeatureLevelShaderPlatform(GEditor->PreviewFeatureLevel);
+				}
+
+				switch (ShaderPlatform)
+				{
+					case SP_VULKAN_ES3_1_ANDROID:
+					{						
+						return LOCTEXT("PreviewModeES31_Vulkan_Text", "Vulkan Preview");
+					}
+				}
+
 				switch (GEditor->PreviewFeatureLevel)
 				{
 					case ERHIFeatureLevel::SM4:
