Update the Preview Mode hotkey description

diff --git a/Engine/Source/Editor/LevelEditor/Private/LevelEditorActions.cpp b/Engine/Source/Editor/LevelEditor/Private/LevelEditorActions.cpp
--- a/Engine/Source/Editor/LevelEditor/Private/LevelEditorActions.cpp
+++ b/Engine/Source/Editor/LevelEditor/Private/LevelEditorActions.cpp
@@ -3410,7 +3410,7 @@ void FLevelEditorCommands::RegisterCommands()
 	UI_COMMAND(MaterialQualityLevel_Medium, "Medium", "Sets material quality in the scene to medium.", EUserInterfaceActionType::RadioButton, FInputChord());
 	UI_COMMAND(MaterialQualityLevel_High, "High", "Sets material quality in the scene to high.", EUserInterfaceActionType::RadioButton, FInputChord());
 
-	UI_COMMAND(ToggleFeatureLevelPreview, "Preview Platform", "Preview Platform", EUserInterfaceActionType::ToggleButton, FInputChord());
+	UI_COMMAND(ToggleFeatureLevelPreview, "Preview Mode Toggle", "Toggles the Preview Mode on or off for the currently selected Preview target", EUserInterfaceActionType::ToggleButton, FInputChord());
 
 	UI_COMMAND(PreviewPlatformOverride_AndroidGLES2, "Android ES2", "Mobile preview using Android's quality settings.", EUserInterfaceActionType::RadioButton, FInputChord());
 	UI_COMMAND(PreviewPlatformOverride_DefaultES2, "HTML5", "HTML5 preview.", EUserInterfaceActionType::RadioButton, FInputChord());
