UE4.22 - Fix for progress prompt not showing up
#rnx

diff --git a/Engine/Source/Editor/LevelEditor/Private/LevelEditorActions.cpp b/Engine/Source/Editor/LevelEditor/Private/LevelEditorActions.cpp
--- a/Engine/Source/Editor/LevelEditor/Private/LevelEditorActions.cpp
+++ b/Engine/Source/Editor/LevelEditor/Private/LevelEditorActions.cpp
@@ -617,6 +617,12 @@ bool FLevelEditorActionCallbacks::IsPreviewPlatformChecked(FName InMaterialQuali
 
 void FLevelEditorActionCallbacks::SetFeatureLevelPreview(ERHIFeatureLevel::Type InPreviewFeatureLevel)
 {
+	// When called through SMenuEntryBlock::OnClicked(), the popup menus are not dismissed when
+	// clicking on a checkbox, but they are dismissed when clicking on a button. We need the popup
+	// menus to go away, or SetFeaturePlatform() is unable to display a progress dialog. Force
+	// the dismissal here.
+	FSlateApplication::Get().DismissAllMenus();
+	
 	GEditor->SetPreviewPlatform(NAME_None, InPreviewFeatureLevel);
 }
 
