Reorder the Preview Rendering Level settings by platform

diff --git a/Engine/Source/Editor/LevelEditor/Private/LevelEditorToolBar.cpp b/Engine/Source/Editor/LevelEditor/Private/LevelEditorToolBar.cpp
--- a/Engine/Source/Editor/LevelEditor/Private/LevelEditorToolBar.cpp
+++ b/Engine/Source/Editor/LevelEditor/Private/LevelEditorToolBar.cpp
@@ -1828,48 +1828,45 @@ static void MakeShaderModelPreviewMenu( FMenuBuilder& MenuBuilder )
 
 	MenuBuilder.BeginSection("EditorPreviewMode", LOCTEXT("EditorPreviewModeDevices", "Preview Devices"));
 
-	for (int32 i = GMaxRHIFeatureLevel; i >= 0; --i)
-	{
-		switch (i)
-		{
-			case ERHIFeatureLevel::ES2:
-				MenuBuilder.AddMenuEntry(FLevelEditorCommands::Get().PreviewPlatformOverride_AndroidGLES2);
-				MenuBuilder.AddMenuEntry(FLevelEditorCommands::Get().PreviewPlatformOverride_DefaultES2);
-				break;
+	// SM5
+	MenuBuilder.AddMenuEntry(FLevelEditorCommands::Get().FeatureLevelPreview[ERHIFeatureLevel::SM5]);
 
-			case ERHIFeatureLevel::ES3_1:
-			{
-				//MenuBuilder.AddMenuEntry(FLevelEditorCommands::Get().PreviewPlatformOverride_DefaultES31);
-
-				bool bAndroidBuildForES31 = false;
-				GConfig->GetBool(TEXT("/Script/AndroidRuntimeSettings.AndroidRuntimeSettings"), TEXT("bBuildForES31"), bAndroidBuildForES31, GEngineIni);
-				if(bAndroidBuildForES31)
-				{
-					MenuBuilder.AddMenuEntry(FLevelEditorCommands::Get().PreviewPlatformOverride_AndroidGLES31);
-				}
+	// SM4
+	MenuBuilder.AddMenuEntry(FLevelEditorCommands::Get().FeatureLevelPreview[ERHIFeatureLevel::SM4]);
 
-				bool bAndroidSupportsVulkan = false;
-				GConfig->GetBool(TEXT("/Script/AndroidRuntimeSettings.AndroidRuntimeSettings"), TEXT("bSupportsVulkan"), bAndroidSupportsVulkan, GEngineIni);
-				if(bAndroidSupportsVulkan)
-				{
-					MenuBuilder.AddMenuEntry(FLevelEditorCommands::Get().PreviewPlatformOverride_AndroidVulkanES31);
-				}
+	// Android
+	bool bAndroidBuildForES31 = false;
+	GConfig->GetBool(TEXT("/Script/AndroidRuntimeSettings.AndroidRuntimeSettings"), TEXT("bBuildForES31"), bAndroidBuildForES31, GEngineIni);
+	if (bAndroidBuildForES31)
+	{
+		MenuBuilder.AddMenuEntry(FLevelEditorCommands::Get().PreviewPlatformOverride_AndroidGLES31);
+	}
 
-				bool bIOSSupportsMetal = false;
-				GConfig->GetBool(TEXT("/Script/IOSRuntimeSettings.IOSRuntimeSettings"), TEXT("bSupportsMetal"), bIOSSupportsMetal, GEngineIni);
-				if(bIOSSupportsMetal)
-				{
-					MenuBuilder.AddMenuEntry(FLevelEditorCommands::Get().PreviewPlatformOverride_IOSMetalES31);
-				}
+	bool bAndroidSupportsVulkan = false;
+	GConfig->GetBool(TEXT("/Script/AndroidRuntimeSettings.AndroidRuntimeSettings"), TEXT("bSupportsVulkan"), bAndroidSupportsVulkan, GEngineIni);
+	if (bAndroidSupportsVulkan)
+	{
+		MenuBuilder.AddMenuEntry(FLevelEditorCommands::Get().PreviewPlatformOverride_AndroidVulkanES31);
+	}
 
-				break;
-			}
+	bool bAndroidBuildForES2 = false;
+	GConfig->GetBool(TEXT("/Script/AndroidRuntimeSettings.AndroidRuntimeSettings"), TEXT("bBuildForES2"), bAndroidBuildForES2, GEngineIni);
+	if (bAndroidBuildForES2)
+	{
+    	MenuBuilder.AddMenuEntry(FLevelEditorCommands::Get().PreviewPlatformOverride_AndroidGLES2);
+	}
 
-			default:
-				MenuBuilder.AddMenuEntry(FLevelEditorCommands::Get().FeatureLevelPreview[i]);
-		}
+	// iOS
+	bool bIOSSupportsMetal = false;
+	GConfig->GetBool(TEXT("/Script/IOSRuntimeSettings.IOSRuntimeSettings"), TEXT("bSupportsMetal"), bIOSSupportsMetal, GEngineIni);
+	if (bIOSSupportsMetal)
+	{
+		MenuBuilder.AddMenuEntry(FLevelEditorCommands::Get().PreviewPlatformOverride_IOSMetalES31);
 	}
 
+	// HTML5
+	MenuBuilder.AddMenuEntry(FLevelEditorCommands::Get().PreviewPlatformOverride_DefaultES2);
+
     MenuBuilder.EndSection();
 
 #undef LOCTEXT_NAMESPACE
