Allow toggling of the Preview Mode during a PIE session

diff --git a/Engine/Source/Editor/LevelEditor/Private/LevelEditorActions.cpp b/Engine/Source/Editor/LevelEditor/Private/LevelEditorActions.cpp
--- a/Engine/Source/Editor/LevelEditor/Private/LevelEditorActions.cpp
+++ b/Engine/Source/Editor/LevelEditor/Private/LevelEditorActions.cpp
@@ -579,7 +579,7 @@ void FLevelEditorActionCallbacks::ToggleFeatureLevelPreview()
 
 bool FLevelEditorActionCallbacks::IsFeatureLevelPreviewEnabled()
 {
-	if (GEditor->PlayWorld || GUnrealEd->bIsSimulatingInEditor || GUnrealEd->IsLightingBuildCurrentlyRunning())
+	if (GUnrealEd->IsLightingBuildCurrentlyRunning())
 	{
 		return false;
 	}
