Change from radio buttons to checks for Settings->Preview Rendering Level for better editor UI compliance

diff --git a/Engine/Source/Editor/LevelEditor/Private/LevelEditorActions.cpp b/Engine/Source/Editor/LevelEditor/Private/LevelEditorActions.cpp
--- a/Engine/Source/Editor/LevelEditor/Private/LevelEditorActions.cpp
+++ b/Engine/Source/Editor/LevelEditor/Private/LevelEditorActions.cpp
@@ -3418,13 +3418,13 @@ void FLevelEditorCommands::RegisterCommands()
 
 	UI_COMMAND(ToggleFeatureLevelPreview, "Preview Mode Toggle", "Toggles the Preview Mode on or off for the currently selected Preview target", EUserInterfaceActionType::ToggleButton, FInputChord());
 
-	UI_COMMAND(PreviewPlatformOverride_AndroidGLES2, "Android ES2", "Mobile preview using Android's quality settings.", EUserInterfaceActionType::RadioButton, FInputChord());
-	UI_COMMAND(PreviewPlatformOverride_DefaultES2, "HTML5", "HTML5 preview.", EUserInterfaceActionType::RadioButton, FInputChord());
+	UI_COMMAND(PreviewPlatformOverride_AndroidGLES2, "Android ES2", "Mobile preview using Android's quality settings.", EUserInterfaceActionType::Check, FInputChord());
+	UI_COMMAND(PreviewPlatformOverride_DefaultES2, "HTML5", "HTML5 preview.", EUserInterfaceActionType::Check, FInputChord());
 
-	UI_COMMAND(PreviewPlatformOverride_DefaultES31, "Default High-End Mobile", "Use default mobile settings (no quality overrides).", EUserInterfaceActionType::RadioButton, FInputChord());
-	UI_COMMAND(PreviewPlatformOverride_AndroidGLES31, "Android ES 3.1", "Mobile preview using Android ES3.1 quality settings.", EUserInterfaceActionType::RadioButton, FInputChord());
-	UI_COMMAND(PreviewPlatformOverride_AndroidVulkanES31, "Android Vulkan", "Mobile preview using Android Vulkan quality settings.", EUserInterfaceActionType::RadioButton, FInputChord());
-	UI_COMMAND(PreviewPlatformOverride_IOSMetalES31, "iOS", "Mobile preview using iOS material quality settings.", EUserInterfaceActionType::RadioButton, FInputChord());
+	UI_COMMAND(PreviewPlatformOverride_DefaultES31, "Default High-End Mobile", "Use default mobile settings (no quality overrides).", EUserInterfaceActionType::Check, FInputChord());
+	UI_COMMAND(PreviewPlatformOverride_AndroidGLES31, "Android ES 3.1", "Mobile preview using Android ES3.1 quality settings.", EUserInterfaceActionType::Check, FInputChord());
+	UI_COMMAND(PreviewPlatformOverride_AndroidVulkanES31, "Android Vulkan", "Mobile preview using Android Vulkan quality settings.", EUserInterfaceActionType::Check, FInputChord());
+	UI_COMMAND(PreviewPlatformOverride_IOSMetalES31, "iOS", "Mobile preview using iOS material quality settings.", EUserInterfaceActionType::Check, FInputChord());
 
 
 	UI_COMMAND( ConnectToSourceControl, "Connect to Source Control...", "Opens a dialog to connect to source control.", EUserInterfaceActionType::Button, FInputChord());
@@ -3463,7 +3463,7 @@ void FLevelEditorCommands::RegisterCommands()
 			Name,
 			FeatureLevelLabels[i],
 			FeatureLevelToolTips[i])
-			.UserInterfaceType(EUserInterfaceActionType::RadioButton)
+			.UserInterfaceType(EUserInterfaceActionType::Check)
 			.DefaultChord(FInputChord());
 	}
 
