Guard against a possible shutdown crash in SLevelEditor

diff --git a/Engine/Source/Editor/LevelEditor/Private/SLevelEditor.cpp b/Engine/Source/Editor/LevelEditor/Private/SLevelEditor.cpp
--- a/Engine/Source/Editor/LevelEditor/Private/SLevelEditor.cpp
+++ b/Engine/Source/Editor/LevelEditor/Private/SLevelEditor.cpp
@@ -289,8 +289,11 @@ SLevelEditor::~SLevelEditor()
 
 	FEditorDelegates::MapChange.RemoveAll(this);
 
-	CastChecked<UEditorEngine>(GEngine)->OnPreviewFeatureLevelChanged().Remove(PreviewFeatureLevelChangedHandle);
-	
+	if (GEngine)
+	{
+		CastChecked<UEditorEngine>(GEngine)->OnPreviewFeatureLevelChanged().Remove(PreviewFeatureLevelChangedHandle);
+	}
+
 	if (GEditor)
 	{
 		GEditor->GetEditorWorldContext(true).RemoveRef(World);
