Fixed debug hud not drawing during simulate in editor
#rnx

diff --git a/Engine/Source/Editor/LevelEditor/Private/SLevelViewport.cpp b/Engine/Source/Editor/LevelEditor/Private/SLevelViewport.cpp
--- a/Engine/Source/Editor/LevelEditor/Private/SLevelViewport.cpp
+++ b/Engine/Source/Editor/LevelEditor/Private/SLevelViewport.cpp
@@ -3522,7 +3522,7 @@ void SLevelViewport::StartPlayInEditorSession(UGameViewportClient* PlayClient, c
 	PlayClient->SetGameLayerManager(GameLayerManager);
 
 	// Set the scene viewport on PIE
-	if (GameLayerManager.IsValid())
+	if (GameLayerManager.IsValid() && !bInSimulateInEditor)
 	{
 		GameLayerManager->SetSceneViewport(ActiveViewport.Get());
 	}
