Fix debug hud not rendering when swapping between PIE/SIE modes

diff --git a/Engine/Source/Editor/LevelEditor/Private/SLevelViewport.cpp b/Engine/Source/Editor/LevelEditor/Private/SLevelViewport.cpp
--- a/Engine/Source/Editor/LevelEditor/Private/SLevelViewport.cpp
+++ b/Engine/Source/Editor/LevelEditor/Private/SLevelViewport.cpp
@@ -3797,7 +3797,12 @@ void SLevelViewport::SwapViewportsForSimulateInEditor()
 	// Resize the viewport to be the same size the previously active viewport
 	// When starting in immersive mode its possible that the viewport has not been resized yet
 	ActiveViewport->OnPlayWorldViewportSwapped( *InactiveViewport );
-	
+
+	if (GameLayerManager.IsValid())
+	{
+		GameLayerManager->SetSceneViewport(ActiveViewport.Get());
+	}
+
 	ViewportWidget->SetViewportInterface( ActiveViewport.ToSharedRef() );
 
 	// Kick off a quick transition effect (border graphics)
@@ -3833,6 +3838,11 @@ void SLevelViewport::SwapViewportsForPlayInEditor()
 	// When starting in immersive mode its possible that the viewport has not been resized yet
 	ActiveViewport->OnPlayWorldViewportSwapped( *InactiveViewport );
 
+	if (GameLayerManager.IsValid())
+	{
+		GameLayerManager->SetSceneViewport(ActiveViewport.Get());
+	}
+
 	InactiveViewportWidgetEditorContent = ViewportWidget->GetContent();
 	ViewportWidget->SetViewportInterface( ActiveViewport.ToSharedRef() );
 
