Fix for vr mode crash when undoing camera duplication

diff --git a/Engine/Source/Editor/LevelEditor/Private/SLevelViewport.cpp b/Engine/Source/Editor/LevelEditor/Private/SLevelViewport.cpp
--- a/Engine/Source/Editor/LevelEditor/Private/SLevelViewport.cpp
+++ b/Engine/Source/Editor/LevelEditor/Private/SLevelViewport.cpp
@@ -1,4 +1,4 @@
-﻿// Copyright 1998-2018 Epic Games, Inc. All Rights Reserved.
+// Copyright 1998-2018 Epic Games, Inc. All Rights Reserved.
 
 #include "SLevelViewport.h"
 #include "Materials/MaterialInterface.h"
@@ -3925,14 +3925,17 @@ void SLevelViewport::RemoveActorPreview( int32 PreviewIndex, const bool bRemoveF
 		// Remove widget from viewport overlay
 		ActorPreviewHorizontalBox->RemoveSlot(ActorPreviews[PreviewIndex].PreviewWidget.ToSharedRef());
 	}
-	// Clean up our level viewport client
-	if( ActorPreviews[PreviewIndex].LevelViewportClient.IsValid() )
+	if (ActorPreviews.IsValidIndex(PreviewIndex))
 	{
-		ActorPreviews[PreviewIndex].LevelViewportClient->Viewport = NULL;
-	}
+		// Clean up our level viewport client
+		if (ActorPreviews[PreviewIndex].LevelViewportClient.IsValid())
+		{
+			ActorPreviews[PreviewIndex].LevelViewportClient->Viewport = NULL;
+		}
 
-	// Remove from our list of actor previews.  This will destroy our level viewport client and viewport widget.
-	ActorPreviews.RemoveAt( PreviewIndex );
+		// Remove from our list of actor previews.  This will destroy our level viewport client and viewport widget.
+		ActorPreviews.RemoveAt(PreviewIndex);
+	}
 }
 
 void SLevelViewport::AddOverlayWidget(TSharedRef<SWidget> OverlaidWidget)
