-Refactor to consolidate code to get menu text and icons for viewport camera menu
-Swap front/side axis so that viewing matches X as forward

diff --git a/Engine/Source/Editor/Kismet/Private/SSCSEditorViewport.cpp b/Engine/Source/Editor/Kismet/Private/SSCSEditorViewport.cpp
--- a/Engine/Source/Editor/Kismet/Private/SSCSEditorViewport.cpp
+++ b/Engine/Source/Editor/Kismet/Private/SSCSEditorViewport.cpp
@@ -111,96 +111,22 @@ public:
 
 	FText GetCameraMenuLabel() const
 	{
-		FText Label = NSLOCTEXT("BlueprintEditor", "CameraMenuTitle_Default", "Camera");
-
 		if(EditorViewport.IsValid())
 		{
-			switch(EditorViewport.Pin()->GetViewportClient()->GetViewportType())
-			{
-			case LVT_Perspective:
-				Label = NSLOCTEXT("BlueprintEditor", "CameraMenuTitle_Perspective", "Perspective");
-				break;
-
-			case LVT_OrthoXY:
-				Label = NSLOCTEXT("BlueprintEditor", "CameraMenuTitle_Top", "Top");
-				break;
-
-			case LVT_OrthoYZ:
-				Label = NSLOCTEXT("BlueprintEditor", "CameraMenuTitle_Left", "Left");
-				break;
-
-			case LVT_OrthoXZ:
-				Label = NSLOCTEXT("BlueprintEditor", "CameraMenuTitle_Front", "Front");
-				break;
-
-			case LVT_OrthoNegativeXY:
-				Label = NSLOCTEXT("BlueprintEditor", "CameraMenuTitle_Bottom", "Bottom");
-				break;
-
-			case LVT_OrthoNegativeYZ:
-				Label = NSLOCTEXT("BlueprintEditor", "CameraMenuTitle_Right", "Right");
-				break;
-
-			case LVT_OrthoNegativeXZ:
-				Label = NSLOCTEXT("BlueprintEditor", "CameraMenuTitle_Back", "Back");
-				break;
-
-			case LVT_OrthoFreelook:
-				Label = NSLOCTEXT("BlueprintEditor", "CameraMenuTitle_OrthoFreelook", "Ortho");
-				break;
-			}
+			return GetCameraMenuLabelFromViewportType(EditorViewport.Pin()->GetViewportClient()->GetViewportType());
 		}
 
-		return Label;
+		return NSLOCTEXT("BlueprintEditor", "CameraMenuTitle_Default", "Camera");
 	}
 
 	const FSlateBrush* GetCameraMenuLabelIcon() const
 	{
-		static FName PerspectiveIconName("EditorViewport.Perspective");
-		static FName TopIconName("EditorViewport.Top");
-		static FName LeftIconName("EditorViewport.Left");
-		static FName FrontIconName("EditorViewport.Front");
-		static FName BottomIconName("EditorViewport.Bottom");
-		static FName RightIconName("EditorViewport.Right");
-		static FName BackIconName("EditorViewport.Back");
-
-		FName Icon = NAME_None;
-
 		if(EditorViewport.IsValid())
 		{
-			switch(EditorViewport.Pin()->GetViewportClient()->GetViewportType())
-			{
-			case LVT_Perspective:
-				Icon = PerspectiveIconName;
-				break;
-
-			case LVT_OrthoXY:
-				Icon = TopIconName;
-				break;
-
-			case LVT_OrthoYZ:
-				Icon = LeftIconName;
-				break;
-
-			case LVT_OrthoXZ:
-				Icon = FrontIconName;
-				break;
-
-			case LVT_OrthoNegativeXY:
-				Icon = BottomIconName;
-				break;
-
-			case LVT_OrthoNegativeYZ:
-				Icon = RightIconName;
-				break;
-
-			case LVT_OrthoNegativeXZ:
-				Icon = BackIconName;
-				break;
-			}
+			return GetCameraMenuLabelIconFromViewportType( EditorViewport.Pin()->GetViewportClient()->GetViewportType() );
 		}
 
-		return FEditorStyle::GetBrush(Icon);
+		return FEditorStyle::GetBrush(NAME_None);
 	}
 
 	TSharedRef<SWidget> GenerateCameraMenu() const
diff --git a/Engine/Source/Editor/LevelEditor/Private/SLevelViewportToolBar.cpp b/Engine/Source/Editor/LevelEditor/Private/SLevelViewportToolBar.cpp
--- a/Engine/Source/Editor/LevelEditor/Private/SLevelViewportToolBar.cpp
+++ b/Engine/Source/Editor/LevelEditor/Private/SLevelViewportToolBar.cpp
@@ -262,45 +262,13 @@ bool SLevelViewportToolBar::IsViewModeSupported(EViewModeIndex ViewModeIndex) co
 
 FText SLevelViewportToolBar::GetCameraMenuLabel() const
 {
-	FText Label = LOCTEXT("CameraMenuTitle_Default", "Camera");
 	TSharedPtr< SLevelViewport > PinnedViewport( Viewport.Pin() );
 	if( PinnedViewport.IsValid() )
 	{
-		switch( PinnedViewport->GetLevelViewportClient().ViewportType )
-		{
-			case LVT_Perspective:
-				Label = LOCTEXT("CameraMenuTitle_Perspective", "Perspective");
-				break;
-
-			case LVT_OrthoXY:
-				Label = LOCTEXT("CameraMenuTitle_Top", "Top");
-				break;
-
-			case LVT_OrthoYZ:
-				Label = LOCTEXT("CameraMenuTitle_Left", "Left");
-				break;
-
-			case LVT_OrthoXZ:
-				Label = LOCTEXT("CameraMenuTitle_Front", "Front");
-				break;
-
-			case LVT_OrthoNegativeXY:
-				Label = LOCTEXT("CameraMenuTitle_Bottom", "Bottom");
-				break;
-
-			case LVT_OrthoNegativeYZ:
-				Label = LOCTEXT("CameraMenuTitle_Right", "Right");
-				break;
-
-			case LVT_OrthoNegativeXZ:
-				Label = LOCTEXT("CameraMenuTitle_Back", "Back");
-				break;
-			case LVT_OrthoFreelook:
-				break;
-		}
+		return GetCameraMenuLabelFromViewportType( PinnedViewport->GetLevelViewportClient().ViewportType );
 	}
 
-	return Label;
+	return LOCTEXT("CameraMenuTitle_Default", "Camera");
 }
 
 FText SLevelViewportToolBar::GetDevicePreviewMenuLabel() const
@@ -342,53 +310,13 @@ const FSlateBrush* SLevelViewportToolBar::GetDevicePreviewMenuLabelIcon() const
 
 const FSlateBrush* SLevelViewportToolBar::GetCameraMenuLabelIcon() const
 {
-	FName Icon = NAME_None;
 	TSharedPtr< SLevelViewport > PinnedViewport( Viewport.Pin() );
 	if( PinnedViewport.IsValid() )
 	{
-		static FName PerspectiveIcon("EditorViewport.Perspective");
-		static FName TopIcon("EditorViewport.Top");
-		static FName LeftIcon("EditorViewport.Left");
-		static FName FrontIcon("EditorViewport.Front");
-		static FName BottomIcon("EditorViewport.Bottom");
-		static FName RightIcon("EditorViewport.Right");
-		static FName BackIcon("EditorViewport.Back");
-
-		switch( PinnedViewport->GetLevelViewportClient().ViewportType )
-		{
-			case LVT_Perspective:
-				Icon = PerspectiveIcon;
-				break;
-
-			case LVT_OrthoXY:
-				Icon = TopIcon;
-				break;
-
-			case LVT_OrthoYZ:
-				Icon = LeftIcon;
-				break;
-
-			case LVT_OrthoXZ:
-				Icon = FrontIcon;
-				break;
-
-			case LVT_OrthoNegativeXY:
-				Icon = BottomIcon;
-				break;
-
-			case LVT_OrthoNegativeYZ:
-				Icon = RightIcon;
-				break;
-
-			case LVT_OrthoNegativeXZ:
-				Icon = BackIcon;
-				break;
-			case LVT_OrthoFreelook:
-				break;
-		}
+		return GetCameraMenuLabelIconFromViewportType(PinnedViewport->GetLevelViewportClient().ViewportType);
 	}
 
-	return FEditorStyle::GetBrush( Icon );
+	return FEditorStyle::GetBrush(NAME_None);
 }
 
 bool SLevelViewportToolBar::IsCurrentLevelViewport() const
diff --git a/Engine/Source/Editor/Persona/Private/SAnimViewportToolBar.cpp b/Engine/Source/Editor/Persona/Private/SAnimViewportToolBar.cpp
--- a/Engine/Source/Editor/Persona/Private/SAnimViewportToolBar.cpp
+++ b/Engine/Source/Editor/Persona/Private/SAnimViewportToolBar.cpp
@@ -1045,96 +1045,24 @@ FText SAnimViewportToolBar::GetPlaybackMenuLabel() const
 
 FText SAnimViewportToolBar::GetCameraMenuLabel() const
 {
-	FText Label = LOCTEXT("Viewport_Default", "Camera");
 	TSharedPtr< SAnimationEditorViewportTabBody > PinnedViewport(Viewport.Pin());
 	if( PinnedViewport.IsValid() )
 	{
-		switch( PinnedViewport->GetLevelViewportClient().ViewportType )
-		{
-		case LVT_Perspective:
-			Label = LOCTEXT("CameraMenuTitle_Perspective", "Perspective");
-			break;
-
-		case LVT_OrthoXY:
-			Label = LOCTEXT("CameraMenuTitle_Top", "Top");
-			break;
-
-		case LVT_OrthoYZ:
-			Label = LOCTEXT("CameraMenuTitle_Left", "Left");
-			break;
-
-		case LVT_OrthoXZ:
-			Label = LOCTEXT("CameraMenuTitle_Front", "Front");
-			break;
-
-		case LVT_OrthoNegativeXY:
-			Label = LOCTEXT("CameraMenuTitle_Bottom", "Bottom");
-			break;
-
-		case LVT_OrthoNegativeYZ:
-			Label = LOCTEXT("CameraMenuTitle_Right", "Right");
-			break;
-
-		case LVT_OrthoNegativeXZ:
-			Label = LOCTEXT("CameraMenuTitle_Back", "Back");
-			break;
-		case LVT_OrthoFreelook:
-			break;
-		}
+		return GetCameraMenuLabelFromViewportType( PinnedViewport->GetLevelViewportClient().ViewportType );
 	}
 
-	return Label;
+	return LOCTEXT("Viewport_Default", "Camera");
 }
 
 const FSlateBrush* SAnimViewportToolBar::GetCameraMenuLabelIcon() const
 {
-	FName Icon = NAME_None;
 	TSharedPtr< SAnimationEditorViewportTabBody > PinnedViewport(Viewport.Pin());
 	if (PinnedViewport.IsValid())
 	{
-		static FName PerspectiveIcon("EditorViewport.Perspective");
-		static FName TopIcon("EditorViewport.Top");
-		static FName LeftIcon("EditorViewport.Left");
-		static FName FrontIcon("EditorViewport.Front");
-		static FName BottomIcon("EditorViewport.Bottom");
-		static FName RightIcon("EditorViewport.Right");
-		static FName BackIcon("EditorViewport.Back");
-
-		switch (PinnedViewport->GetLevelViewportClient().ViewportType)
-		{
-		case LVT_Perspective:
-			Icon = PerspectiveIcon;
-			break;
-
-		case LVT_OrthoXY:
-			Icon = TopIcon;
-			break;
-
-		case LVT_OrthoYZ:
-			Icon = LeftIcon;
-			break;
-
-		case LVT_OrthoXZ:
-			Icon = FrontIcon;
-			break;
-
-		case LVT_OrthoNegativeXY:
-			Icon = BottomIcon;
-			break;
-
-		case LVT_OrthoNegativeYZ:
-			Icon = RightIcon;
-			break;
-
-		case LVT_OrthoNegativeXZ:
-			Icon = BackIcon;
-			break;
-		case LVT_OrthoFreelook:
-			break;
-		}
+		return GetCameraMenuLabelIconFromViewportType( PinnedViewport->GetLevelViewportClient().ViewportType );
 	}
 
-	return FEditorStyle::GetBrush(Icon);
+	return FEditorStyle::GetBrush(NAME_None);
 }
 
 TOptional<float> SAnimViewportToolBar::OnGetFOVValue() const
diff --git a/Engine/Source/Editor/UnrealEd/Private/SCommonEditorViewportToolbarBase.cpp b/Engine/Source/Editor/UnrealEd/Private/SCommonEditorViewportToolbarBase.cpp
--- a/Engine/Source/Editor/UnrealEd/Private/SCommonEditorViewportToolbarBase.cpp
+++ b/Engine/Source/Editor/UnrealEd/Private/SCommonEditorViewportToolbarBase.cpp
@@ -125,85 +125,12 @@ void SCommonEditorViewportToolbarBase::Construct(const FArguments& InArgs, TShar
 
 FText SCommonEditorViewportToolbarBase::GetCameraMenuLabel() const
 {
-	FText Label = LOCTEXT("CameraMenuTitle_Default", "Camera");
-
-	switch (GetViewportClient().GetViewportType())
-	{
-		case LVT_Perspective:
-			Label = LOCTEXT("CameraMenuTitle_Perspective", "Perspective");
-			break;
-
-		case LVT_OrthoXY:
-			Label = LOCTEXT("CameraMenuTitle_Top", "Top");
-			break;
-
-		case LVT_OrthoYZ:
-			Label = LOCTEXT("CameraMenuTitle_Left", "Left");
-			break;
-
-		case LVT_OrthoXZ:
-			Label = LOCTEXT("CameraMenuTitle_Front", "Front");
-			break;
-
-		case LVT_OrthoNegativeXY:
-			Label = LOCTEXT("CameraMenuTitle_Bottom", "Bottom");
-			break;
-
-		case LVT_OrthoNegativeYZ:
-			Label = LOCTEXT("CameraMenuTitle_Right", "Right");
-			break;
-
-		case LVT_OrthoNegativeXZ:
-			Label = LOCTEXT("CameraMenuTitle_Back", "Back");
-			break;
-
-		case LVT_OrthoFreelook:
-			Label = LOCTEXT("CameraMenuTitle_OrthoFreelook", "Ortho");
-			break;
-	}
-
-	return Label;
+	return GetCameraMenuLabelFromViewportType( GetViewportClient().GetViewportType() );
 }
 
 const FSlateBrush* SCommonEditorViewportToolbarBase::GetCameraMenuLabelIcon() const
 {
-	FName Icon = NAME_None;
-
-	switch (GetViewportClient().GetViewportType())
-	{
-		case LVT_Perspective:
-			Icon = FName( "EditorViewport.Perspective" );
-			break;
-
-		case LVT_OrthoXY:
-			Icon = FName( "EditorViewport.Top" );
-			break;
-
-		case LVT_OrthoYZ:
-			Icon = FName( "EditorViewport.Left" );
-			break;
-
-		case LVT_OrthoXZ:
-			Icon = FName( "EditorViewport.Front" );
-			break;
-
-		case LVT_OrthoNegativeXY:
-			Icon = FName("EditorViewport.Bottom");
-			break;
-
-		case LVT_OrthoNegativeYZ:
-			Icon = FName("EditorViewport.Right");
-			break;
-
-		case LVT_OrthoNegativeXZ:
-			Icon = FName("EditorViewport.Back");
-			break;
-
-		case LVT_OrthoFreelook:
-			break;
-	}
-
-	return FEditorStyle::GetBrush( Icon );
+	return GetCameraMenuLabelIconFromViewportType( GetViewportClient().GetViewportType() );
 }
 
 EVisibility SCommonEditorViewportToolbarBase::GetViewModeOptionsVisibility() const
diff --git a/Engine/Source/Editor/UnrealEd/Private/SEditorViewport.cpp b/Engine/Source/Editor/UnrealEd/Private/SEditorViewport.cpp
--- a/Engine/Source/Editor/UnrealEd/Private/SEditorViewport.cpp
+++ b/Engine/Source/Editor/UnrealEd/Private/SEditorViewport.cpp
@@ -189,15 +189,15 @@ void SEditorViewport::BindCommands()
 
 	CommandListRef.MapAction( 
 		Commands.Front,
-		FExecuteAction::CreateSP( ClientRef, &FEditorViewportClient::SetViewportType, LVT_OrthoXZ ),
+		FExecuteAction::CreateSP( ClientRef, &FEditorViewportClient::SetViewportType, LVT_OrthoNegativeYZ ),
 		FCanExecuteAction(),
-		FIsActionChecked::CreateSP(ClientRef, &FEditorViewportClient::IsActiveViewportType, LVT_OrthoXZ));
+		FIsActionChecked::CreateSP(ClientRef, &FEditorViewportClient::IsActiveViewportType, LVT_OrthoNegativeYZ));
 
 	CommandListRef.MapAction( 
 		Commands.Left,
-		FExecuteAction::CreateSP( ClientRef, &FEditorViewportClient::SetViewportType, LVT_OrthoYZ ),
+		FExecuteAction::CreateSP( ClientRef, &FEditorViewportClient::SetViewportType, LVT_OrthoNegativeXZ ),
 		FCanExecuteAction(),
-		FIsActionChecked::CreateSP(ClientRef, &FEditorViewportClient::IsActiveViewportType, LVT_OrthoYZ));
+		FIsActionChecked::CreateSP(ClientRef, &FEditorViewportClient::IsActiveViewportType, LVT_OrthoNegativeXZ));
 
 	CommandListRef.MapAction( 
 		Commands.Top,
@@ -207,15 +207,15 @@ void SEditorViewport::BindCommands()
 
 	CommandListRef.MapAction(
 		Commands.Back,
-		FExecuteAction::CreateSP(ClientRef, &FEditorViewportClient::SetViewportType, LVT_OrthoNegativeXZ),
+		FExecuteAction::CreateSP(ClientRef, &FEditorViewportClient::SetViewportType, LVT_OrthoYZ),
 		FCanExecuteAction(),
-		FIsActionChecked::CreateSP(ClientRef, &FEditorViewportClient::IsActiveViewportType, LVT_OrthoNegativeXZ));
+		FIsActionChecked::CreateSP(ClientRef, &FEditorViewportClient::IsActiveViewportType, LVT_OrthoYZ));
 
 	CommandListRef.MapAction(
 		Commands.Right,
-		FExecuteAction::CreateSP(ClientRef, &FEditorViewportClient::SetViewportType, LVT_OrthoNegativeYZ),
+		FExecuteAction::CreateSP(ClientRef, &FEditorViewportClient::SetViewportType, LVT_OrthoXZ),
 		FCanExecuteAction(),
-		FIsActionChecked::CreateSP(ClientRef, &FEditorViewportClient::IsActiveViewportType, LVT_OrthoNegativeYZ));
+		FIsActionChecked::CreateSP(ClientRef, &FEditorViewportClient::IsActiveViewportType, LVT_OrthoXZ));
 
 	CommandListRef.MapAction(
 		Commands.Bottom,
diff --git a/Engine/Source/Editor/UnrealEd/Private/SViewportToolBar.cpp b/Engine/Source/Editor/UnrealEd/Private/SViewportToolBar.cpp
--- a/Engine/Source/Editor/UnrealEd/Private/SViewportToolBar.cpp
+++ b/Engine/Source/Editor/UnrealEd/Private/SViewportToolBar.cpp
@@ -4,6 +4,7 @@
 #include "SViewportToolBar.h"
 #include "Framework/Application/SlateApplication.h"
 #include "Widgets/Input/SMenuAnchor.h"
+#include "EditorStyleSet.h"
 
 #define LOCTEXT_NAMESPACE "ViewportToolBar"
 
@@ -97,6 +98,93 @@ void SViewportToolBar::OnMouseLeave( const FPointerEvent& MouseEvent )
 	}
 }
 
+FText SViewportToolBar::GetCameraMenuLabelFromViewportType(const ELevelViewportType ViewportType) const
+{
+	FText Label = LOCTEXT("CameraMenuTitle_Default", "Camera");
+	switch (ViewportType)
+	{
+	case LVT_Perspective:
+		Label = LOCTEXT("CameraMenuTitle_Perspective", "Perspective");
+		break;
+
+	case LVT_OrthoXY:
+		Label = LOCTEXT("CameraMenuTitle_Top", "Top");
+		break;
+
+	case LVT_OrthoNegativeXZ:
+		Label = LOCTEXT("CameraMenuTitle_Left", "Left");
+		break;
+
+	case LVT_OrthoNegativeYZ:
+		Label = LOCTEXT("CameraMenuTitle_Front", "Front");
+		break;
+
+	case LVT_OrthoNegativeXY:
+		Label = LOCTEXT("CameraMenuTitle_Bottom", "Bottom");
+		break;
+
+	case LVT_OrthoXZ:
+		Label = LOCTEXT("CameraMenuTitle_Right", "Right");
+		break;
+
+	case LVT_OrthoYZ:
+		Label = LOCTEXT("CameraMenuTitle_Back", "Back");
+		break;
+	case LVT_OrthoFreelook:
+		break;
+	}
+
+	return Label;
+}
+
+const FSlateBrush* SViewportToolBar::GetCameraMenuLabelIconFromViewportType(const ELevelViewportType ViewportType) const
+{
+	static FName PerspectiveIcon("EditorViewport.Perspective");
+	static FName TopIcon("EditorViewport.Top");
+	static FName LeftIcon("EditorViewport.Left");
+	static FName FrontIcon("EditorViewport.Front");
+	static FName BottomIcon("EditorViewport.Bottom");
+	static FName RightIcon("EditorViewport.Right");
+	static FName BackIcon("EditorViewport.Back");
+
+	FName Icon = NAME_None;
+
+	switch (ViewportType)
+	{
+	case LVT_Perspective:
+		Icon = PerspectiveIcon;
+		break;
+
+	case LVT_OrthoXY:
+		Icon = TopIcon;
+		break;
+
+	case LVT_OrthoNegativeXZ:
+		Icon = LeftIcon;
+		break;
+
+	case LVT_OrthoNegativeYZ:
+		Icon = FrontIcon;
+		break;
+
+	case LVT_OrthoNegativeXY:
+		Icon = BottomIcon;
+		break;
+
+	case LVT_OrthoXZ:
+		Icon = RightIcon;
+		break;
+
+	case LVT_OrthoYZ:
+		Icon = BackIcon;
+		break;
+	case LVT_OrthoFreelook:
+		break;
+	}
+
+	return FEditorStyle::GetBrush(Icon);
+}
+
 bool SViewportToolBar::IsViewModeSupported(EViewModeIndex ViewModeIndex) const 
 {
 	switch (ViewModeIndex)
diff --git a/Engine/Source/Editor/UnrealEd/Public/SViewportToolBar.h b/Engine/Source/Editor/UnrealEd/Public/SViewportToolBar.h
--- a/Engine/Source/Editor/UnrealEd/Public/SViewportToolBar.h
+++ b/Engine/Source/Editor/UnrealEd/Public/SViewportToolBar.h
@@ -8,6 +8,7 @@
 #include "Widgets/DeclarativeSyntaxSupport.h"
 #include "Widgets/SCompoundWidget.h"
 #include "Animation/CurveSequence.h"
+#include "Editor/UnrealEdTypes.h"
 
 class SMenuAnchor;
 
@@ -59,6 +60,25 @@ private:
 	 */
 	virtual void OnMouseLeave( const FPointerEvent& MouseEvent ) override;
 
+protected:
+	/**
+	* Returns the label for the "Camera" tool bar menu based on based viewport type
+	*
+	* @param ViewportType	The Viewport type we want the label for
+	*
+	* @return	Label to use for this menu label
+	*/
+	virtual FText GetCameraMenuLabelFromViewportType(const ELevelViewportType ViewportType) const;
+
+	/**
+	* Returns the label icon for the "Camera" tool bar menu, which changes depending on the viewport type
+	*
+	* @param ViewportType	The Viewport type we want the icon for
+	*
+	* @return	Label icon to use for this menu label
+	*/
+	virtual const FSlateBrush* GetCameraMenuLabelIconFromViewportType(const ELevelViewportType ViewportType) const;
+
 private:
 	/** Curve sequence for fading out the toolbar */
 	FCurveSequence FadeOutSequence;
