Camera preview viewports now also can preview the game world if you are currently simulating

diff --git a/Engine/Source/Editor/LevelEditor/Private/SLevelViewport.cpp b/Engine/Source/Editor/LevelEditor/Private/SLevelViewport.cpp
--- a/Engine/Source/Editor/LevelEditor/Private/SLevelViewport.cpp
+++ b/Engine/Source/Editor/LevelEditor/Private/SLevelViewport.cpp
@@ -1,4 +1,4 @@
-﻿// Copyright 1998-2018 Epic Games, Inc. All Rights Reserved.
+// Copyright 1998-2018 Epic Games, Inc. All Rights Reserved.
 
 #include "SLevelViewport.h"
 #include "Materials/MaterialInterface.h"
@@ -205,6 +205,9 @@ void SLevelViewport::Construct(const FArguments& InArgs)
 	LevelEditor.OnMapChanged().AddRaw( this, &SLevelViewport::OnMapChanged );
 
 	GEngine->OnLevelActorDeleted().AddRaw( this, &SLevelViewport::OnLevelActorsRemoved );
+
+	FEditorDelegates::PostPIEStarted.AddSP(this, &SLevelViewport::TransitionToPIE);
+	FEditorDelegates::PrePIEEnded.AddSP(this, &SLevelViewport::TransitionFromPIE);
 }
 
 void SLevelViewport::ConstructViewportOverlayContent()
@@ -468,6 +471,28 @@ const FSceneViewport* SLevelViewport::GetGameSceneViewport() const
 	return ActiveViewport.Get();
 }
 
+void SLevelViewport::TransitionToPIE(bool bIsSimulating)
+{
+	for (FViewportActorPreview& ActorPreview : ActorPreviews)
+	{
+		if (ActorPreview.LevelViewportClient.IsValid() && !ActorPreview.LevelViewportClient->IsSimulateInEditorViewport())
+		{
+			ActorPreview.LevelViewportClient->SetIsSimulateInEditorViewport(true);
+		}
+	}
+}
+
+void SLevelViewport::TransitionFromPIE(bool bIsSimulating)
+{
+	for (FViewportActorPreview& ActorPreview : ActorPreviews)
+	{
+		if (ActorPreview.LevelViewportClient.IsValid() && ActorPreview.LevelViewportClient->IsSimulateInEditorViewport())
+		{
+			ActorPreview.LevelViewportClient->SetIsSimulateInEditorViewport(false);
+		}
+	}
+}
+
 FReply SLevelViewport::OnKeyDown( const FGeometry& MyGeometry, const FKeyEvent& InKeyEvent )
 {
 
@@ -3140,7 +3165,13 @@ void SLevelViewport::PreviewActors( const TArray< AActor* >& ActorsToPreview, co
 				// Push actor transform to view.  From here on out, this will happen automatically in FLevelEditorViewportClient::Tick.
 				// The reason we allow the viewport client to update this is to avoid off-by-one-frame issues when dragging actors around.
 				ActorPreviewLevelViewportClient->SetActorLock( CurActor );
-				ActorPreviewLevelViewportClient->UpdateViewForLockedActor();			
+				ActorPreviewLevelViewportClient->UpdateViewForLockedActor();	
+
+				// Preview the play world if the current actor is in the play world
+				if (CurActor->GetWorld()->IsGameWorld())
+				{
+					ActorPreviewLevelViewportClient->SetIsSimulateInEditorViewport(true);
+				}
 			}
 
 			TSharedPtr< SActorPreview > ActorPreviewWidget = SNew(SActorPreview)
diff --git a/Engine/Source/Editor/LevelEditor/Public/SLevelViewport.h b/Engine/Source/Editor/LevelEditor/Public/SLevelViewport.h
--- a/Engine/Source/Editor/LevelEditor/Public/SLevelViewport.h
+++ b/Engine/Source/Editor/LevelEditor/Public/SLevelViewport.h
@@ -1,4 +1,4 @@
-﻿// Copyright 1998-2018 Epic Games, Inc. All Rights Reserved.
+// Copyright 1998-2018 Epic Games, Inc. All Rights Reserved.
 
 
 #pragma once
@@ -684,6 +684,12 @@ private:
 	/** Check whether we should display the full toolbar or not */
 	bool ShouldShowFullToolbar() const { return bShowFullToolbar; }
 
+	/** Handle any level viewport changes on entering PIE or simulate */
+	void TransitionToPIE(bool bIsSimulating);
+
+	/** Handle any level viewport changes on leaving PIE or simulate */
+	void TransitionFromPIE(bool bIsSimulating);
+
 private:
 	/** Tab which this viewport is located in */
 	TWeakPtr<class FLevelViewportLayout> ParentLayout;
diff --git a/Engine/Source/Editor/UnrealEd/Private/LevelEditorViewport.cpp b/Engine/Source/Editor/UnrealEd/Private/LevelEditorViewport.cpp
--- a/Engine/Source/Editor/UnrealEd/Private/LevelEditorViewport.cpp
+++ b/Engine/Source/Editor/UnrealEd/Private/LevelEditorViewport.cpp
@@ -4710,16 +4710,6 @@ void FLevelEditorViewportClient::RemoveReferenceToWorldContext(FWorldContext& Wo
 void FLevelEditorViewportClient::SetIsSimulateInEditorViewport( bool bInIsSimulateInEditorViewport )
 { 
 	bIsSimulateInEditorViewport = bInIsSimulateInEditorViewport; 
-
-	if (bInIsSimulateInEditorViewport)
-	{
-		TSharedRef<FPhysicsManipulationEdModeFactory> Factory = MakeShareable( new FPhysicsManipulationEdModeFactory );
-		FEditorModeRegistry::Get().RegisterMode( FBuiltinEditorModes::EM_Physics, Factory );
-	}
-	else
-	{
-		FEditorModeRegistry::Get().UnregisterMode( FBuiltinEditorModes::EM_Physics );
-	}
 }
 
 #undef LOCTEXT_NAMESPACE
diff --git a/Engine/Source/Editor/UnrealEd/Private/PlayLevel.cpp b/Engine/Source/Editor/UnrealEd/Private/PlayLevel.cpp
--- a/Engine/Source/Editor/UnrealEd/Private/PlayLevel.cpp
+++ b/Engine/Source/Editor/UnrealEd/Private/PlayLevel.cpp
@@ -100,6 +100,8 @@
 #include "Kismet2/DebuggerCommands.h"
 #include "Misc/ScopeExit.h"
 #include "IVREditorModule.h"
+#include "EditorModeRegistry.h"
+#include "PhysicsManipulationMode.h"
 
 
 DEFINE_LOG_CATEGORY_STATIC(LogPlayLevel, Log, All);
@@ -610,6 +612,10 @@ void UEditorEngine::TeardownPlaySession(FWorldContext& PieWorldContext)
 				// No longer simulating in the viewport
 				Viewport->GetLevelViewportClient().SetIsSimulateInEditorViewport( false );
 
+				
+				FEditorModeRegistry::Get().UnregisterMode(FBuiltinEditorModes::EM_Physics);
+				
+
 				// Clear out the hit proxies before GC'ing
 				Viewport->GetLevelViewportClient().Viewport->InvalidateHitProxy();
 			}
@@ -3455,6 +3461,9 @@ void UEditorEngine::ToggleBetweenPIEandSIE( bool bNewSession )
 			// No longer simulating
 			GameViewport->SetIsSimulateInEditorViewport(false);
 			EditorViewportClient.SetIsSimulateInEditorViewport(false);
+
+			FEditorModeRegistry::Get().UnregisterMode(FBuiltinEditorModes::EM_Physics);
+			
 			bIsSimulatingInEditor = false;
 		}
 		else
@@ -3465,6 +3474,11 @@ void UEditorEngine::ToggleBetweenPIEandSIE( bool bNewSession )
 			GameViewport->SetIsSimulateInEditorViewport(true);
 			GameViewport->GetGameViewport()->SetPlayInEditorIsSimulate(true);
 			EditorViewportClient.SetIsSimulateInEditorViewport(true);
+
+		
+			TSharedRef<FPhysicsManipulationEdModeFactory> Factory = MakeShareable(new FPhysicsManipulationEdModeFactory);
+			FEditorModeRegistry::Get().RegisterMode(FBuiltinEditorModes::EM_Physics, Factory);
+			
 			bIsSimulatingInEditor = true;
 
 			// Make sure the viewport is in real-time mode
