PR #5056: Added bindable command AlignBrushVerticesToGrid (Contributed by [REDACTED_PERSON])

diff --git a/Engine/Source/Editor/LevelEditor/Private/LevelEditor.cpp b/Engine/Source/Editor/LevelEditor/Private/LevelEditor.cpp
--- a/Engine/Source/Editor/LevelEditor/Private/LevelEditor.cpp
+++ b/Engine/Source/Editor/LevelEditor/Private/LevelEditor.cpp
@@ -799,6 +799,10 @@ void FLevelEditorModule::BindGlobalLevelEditorCommands()
 		Commands.Select2DLayerBelow,
 		FExecuteAction::CreateStatic(&FLevelEditorActionCallbacks::Select2DLayerDeltaAway_Clicked, 1)
 		);
+	ActionList.MapAction(
+		Commands.AlignBrushVerticesToGrid,
+		FExecuteAction::CreateStatic(&FLevelEditorActionCallbacks::AlignBrushVerticesToGrid_Execute)
+		);
 
 	bAlign = false;
 	bool bUseLineTrace = false;
@@ -1317,6 +1321,7 @@ void FLevelEditorModule::BindGlobalLevelEditorCommands()
 		FExecuteAction::CreateStatic( &FLevelEditorActionCallbacks::ExecuteExecCommand, FString( TEXT("BRUSH SEPARATEPOLYS") ) )
 		);
 
+
 	ActionList.MapAction(
 		Commands.CreateBoundingBoxVolume,
 		FExecuteAction::CreateStatic( &FLevelEditorActionCallbacks::ExecuteExecCommand, FString( TEXT("ACTOR CREATE_BV_BOUNDINGBOX SnapToGrid=1") ) )
diff --git a/Engine/Source/Editor/LevelEditor/Private/LevelEditorActions.cpp b/Engine/Source/Editor/LevelEditor/Private/LevelEditorActions.cpp
--- a/Engine/Source/Editor/LevelEditor/Private/LevelEditorActions.cpp
+++ b/Engine/Source/Editor/LevelEditor/Private/LevelEditorActions.cpp
@@ -2851,6 +2851,12 @@ void FLevelEditorActionCallbacks::SnapTo_Clicked( const bool InAlign, const bool
 	GEditor->RedrawLevelEditingViewports();
 }
 
+void FLevelEditorActionCallbacks::AlignBrushVerticesToGrid_Execute()
+{
+	UWorld* World = GUnrealEd->GetWorld();
+	GEditor->Exec(World, TEXT("ACTOR ALIGN VERTS"));
+}
+
 bool FLevelEditorActionCallbacks::ActorSelected_CanExecute()
 {
 	// Had to have something selected
@@ -3016,6 +3022,7 @@ void FLevelEditorCommands::RegisterCommands()
 
 	UI_COMMAND( MergePolys, "Merge", "Merges multiple polygons on a brush face into as few as possible", EUserInterfaceActionType::Button, FInputChord() );
 	UI_COMMAND( SeparatePolys, "Separate", "Reverses the effect of a previous merge", EUserInterfaceActionType::Button, FInputChord() );
+	UI_COMMAND( AlignBrushVerticesToGrid, "Align Brush Verticies To Grid", "Align brush verticies to the grid", EUserInterfaceActionType::Button, FInputChord() );
 
 	// RegroupActors uses GroupActors for it's label and tooltip when simply grouping a selection of actors using overrides. This is to provide display of the chord which is the same for both.
 	UI_COMMAND( GroupActors, "Group", "Groups the selected actors", EUserInterfaceActionType::Button, FInputChord( /*EKeys::G, EModifierKey::Control*/ ) );
diff --git a/Engine/Source/Editor/LevelEditor/Public/LevelEditorActions.h b/Engine/Source/Editor/LevelEditor/Public/LevelEditorActions.h
--- a/Engine/Source/Editor/LevelEditor/Public/LevelEditorActions.h
+++ b/Engine/Source/Editor/LevelEditor/Public/LevelEditorActions.h
@@ -290,6 +290,8 @@ public:
 	/** Reverse a merge */
 	TSharedPtr< FUICommandInfo > SeparatePolys;
 
+	/** Align brush verticies to the grid */
+	TSharedPtr<FUICommandInfo> AlignBrushVerticesToGrid;
 
 	/**
 	 * Actor group commands
@@ -1255,6 +1257,11 @@ public:
 	 */
 	static void SnapActorToActor_Clicked( bool InAlign, bool InUseLineTrace, bool InUseBounds, bool InUsePivot );
 
+	/**
+	 * Aligns brush verticies to the nearest grid point.
+	 */
+	static void AlignBrushVerticesToGrid_Execute();
+
 	/**
 	 * Checks to see if multiple actors are selected
 	 *	@return true if it can execute.
