PR #5115: UE-64064: Show remove button when favorite level count is over 10 (Contributed by [REDACTED_PERSON])

diff --git a/Engine/Source/Editor/LevelEditor/Private/LevelEditorActions.cpp b/Engine/Source/Editor/LevelEditor/Private/LevelEditorActions.cpp
--- a/Engine/Source/Editor/LevelEditor/Private/LevelEditorActions.cpp
+++ b/Engine/Source/Editor/LevelEditor/Private/LevelEditorActions.cpp
@@ -358,10 +358,15 @@ void FLevelEditorActionCallbacks::RemoveFavorite( int32 FavoriteFileIndex )
 
 bool FLevelEditorActionCallbacks::ToggleFavorite_CanExecute()
 {
-	const FMainMRUFavoritesList& MRUFavorites = *FModuleManager::LoadModuleChecked<IMainFrameModule>("MainFrame").GetMRUFavoritesList();
-	const int32 NumFavorites = MRUFavorites.GetNumFavorites();
-	// Disable the favorites button if the map isn't associated to a file yet (new map, never before saved, etc.)
-	return LevelEditorActionsHelpers::IsPersistentWorld(GetWorld()) && NumFavorites <= FLevelEditorCommands::Get().OpenFavoriteFileCommands.Num();
+	if (LevelEditorActionsHelpers::IsPersistentWorld(GetWorld()))
+	{
+		const FMainMRUFavoritesList& MRUFavorites = *FModuleManager::LoadModuleChecked<IMainFrameModule>("MainFrame").GetMRUFavoritesList();
+		const int32 NumFavorites = MRUFavorites.GetNumFavorites();
+		// Disable the favorites button if the map isn't associated to a file yet (new map, never before saved, etc.)
+		const FString PackageName = GetWorld()->GetOutermost()->GetName();
+		return (NumFavorites <= FLevelEditorCommands::Get().OpenFavoriteFileCommands.Num() || MRUFavorites.ContainsFavoritesItem(PackageName));
+	}
+	return false;
 }
 
 
diff --git a/Engine/Source/Editor/LevelEditor/Private/LevelEditorMenu.cpp b/Engine/Source/Editor/LevelEditor/Private/LevelEditorMenu.cpp
--- a/Engine/Source/Editor/LevelEditor/Private/LevelEditorMenu.cpp
+++ b/Engine/Source/Editor/LevelEditor/Private/LevelEditorMenu.cpp
@@ -75,13 +75,9 @@ TSharedRef< SWidget > FLevelEditorMenu::MakeLevelEditorMenu( const TSharedPtr<FU
 						static FText GetToggleFavoriteLabelText()
 						{
 							const FText LevelName = FText::FromString(FPackageName::GetShortName(GWorld->GetOutermost()->GetFName()));
-							if (FLevelEditorActionCallbacks::ToggleFavorite_CanExecute())
+							if (!FLevelEditorActionCallbacks::ToggleFavorite_IsChecked())
 							{
-
-								if (!FLevelEditorActionCallbacks::ToggleFavorite_IsChecked())
-								{
-									return FText::Format(LOCTEXT("ToggleFavorite_Add", "Add {0} to Favorites"), LevelName);
-								}
+								return FText::Format(LOCTEXT("ToggleFavorite_Add", "Add {0} to Favorites"), LevelName);
 							}
 							return FText::Format(LOCTEXT("ToggleFavorite_Remove", "Remove {0} from Favorites"), LevelName);
 						}
@@ -94,8 +90,6 @@ TSharedRef< SWidget > FLevelEditorMenu::MakeLevelEditorMenu( const TSharedPtr<FU
 						{
 							InMenuBuilder.BeginSection("LevelEditorToggleFavorite");
 							{
-								const FString LevelName = FPackageName::GetShortName(GWorld->GetOutermost()->GetFName());
-
 								TAttribute<FText> ToggleFavoriteLabel;
 								ToggleFavoriteLabel.BindStatic(&Local::GetToggleFavoriteLabelText);
 								InMenuBuilder.AddMenuEntry(FLevelEditorCommands::Get().ToggleFavorite, NAME_None, ToggleFavoriteLabel);
diff --git a/Engine/Source/Editor/UnrealEd/Private/MRUFavoritesList.cpp b/Engine/Source/Editor/UnrealEd/Private/MRUFavoritesList.cpp
--- a/Engine/Source/Editor/UnrealEd/Private/MRUFavoritesList.cpp
+++ b/Engine/Source/Editor/UnrealEd/Private/MRUFavoritesList.cpp
@@ -102,7 +102,7 @@ void FMainMRUFavoritesList::MoveFavoritesItemToHead(const FString& Item)
  *
  * @return	true if the provided item is in the favorite's list; false if it is not
  */
-bool FMainMRUFavoritesList::ContainsFavoritesItem( const FString& Item )
+bool FMainMRUFavoritesList::ContainsFavoritesItem( const FString& Item ) const
 {
 	if (ensureMsgf(FPackageName::IsValidLongPackageName(Item), TEXT("Item is not a valid long package name: '%s'"), *Item))
 	{
diff --git a/Engine/Source/Editor/UnrealEd/Public/MRUFavoritesList.h b/Engine/Source/Editor/UnrealEd/Public/MRUFavoritesList.h
--- a/Engine/Source/Editor/UnrealEd/Public/MRUFavoritesList.h
+++ b/Engine/Source/Editor/UnrealEd/Public/MRUFavoritesList.h
@@ -59,7 +59,7 @@ public:
 	 *
 	 * @return	true if the provided item is in the favorite's list; false if it is not
 	 */
-	bool ContainsFavoritesItem( const FString& Item );
+	bool ContainsFavoritesItem( const FString& Item ) const;
 
 	/**
 	 * Return the favorites item specified by the provided index
