Changed the default keyboard shortcuts for editor viewport's immersive mode and blueprint debugger's step into command to use Cmd+F11 instead of F11 on Mac to avoid conflict with default macOS Expose shortcut

diff --git a/Engine/Source/Editor/LevelEditor/Private/LevelViewportActions.cpp b/Engine/Source/Editor/LevelEditor/Private/LevelViewportActions.cpp
--- a/Engine/Source/Editor/LevelEditor/Private/LevelViewportActions.cpp
+++ b/Engine/Source/Editor/LevelEditor/Private/LevelViewportActions.cpp
@@ -1,4 +1,4 @@
-﻿// Copyright 1998-2018 Epic Games, Inc. All Rights Reserved.
+// Copyright 1998-2018 Epic Games, Inc. All Rights Reserved.
 
 #include "LevelViewportActions.h"
 #include "ShowFlags.h"
@@ -31,7 +31,7 @@ void FLevelViewportCommands::RegisterCommands()
 {
 	UI_COMMAND( ToggleMaximize, "Maximize Viewport", "Toggles the Maximize state of the current viewport", EUserInterfaceActionType::ToggleButton, FInputChord() );
 	UI_COMMAND( ToggleGameView, "Game View", "Toggles game view.  Game view shows the scene as it appears in game", EUserInterfaceActionType::ToggleButton, FInputChord( EKeys::G ) );
-	UI_COMMAND( ToggleImmersive, "Immersive Mode", "Switches this viewport between immersive mode and regular mode", EUserInterfaceActionType::ToggleButton, FInputChord( EKeys::F11 ) );
+	UI_COMMAND( ToggleImmersive, "Immersive Mode", "Switches this viewport between immersive mode and regular mode", EUserInterfaceActionType::ToggleButton, PLATFORM_MAC ? FInputChord( EModifierKey::Control, EKeys::F11 ) : FInputChord( EKeys::F11 ) );
 
 	UI_COMMAND( CreateCamera, "Create Camera Here", "Creates a new camera actor at the current location of this viewport's camera", EUserInterfaceActionType::Button, FInputChord() );
 	UI_COMMAND( HighResScreenshot, "High Resolution Screenshot...", "Opens the control panel for high resolution screenshots", EUserInterfaceActionType::Button, FInputChord() );
diff --git a/Engine/Source/Editor/UnrealEd/Private/Kismet2/DebuggerCommands.cpp b/Engine/Source/Editor/UnrealEd/Private/Kismet2/DebuggerCommands.cpp
--- a/Engine/Source/Editor/UnrealEd/Private/Kismet2/DebuggerCommands.cpp
+++ b/Engine/Source/Editor/UnrealEd/Private/Kismet2/DebuggerCommands.cpp
@@ -340,7 +340,7 @@ void FPlayWorldCommands::RegisterCommands()
 	UI_COMMAND( TogglePlayPauseOfPlaySession, "Toggle Play/Pause", "Resume playing if paused, or pause if playing", EUserInterfaceActionType::Button, FInputChord( EKeys::Pause ) );
 	UI_COMMAND( PossessEjectPlayer, "Possess or Eject Player", "Possesses or ejects the player from the camera", EUserInterfaceActionType::Button, FInputChord( EKeys::F8 ) );
 	UI_COMMAND( ShowCurrentStatement, "Find Node", "Show the current node", EUserInterfaceActionType::Button, FInputChord() );
-	UI_COMMAND( StepInto, "Step Into", "Step Into the next node to be executed", EUserInterfaceActionType::Button, FInputChord(EKeys::F11) );
+	UI_COMMAND( StepInto, "Step Into", "Step Into the next node to be executed", EUserInterfaceActionType::Button, PLATFORM_MAC ? FInputChord(EModifierKey::Control, EKeys::F11) : FInputChord(EKeys::F11) );
 	UI_COMMAND( StepOver, "Step Over", "Step to the next node to be executed in the current graph", EUserInterfaceActionType::Button, FInputChord(EKeys::F10) );
 	UI_COMMAND( StepOut, "Step Out", "Step Out to the next node to be executed in the parent graph", EUserInterfaceActionType::Button, FInputChord(EModifierKey::Alt|EModifierKey::Shift, EKeys::F11) );
 
