Set SCSEditor hidden in Details view when locked actor is invalidated

diff --git a/Engine/Source/Editor/LevelEditor/Private/SActorDetails.cpp b/Engine/Source/Editor/LevelEditor/Private/SActorDetails.cpp
--- a/Engine/Source/Editor/LevelEditor/Private/SActorDetails.cpp
+++ b/Engine/Source/Editor/LevelEditor/Private/SActorDetails.cpp
@@ -119,8 +119,8 @@ void SActorDetails::Construct(const FArguments& InArgs, const FName TabIdentifie
 
 	DetailsView->SetIsPropertyVisibleDelegate(FIsPropertyVisible::CreateLambda(IsPropertyVisible));
 	DetailsView->SetIsPropertyReadOnlyDelegate(FIsPropertyReadOnly::CreateSP(this, &SActorDetails::IsPropertyReadOnly));
-
 	DetailsView->SetIsPropertyEditingEnabledDelegate(FIsPropertyEditingEnabled::CreateSP(this, &SActorDetails::IsPropertyEditingEnabled));
+	DetailsView->SetOnObjectArrayChanged(FOnObjectArrayChanged::CreateSP(this, &SActorDetails::OnDetailsViewObjectArrayChanged));
 
 
 	// Set up a delegate to call to add generic details to the view
@@ -219,6 +219,15 @@ SActorDetails::~SActorDetails()
 	}
 }
 
+void SActorDetails::OnDetailsViewObjectArrayChanged(const FString& InTitle, const TArray<TWeakObjectPtr<UObject>>& InObjects)
+{
+	// The DetailsView will already check validity every tick and hide itself when invalid, so this piggy-backs on that code instead of needing a second tick function.
+	if (InObjects.Num() == 0 && !LockedActorSelection.IsValid())
+	{
+		ComponentsBox->SetVisibility(EVisibility::Collapsed);
+	}
+}
+
 void SActorDetails::SetObjects(const TArray<UObject*>& InObjects, bool bForceRefresh)
 {
 	if(!DetailsView->IsLocked())
diff --git a/Engine/Source/Editor/LevelEditor/Private/SActorDetails.h b/Engine/Source/Editor/LevelEditor/Private/SActorDetails.h
--- a/Engine/Source/Editor/LevelEditor/Private/SActorDetails.h
+++ b/Engine/Source/Editor/LevelEditor/Private/SActorDetails.h
@@ -55,6 +55,7 @@ private:
 	void OnSCSEditorTreeViewSelectionChanged(const TArray<TSharedPtr<class FSCSEditorTreeNode> >& SelectedNodes);
 	void OnSCSEditorTreeViewItemDoubleClicked(const TSharedPtr<class FSCSEditorTreeNode> ClickedNode);
 	void UpdateComponentTreeFromEditorSelection();
+	void OnDetailsViewObjectArrayChanged(const FString& InTitle, const TArray<TWeakObjectPtr<UObject>>& InObjects);
 
 	bool IsPropertyReadOnly(const struct FPropertyAndParent& PropertyAndParent) const;
 	bool IsPropertyEditingEnabled() const;
