Fix crash related to registering an editor mode after startup

diff --git a/Engine/Plugins/Editor/MeshEditor/Source/MeshEditor/MeshEditorModule.cpp b/Engine/Plugins/Editor/MeshEditor/Source/MeshEditor/MeshEditorModule.cpp
--- a/Engine/Plugins/Editor/MeshEditor/Source/MeshEditor/MeshEditorModule.cpp
+++ b/Engine/Plugins/Editor/MeshEditor/Source/MeshEditor/MeshEditorModule.cpp
@@ -96,7 +96,6 @@ void FMeshEditorModule::Register()
 
 	FMeshEditorStyle::Initialize();
 
-	FLevelEditorModesCommands::Unregister();
 	FEditorModeRegistry::Get().RegisterMode<FMeshEditorMode>(
 		GetEditorModeID(),
 		LOCTEXT( "ModeName", "Mesh Editor" ),
@@ -157,7 +156,6 @@ void FMeshEditorModule::Unregister()
 		SettingsModule->UnregisterSettings( "Editor", "ContentEditors", "MeshEditor" );
 	}
 
-	FLevelEditorModesCommands::Unregister();
 	FEditorModeRegistry::Get().UnregisterMode(GetEditorModeID());
 
 	FMeshEditorStyle::Shutdown();
diff --git a/Engine/Source/Editor/LevelEditor/Private/SLevelEditor.cpp b/Engine/Source/Editor/LevelEditor/Private/SLevelEditor.cpp
--- a/Engine/Source/Editor/LevelEditor/Private/SLevelEditor.cpp
+++ b/Engine/Source/Editor/LevelEditor/Private/SLevelEditor.cpp
@@ -155,7 +155,7 @@ void SLevelEditor::Construct( const SLevelEditor::FArguments& InArgs)
 	BindCommands();
 
 	// We need to register when modes list changes so that we can refresh the auto generated commands.
-	FEditorModeRegistry::Get().OnRegisteredModesChanged().AddRaw(this, &SLevelEditor::RefreshEditorModeCommands);
+	FEditorModeRegistry::Get().OnRegisteredModesChanged().AddRaw(this, &SLevelEditor::EditorModeCommandsChanged);
 
 	// @todo This is a hack to get this working for now. This won't work with multiple worlds
 	GEditor->GetEditorWorldContext(true).AddRef(World);
@@ -1305,6 +1305,16 @@ bool SLevelEditor::IsModeActive( FEditorModeID ModeID )
 	return GLevelEditorModeTools().IsModeActive( ModeID );
 }
 
+void SLevelEditor::EditorModeCommandsChanged()
+{
+	if (FLevelEditorModesCommands::IsRegistered())
+	{
+		FLevelEditorModesCommands::Unregister();
+	}
+
+	RefreshEditorModeCommands();
+}
+
 void SLevelEditor::RefreshEditorModeCommands()
 {
 	FLevelEditorModule& LevelEditorModule = FModuleManager::GetModuleChecked<FLevelEditorModule>( "LevelEditor" );
diff --git a/Engine/Source/Editor/LevelEditor/Private/SLevelEditor.h b/Engine/Source/Editor/LevelEditor/Private/SLevelEditor.h
--- a/Engine/Source/Editor/LevelEditor/Private/SLevelEditor.h
+++ b/Engine/Source/Editor/LevelEditor/Private/SLevelEditor.h
@@ -155,6 +155,9 @@ private:
 	/** Rebuilds the command list for spawning editor modes, this is done when new modes are registered. */
 	void RefreshEditorModeCommands();
 
+	/** Editor mode has been added or removed, clears cached command list so it will be rebuilt */
+	void EditorModeCommandsChanged();
+
 	/** Gets the tabId mapping to an editor mode */
 	static FName GetEditorModeTabId( FEditorModeID ModeID );
 
