Disabled Preview Mode button dropdown menu only during the PIE session

This fixes the issue where the dropdown worked fine when a preview mode was selected, but if there wasn't a preview mode selected, the dropdown button was disabled. The user could never select the preview mode.

diff --git a/Engine/Source/Editor/LevelEditor/Private/LevelEditorActions.cpp b/Engine/Source/Editor/LevelEditor/Private/LevelEditorActions.cpp
--- a/Engine/Source/Editor/LevelEditor/Private/LevelEditorActions.cpp
+++ b/Engine/Source/Editor/LevelEditor/Private/LevelEditorActions.cpp
@@ -586,6 +586,16 @@ bool FLevelEditorActionCallbacks::IsFeatureLevelPreviewActive()
 	return GEditor->IsFeatureLevelPreviewEnabled() && GEditor->IsFeatureLevelPreviewActive();
 }
 
+bool FLevelEditorActionCallbacks::IsFeatureLevelPreviewDropdownEnabled()
+{
+	if (GEditor->PlayWorld || GUnrealEd->bIsSimulatingInEditor)
+	{
+		return false;
+	}
+
+	return true;
+}
+
 void FLevelEditorActionCallbacks::SetPreviewPlatform(FName MaterialQualityPlatform, ERHIFeatureLevel::Type PreviewFeatureLevel)
 {
 	GEditor->SetPreviewPlatform(MaterialQualityPlatform, PreviewFeatureLevel);
diff --git a/Engine/Source/Editor/LevelEditor/Private/LevelEditorToolBar.cpp b/Engine/Source/Editor/LevelEditor/Private/LevelEditorToolBar.cpp
--- a/Engine/Source/Editor/LevelEditor/Private/LevelEditorToolBar.cpp
+++ b/Engine/Source/Editor/LevelEditor/Private/LevelEditorToolBar.cpp
@@ -1380,7 +1380,7 @@ TSharedRef< SWidget > FLevelEditorToolBar::MakeLevelEditorToolBar( const TShared
 			);
 
 		FUIAction PreviewModeMenuCanExecute;
-		PreviewModeMenuCanExecute.CanExecuteAction = FIsActionChecked::CreateStatic(&FLevelEditorActionCallbacks::IsFeatureLevelPreviewEnabled);
+		PreviewModeMenuCanExecute.CanExecuteAction = FIsActionChecked::CreateStatic(&FLevelEditorActionCallbacks::IsFeatureLevelPreviewDropdownEnabled);
 
 		ToolbarBuilder.AddComboButton(
 			PreviewModeMenuCanExecute,
diff --git a/Engine/Source/Editor/LevelEditor/Public/LevelEditorActions.h b/Engine/Source/Editor/LevelEditor/Public/LevelEditorActions.h
--- a/Engine/Source/Editor/LevelEditor/Public/LevelEditorActions.h
+++ b/Engine/Source/Editor/LevelEditor/Public/LevelEditorActions.h
@@ -808,6 +808,7 @@ public:
 	static void ToggleFeatureLevelPreview();
 	static bool IsFeatureLevelPreviewEnabled();
 	static bool IsFeatureLevelPreviewActive();
+	static bool IsFeatureLevelPreviewDropdownEnabled();
 	static void SetPreviewPlatform(FName MaterialQualityPlatform,ERHIFeatureLevel::Type PreviewFeatureLevel);
 	static bool IsPreviewPlatformChecked(FName MaterialQualityPlatform, ERHIFeatureLevel::Type PreviewFeatureLevel);
 	static void SetFeatureLevelPreview(ERHIFeatureLevel::Type InFeatureLevel);
