Restored the Settings->Preview Rendering Level menu and made it the access point to the Preview Mode button

diff --git a/Engine/Source/Editor/LevelEditor/Private/LevelEditor.cpp b/Engine/Source/Editor/LevelEditor/Private/LevelEditor.cpp
--- a/Engine/Source/Editor/LevelEditor/Private/LevelEditor.cpp
+++ b/Engine/Source/Editor/LevelEditor/Private/LevelEditor.cpp
@@ -1715,7 +1715,8 @@ void FLevelEditorModule::BindGlobalLevelEditorCommands()
 		Commands.ToggleFeatureLevelPreview,
 		FExecuteAction::CreateStatic(&FLevelEditorActionCallbacks::ToggleFeatureLevelPreview),
 		FIsActionChecked::CreateStatic(&FLevelEditorActionCallbacks::IsFeatureLevelPreviewEnabled),
-		FIsActionChecked::CreateStatic(&FLevelEditorActionCallbacks::IsFeatureLevelPreviewActive));
+		FIsActionChecked::CreateStatic(&FLevelEditorActionCallbacks::IsFeatureLevelPreviewActive),
+		FIsActionButtonVisible::CreateStatic(FLevelEditorActionCallbacks::IsPreviewModeButtonVisible));
 
 	ActionList.MapAction(
 		Commands.PreviewPlatformOverride_DefaultES2,
diff --git a/Engine/Source/Editor/LevelEditor/Private/LevelEditorActions.cpp b/Engine/Source/Editor/LevelEditor/Private/LevelEditorActions.cpp
--- a/Engine/Source/Editor/LevelEditor/Private/LevelEditorActions.cpp
+++ b/Engine/Source/Editor/LevelEditor/Private/LevelEditorActions.cpp
@@ -578,22 +578,25 @@ bool FLevelEditorActionCallbacks::IsFeatureLevelPreviewEnabled()
 	{
 		return false;
 	}
+	if (GEditor->PreviewFeatureLevel == ERHIFeatureLevel::SM5)
+	{
+		return true;
+	}
 	return GEditor->IsFeatureLevelPreviewEnabled();
 }
 
 bool FLevelEditorActionCallbacks::IsFeatureLevelPreviewActive()
 {
-	return GEditor->IsFeatureLevelPreviewEnabled() && GEditor->IsFeatureLevelPreviewActive();
-}
-
-bool FLevelEditorActionCallbacks::IsFeatureLevelPreviewDropdownEnabled()
-{
-	if (GEditor->PlayWorld || GUnrealEd->bIsSimulatingInEditor)
+	if (GEditor->PreviewFeatureLevel == ERHIFeatureLevel::SM5)
 	{
 		return false;
 	}
+	return GEditor->IsFeatureLevelPreviewEnabled() && GEditor->IsFeatureLevelPreviewActive();
+}
 
-	return true;
+bool FLevelEditorActionCallbacks::IsPreviewModeButtonVisible()
+{
+	return GEditor->PreviewFeatureLevel != ERHIFeatureLevel::SM5;
 }
 
 void FLevelEditorActionCallbacks::SetPreviewPlatform(FName MaterialQualityPlatform, ERHIFeatureLevel::Type PreviewFeatureLevel)
@@ -3310,13 +3313,13 @@ void FLevelEditorCommands::RegisterCommands()
 
 	UI_COMMAND(ToggleFeatureLevelPreview, "Preview Platform", "Preview Platform", EUserInterfaceActionType::ToggleButton, FInputChord());
 
-	UI_COMMAND(PreviewPlatformOverride_AndroidGLES2, "Preview as Android ES2", "Mobile preview using Android's quality settings.", EUserInterfaceActionType::RadioButton, FInputChord());
-	UI_COMMAND(PreviewPlatformOverride_DefaultES2, "Preview as HTML5", "HTML5 preview.", EUserInterfaceActionType::RadioButton, FInputChord());
+	UI_COMMAND(PreviewPlatformOverride_AndroidGLES2, "Android ES2", "Mobile preview using Android's quality settings.", EUserInterfaceActionType::RadioButton, FInputChord());
+	UI_COMMAND(PreviewPlatformOverride_DefaultES2, "HTML5", "HTML5 preview.", EUserInterfaceActionType::RadioButton, FInputChord());
 
 	UI_COMMAND(PreviewPlatformOverride_DefaultES31, "Default High-End Mobile", "Use default mobile settings (no quality overrides).", EUserInterfaceActionType::RadioButton, FInputChord());
-	UI_COMMAND(PreviewPlatformOverride_AndroidGLES31, "Preview as Android ES31", "Mobile preview using Android ES3.1 quality settings.", EUserInterfaceActionType::RadioButton, FInputChord());
-	UI_COMMAND(PreviewPlatformOverride_AndroidVulkanES31, "Preview as Android Vulkan", "Mobile preview using Android Vulkan quality settings.", EUserInterfaceActionType::RadioButton, FInputChord());
-	UI_COMMAND(PreviewPlatformOverride_IOSMetalES31, "Preview as iOS", "Mobile preview using iOS material quality settings.", EUserInterfaceActionType::RadioButton, FInputChord());
+	UI_COMMAND(PreviewPlatformOverride_AndroidGLES31, "Android ES 3.1", "Mobile preview using Android ES3.1 quality settings.", EUserInterfaceActionType::RadioButton, FInputChord());
+	UI_COMMAND(PreviewPlatformOverride_AndroidVulkanES31, "Android Vulkan", "Mobile preview using Android Vulkan quality settings.", EUserInterfaceActionType::RadioButton, FInputChord());
+	UI_COMMAND(PreviewPlatformOverride_IOSMetalES31, "iOS", "Mobile preview using iOS material quality settings.", EUserInterfaceActionType::RadioButton, FInputChord());
 
 
 	UI_COMMAND( ConnectToSourceControl, "Connect to Source Control...", "Opens a dialog to connect to source control.", EUserInterfaceActionType::Button, FInputChord());
@@ -3332,8 +3335,8 @@ void FLevelEditorCommands::RegisterCommands()
 	{
 		NSLOCTEXT("LevelEditorCommands", "FeatureLevelPreviewType_ES2", "Mobile / HTML5"),
 		NSLOCTEXT("LevelEditorCommands", "FeatureLevelPreviewType_ES31", "High-End Mobile"),
-		NSLOCTEXT("LevelEditorCommands", "FeatureLevelPreviewType_SM4", "Preview as SM4"),
-		NSLOCTEXT("LevelEditorCommands", "FeatureLevelPreviewType_SM5", "No Preview Device. (View as SM5)"),
+		NSLOCTEXT("LevelEditorCommands", "FeatureLevelPreviewType_SM4", "Shader Model 4"),
+		NSLOCTEXT("LevelEditorCommands", "FeatureLevelPreviewType_SM5", "Shader Model 5"),
 	};
 
 	static const FText FeatureLevelToolTips[ERHIFeatureLevel::Num] = 
diff --git a/Engine/Source/Editor/LevelEditor/Private/LevelEditorToolBar.cpp b/Engine/Source/Editor/LevelEditor/Private/LevelEditorToolBar.cpp
--- a/Engine/Source/Editor/LevelEditor/Private/LevelEditorToolBar.cpp
+++ b/Engine/Source/Editor/LevelEditor/Private/LevelEditorToolBar.cpp
@@ -1277,9 +1277,7 @@ TSharedRef< SWidget > FLevelEditorToolBar::MakeLevelEditorToolBar( const TShared
 			);
 
 	}
-	ToolbarBuilder.EndSection();
 
-	ToolbarBuilder.BeginSection(NAME_None);
 	{
 		struct FPreviewModeFunctionality
 		{
@@ -1296,22 +1294,16 @@ TSharedRef< SWidget > FLevelEditorToolBar::MakeLevelEditorToolBar( const TShared
 
 				EShaderPlatform MaxRHIFeatureLevelPlatform = GetFeatureLevelShaderPlatform(GMaxRHIFeatureLevel);
 
-				if (GEditor->PreviewFeatureLevel == GMaxRHIFeatureLevel)
-				{
-                    const FText& RenderingAsPlatformName = GetFriendlyShaderPlatformName(MaxRHIFeatureLevelPlatform);
-                    return FText::Format(LOCTEXT("PreviewModeRenderingAs", "Viewing as {0}."), RenderingAsPlatformName);
-				}
-				else
 				{
                     const FText& RenderingAsPlatformName = GetFriendlyShaderPlatformName(GWorld->FeatureLevel == GMaxRHIFeatureLevel ? MaxRHIFeatureLevelPlatform : PreviewShaderPlatform);
                     const FText& SwitchToPlatformName = GetFriendlyShaderPlatformName(GWorld->FeatureLevel == GMaxRHIFeatureLevel ? PreviewShaderPlatform : MaxRHIFeatureLevelPlatform);
                     if (GWorld->FeatureLevel == GMaxRHIFeatureLevel)
                     {
-                        return FText::Format(LOCTEXT("PreviewModeViewingAsSwitchTo", "Viewing as {0}. Click to preview as {1}."), RenderingAsPlatformName, SwitchToPlatformName);
+                        return FText::Format(LOCTEXT("PreviewModeViewingAsSwitchTo", "Viewing {0}. Click to preview {1}."), RenderingAsPlatformName, SwitchToPlatformName);
                     }
                     else
                     {
-                        return FText::Format(LOCTEXT("PreviewModePreviewingAsSwitchTo", "Previewing as {0}. Click to view as {1}."), RenderingAsPlatformName, SwitchToPlatformName);
+                        return FText::Format(LOCTEXT("PreviewModePreviewingAsSwitchTo", "Previewing {0}. Click to view {1}."), RenderingAsPlatformName, SwitchToPlatformName);
                     }
 				}
 			}
@@ -1340,6 +1332,7 @@ TSharedRef< SWidget > FLevelEditorToolBar::MakeLevelEditorToolBar( const TShared
 					{
 						return FSlateIcon(FEditorStyle::GetStyleSetName(), GEditor->IsFeatureLevelPreviewActive() ? "LevelEditor.PreviewMode.iOS.Enabled" : "LevelEditor.PreviewMode.iOS.Disabled");
 					}
+					case SP_METAL_MACES2:
                 	case SP_OPENGL_ES2_WEBGL:
 					{
 						return FSlateIcon(FEditorStyle::GetStyleSetName(), GEditor->IsFeatureLevelPreviewActive() ? "LevelEditor.PreviewMode.HTML5.Enabled" : "LevelEditor.PreviewMode.HTML5.Disabled");
@@ -1347,13 +1340,9 @@ TSharedRef< SWidget > FLevelEditorToolBar::MakeLevelEditorToolBar( const TShared
 				}
 				switch (GEditor->PreviewFeatureLevel)
 				{
-					case ERHIFeatureLevel::SM5:
-					{
-						return FSlateIcon(FEditorStyle::GetStyleSetName(), GEditor->IsFeatureLevelPreviewActive() ? "LevelEditor.PreviewMode.SM5.Enabled" : "LevelEditor.PreviewMode.SM5.Disabled");
-					}
 					case ERHIFeatureLevel::SM4:
 					{
-						return FSlateIcon(FEditorStyle::GetStyleSetName(), GEditor->IsFeatureLevelPreviewActive() ? "LevelEditor.PreviewMode.SM4.Enabled" : "LevelEditor.PreviewMode.SM4.Enabled");
+						return FSlateIcon(FEditorStyle::GetStyleSetName(), GEditor->IsFeatureLevelPreviewActive() ? "LevelEditor.PreviewMode.SM4.Enabled" : "LevelEditor.PreviewMode.SM4.Disabled");
 					}
 					case ERHIFeatureLevel::ES2:
 					{
@@ -1378,18 +1367,6 @@ TSharedRef< SWidget > FLevelEditorToolBar::MakeLevelEditorToolBar( const TShared
         	TAttribute<FText>::Create(&FPreviewModeFunctionality::GetPreviewModeTooltip),
         	TAttribute<FSlateIcon>::Create(&FPreviewModeFunctionality::GetPreviewModeIcon)
 			);
-
-		FUIAction PreviewModeMenuCanExecute;
-		PreviewModeMenuCanExecute.CanExecuteAction = FIsActionChecked::CreateStatic(&FLevelEditorActionCallbacks::IsFeatureLevelPreviewDropdownEnabled);
-
-		ToolbarBuilder.AddComboButton(
-			PreviewModeMenuCanExecute,
-			FOnGetContent::CreateStatic(&FLevelEditorToolBar::GeneratePreviewModeMenu, InCommandList),
-			LOCTEXT("PreviewModeCombo", "Preview Mode"),
-			LOCTEXT("PreviewModeCombo_ToolTip", "Preview Mode settings"),
-			FSlateIcon(FEditorStyle::GetStyleSetName(), "LevelEditor.GameSettings"),
-            true
-			);
 	}
 	ToolbarBuilder.EndSection();
 
@@ -1788,6 +1765,59 @@ static void MakeMaterialQualityLevelMenu( FMenuBuilder& MenuBuilder )
 	MenuBuilder.EndSection();
 }
 
+static void MakeShaderModelPreviewMenu( FMenuBuilder& MenuBuilder )
+{
+#define LOCTEXT_NAMESPACE "LevelToolBarViewMenu"
+
+	MenuBuilder.BeginSection("EditorPreviewMode", LOCTEXT("EditorPreviewModeDevices", "Preview Devices"));
+
+	for (int32 i = GMaxRHIFeatureLevel; i >= 0; --i)
+	{
+		switch (i)
+		{
+			case ERHIFeatureLevel::ES2:
+				MenuBuilder.AddMenuEntry(FLevelEditorCommands::Get().PreviewPlatformOverride_AndroidGLES2);
+				MenuBuilder.AddMenuEntry(FLevelEditorCommands::Get().PreviewPlatformOverride_DefaultES2);
+				break;
+
+			case ERHIFeatureLevel::ES3_1:
+			{
+				//MenuBuilder.AddMenuEntry(FLevelEditorCommands::Get().PreviewPlatformOverride_DefaultES31);
+
+				bool bAndroidBuildForES31 = false;
+				GConfig->GetBool(TEXT("/Script/AndroidRuntimeSettings.AndroidRuntimeSettings"), TEXT("bBuildForES31"), bAndroidBuildForES31, GEngineIni);
+				if(bAndroidBuildForES31)
+				{
+					MenuBuilder.AddMenuEntry(FLevelEditorCommands::Get().PreviewPlatformOverride_AndroidGLES31);
+				}
+
+				bool bAndroidSupportsVulkan = false;
+				GConfig->GetBool(TEXT("/Script/AndroidRuntimeSettings.AndroidRuntimeSettings"), TEXT("bSupportsVulkan"), bAndroidSupportsVulkan, GEngineIni);
+				if(bAndroidSupportsVulkan)
+				{
+					MenuBuilder.AddMenuEntry(FLevelEditorCommands::Get().PreviewPlatformOverride_AndroidVulkanES31);
+				}
+
+				bool bIOSSupportsMetal = false;
+				GConfig->GetBool(TEXT("/Script/IOSRuntimeSettings.IOSRuntimeSettings"), TEXT("bSupportsMetal"), bIOSSupportsMetal, GEngineIni);
+				if(bIOSSupportsMetal)
+				{
+					MenuBuilder.AddMenuEntry(FLevelEditorCommands::Get().PreviewPlatformOverride_IOSMetalES31);
+				}
+
+				break;
+			}
+
+			default:
+				MenuBuilder.AddMenuEntry(FLevelEditorCommands::Get().FeatureLevelPreview[i]);
+		}
+	}
+
+    MenuBuilder.EndSection();
+
+#undef LOCTEXT_NAMESPACE
+}
+
 static void MakeScalabilityMenu( FMenuBuilder& MenuBuilder )
 {
 	MenuBuilder.AddWidget(SNew(SScalabilitySettings), FText(), true);
@@ -1880,6 +1910,11 @@ TSharedRef< SWidget > FLevelEditorToolBar::GenerateQuickSettingsMenu( TSharedRef
 			LOCTEXT( "MaterialQualityLevelSubMenu", "Material Quality Level" ),
 			LOCTEXT( "MaterialQualityLevelSubMenu_ToolTip", "Sets the value of the CVar \"r.MaterialQualityLevel\" (low=0, high=1, medium=2). This affects materials via the QualitySwitch material expression." ),
 			FNewMenuDelegate::CreateStatic( &MakeMaterialQualityLevelMenu ) );
+
+		MenuBuilder.AddSubMenu(
+			LOCTEXT("FeatureLevelPreviewSubMenu", "Preview Rendering Level"),
+			LOCTEXT("FeatureLevelPreviewSubMenu_ToolTip", "Sets the rendering level used by the main editor"),
+			FNewMenuDelegate::CreateStatic(&MakeShaderModelPreviewMenu));
 	}
 	MenuBuilder.EndSection();
 
@@ -1940,79 +1975,6 @@ TSharedRef< SWidget > FLevelEditorToolBar::GenerateQuickSettingsMenu( TSharedRef
 }
 
 
-TSharedRef< SWidget > FLevelEditorToolBar::GeneratePreviewModeMenu( TSharedRef<FUICommandList> InCommandList )
-{
-#define LOCTEXT_NAMESPACE "LevelToolBarViewMenu"
-
-	// Get all menu extenders for this context menu from the level editor module
-	FLevelEditorModule& LevelEditorModule = FModuleManager::GetModuleChecked<FLevelEditorModule>( TEXT("LevelEditor") );
-	TArray<FLevelEditorModule::FLevelEditorMenuExtender> MenuExtenderDelegates = LevelEditorModule.GetAllLevelEditorToolbarViewMenuExtenders();
-
-	TArray<TSharedPtr<FExtender>> Extenders;
-	for (int32 i = 0; i < MenuExtenderDelegates.Num(); ++i)
-	{
-		if (MenuExtenderDelegates[i].IsBound())
-		{
-			Extenders.Add(MenuExtenderDelegates[i].Execute(InCommandList));
-		}
-	}
-	TSharedPtr<FExtender> MenuExtender = FExtender::Combine(Extenders);
-
-	const bool bShouldCloseWindowAfterMenuSelection = true;
-	FMenuBuilder MenuBuilder(bShouldCloseWindowAfterMenuSelection, InCommandList, MenuExtender);
-
-	MenuBuilder.BeginSection("EditorPreviewMode", LOCTEXT("EditorPreviewModeDevices", "Preview Devices"));
-
-	for (int32 i = GMaxRHIFeatureLevel; i >= 0; --i)
-	{
-		switch (i)
-		{
-			case ERHIFeatureLevel::ES2:
-				MenuBuilder.AddMenuEntry(FLevelEditorCommands::Get().PreviewPlatformOverride_AndroidGLES2);
-				MenuBuilder.AddMenuEntry(FLevelEditorCommands::Get().PreviewPlatformOverride_DefaultES2);
-				break;
-
-			case ERHIFeatureLevel::ES3_1:
-			{
-				//MenuBuilder.AddMenuEntry(FLevelEditorCommands::Get().PreviewPlatformOverride_DefaultES31);
-
-				bool bAndroidBuildForES31 = false;
-				GConfig->GetBool(TEXT("/Script/AndroidRuntimeSettings.AndroidRuntimeSettings"), TEXT("bBuildForES31"), bAndroidBuildForES31, GEngineIni);
-				if(bAndroidBuildForES31)
-				{
-					MenuBuilder.AddMenuEntry(FLevelEditorCommands::Get().PreviewPlatformOverride_AndroidGLES31);
-				}
-
-				bool bAndroidSupportsVulkan = false;
-				GConfig->GetBool(TEXT("/Script/AndroidRuntimeSettings.AndroidRuntimeSettings"), TEXT("bSupportsVulkan"), bAndroidSupportsVulkan, GEngineIni);
-				if(bAndroidSupportsVulkan)
-				{
-					MenuBuilder.AddMenuEntry(FLevelEditorCommands::Get().PreviewPlatformOverride_AndroidVulkanES31);
-				}
-
-				bool bIOSSupportsMetal = false;
-				GConfig->GetBool(TEXT("/Script/IOSRuntimeSettings.IOSRuntimeSettings"), TEXT("bSupportsMetal"), bIOSSupportsMetal, GEngineIni);
-				if(bIOSSupportsMetal)
-				{
-					MenuBuilder.AddMenuEntry(FLevelEditorCommands::Get().PreviewPlatformOverride_IOSMetalES31);
-				}
-
-				break;
-			}
-
-			default:
-				MenuBuilder.AddMenuEntry(FLevelEditorCommands::Get().FeatureLevelPreview[i]);
-		}
-	}
-
-    MenuBuilder.EndSection();
-
-#undef LOCTEXT_NAMESPACE
-
-	return MenuBuilder.MakeWidget();
-}
-
-
 TSharedRef< SWidget > FLevelEditorToolBar::GenerateSourceControlMenu(TSharedRef<FUICommandList> InCommandList)
 {
 #define LOCTEXT_NAMESPACE "LevelToolBarSourceControlMenu"
diff --git a/Engine/Source/Editor/LevelEditor/Private/LevelEditorToolBar.h b/Engine/Source/Editor/LevelEditor/Private/LevelEditorToolBar.h
--- a/Engine/Source/Editor/LevelEditor/Private/LevelEditorToolBar.h
+++ b/Engine/Source/Editor/LevelEditor/Private/LevelEditorToolBar.h
@@ -47,13 +47,6 @@ protected:
 	 */
 	static TSharedRef< SWidget > GenerateQuickSettingsMenu( TSharedRef<FUICommandList> InCommandList );
 
-	/**
-	 * Generates menu content for the preview mode combo button drop down menu
-	 *
-	 * @return	Menu content widget
-	 */
-	static TSharedRef< SWidget > GeneratePreviewModeMenu( TSharedRef<FUICommandList> InCommandList );
-
 	/**
 	 * Generates menu content for the source control combo button drop down menu
 	 *
diff --git a/Engine/Source/Editor/LevelEditor/Public/LevelEditorActions.h b/Engine/Source/Editor/LevelEditor/Public/LevelEditorActions.h
--- a/Engine/Source/Editor/LevelEditor/Public/LevelEditorActions.h
+++ b/Engine/Source/Editor/LevelEditor/Public/LevelEditorActions.h
@@ -808,7 +808,7 @@ public:
 	static void ToggleFeatureLevelPreview();
 	static bool IsFeatureLevelPreviewEnabled();
 	static bool IsFeatureLevelPreviewActive();
-	static bool IsFeatureLevelPreviewDropdownEnabled();
+	static bool IsPreviewModeButtonVisible();
 	static void SetPreviewPlatform(FName MaterialQualityPlatform,ERHIFeatureLevel::Type PreviewFeatureLevel);
 	static bool IsPreviewPlatformChecked(FName MaterialQualityPlatform, ERHIFeatureLevel::Type PreviewFeatureLevel);
 	static void SetFeatureLevelPreview(ERHIFeatureLevel::Type InFeatureLevel);
