Adding support to create subclasses of ACameraActor with the Create Camera Here viewport action. Added support for UCineCameraComponent to set Focal length from FOV when calling SetFieldOfView.

diff --git a/Engine/Source/Editor/LevelEditor/Private/LevelViewportActions.cpp b/Engine/Source/Editor/LevelEditor/Private/LevelViewportActions.cpp
--- a/Engine/Source/Editor/LevelEditor/Private/LevelViewportActions.cpp
+++ b/Engine/Source/Editor/LevelEditor/Private/LevelViewportActions.cpp
@@ -1,4 +1,4 @@
-﻿// Copyright 1998-2019 Epic Games, Inc. All Rights Reserved.
+// Copyright 1998-2019 Epic Games, Inc. All Rights Reserved.
 
 #include "LevelViewportActions.h"
 #include "ShowFlags.h"
@@ -33,7 +33,6 @@ void FLevelViewportCommands::RegisterCommands()
 	UI_COMMAND( ToggleGameView, "Game View", "Toggles game view.  Game view shows the scene as it appears in game", EUserInterfaceActionType::ToggleButton, FInputChord( EKeys::G ) );
 	UI_COMMAND( ToggleImmersive, "Immersive Mode", "Switches this viewport between immersive mode and regular mode", EUserInterfaceActionType::ToggleButton, FInputChord( EKeys::F11 ) );
 
-	UI_COMMAND( CreateCamera, "Create Camera Here", "Creates a new camera actor at the current location of this viewport's camera", EUserInterfaceActionType::Button, FInputChord() );
 	UI_COMMAND( HighResScreenshot, "High Resolution Screenshot...", "Opens the control panel for high resolution screenshots", EUserInterfaceActionType::Button, FInputChord() );
 	
 	UI_COMMAND( UseDefaultShowFlags, "Use Defaults", "Resets all show flags to default", EUserInterfaceActionType::Button, FInputChord() );
diff --git a/Engine/Source/Editor/LevelEditor/Private/SLevelViewport.cpp b/Engine/Source/Editor/LevelEditor/Private/SLevelViewport.cpp
--- a/Engine/Source/Editor/LevelEditor/Private/SLevelViewport.cpp
+++ b/Engine/Source/Editor/LevelEditor/Private/SLevelViewport.cpp
@@ -67,6 +67,8 @@
 #include "FoliageType.h"
 #include "IVREditorModule.h"
 #include "ShowFlagMenuCommands.h"
+#include "AssetRegistryModule.h"
+#include "IAssetRegistry.h"
 #include "BufferVisualizationMenuCommands.h"
 
 static const FName LevelEditorName("LevelEditor");
@@ -1168,13 +1170,51 @@ void SLevelViewport::BindOptionCommands( FUICommandList& OutCommandList )
 		FIsActionChecked::CreateSP( this, &SLevelViewport::AllowsCinematicPreview )
 		);
 
+	IAssetRegistry & AssetRegistry = FModuleManager::LoadModuleChecked<FAssetRegistryModule>("AssetRegistry").Get();
+	TArray<FName> ClassNames;
+	TSet<FName> DerivedClassNames;
+	ClassNames.Add(ACameraActor::StaticClass()->GetFName());
+	AssetRegistry.GetDerivedClassNames(ClassNames, TSet<FName>(), DerivedClassNames);
 
-	OutCommandList.MapAction(
-		ViewportActions.CreateCamera,
-		FExecuteAction::CreateSP( this, &SLevelViewport::OnCreateCameraActor ),
-		FCanExecuteAction(),
-		FCanExecuteAction::CreateSP( this, &SLevelViewport::IsPerspectiveViewport ) 
-		);
+	for (FName Name : DerivedClassNames)
+	{
+		// Ignore generated types that cannot be spawned
+		if (Name.ToString().StartsWith("SKEL_") || Name.ToString().StartsWith("REINST_"))
+		{
+			continue;
+		}
+		
+		// Remove _C from display names for blueprint classes
+		FName DisplayName = Name;
+		if (Name.ToString().EndsWith("_C"))
+		{
+			DisplayName = FName(*DisplayName.ToString().LeftChop(2));
+		}
+		
+		// Look for existing UI Command info so one isn't created for every viewport
+		TSharedPtr<FUICommandInfo> * FoundCamera = FLevelViewportCommands::Get().CreateCameras.FindByPredicate([Name](TSharedPtr<FUICommandInfo> Camera) { return Camera->GetCommandName() == Name; });
+		
+		if (FoundCamera)
+		{
+			OutCommandList.MapAction(
+				*FoundCamera,
+				FExecuteAction::CreateSP(this, &SLevelViewport::OnCreateCameraActor, FindObject<UClass>(ANY_PACKAGE, *Name.ToString()))
+			);
+		}
+		else
+		{
+			// If command info isn't found, create a new one
+			TSharedRef<FUICommandInfo> NewCamera = FUICommandInfoDecl(FLevelViewportCommands::Get().AsShared(), Name, FText::FromName(DisplayName), FText::Format(LOCTEXT("SpawnCamerasTooltip", "Spawn Camera here of type {0}"), FText::FromName(Name))).UserInterfaceType(EUserInterfaceActionType::Button).DefaultChord(FInputChord());
+			
+			OutCommandList.MapAction(
+				NewCamera,
+				FExecuteAction::CreateSP(this, &SLevelViewport::OnCreateCameraActor, FindObject<UClass>(ANY_PACKAGE, *Name.ToString()))
+			);
+			
+			FLevelViewportCommands::Get().CreateCameras.Add(NewCamera);
+		}
+
+	}
 
 	OutCommandList.MapAction(
 		ViewportActions.HighResScreenshot,
@@ -1660,7 +1700,7 @@ bool SLevelViewport::IsImmersive() const
 	return false;
 }
 
-void SLevelViewport::OnCreateCameraActor()
+void SLevelViewport::OnCreateCameraActor(UClass* InClass)
 {		
 	// Find the perspective viewport we were using
 	FViewport* pViewPort = GEditor->GetActiveViewport();
@@ -1684,10 +1724,10 @@ void SLevelViewport::OnCreateCameraActor()
 	const FScopedTransaction Transaction(NSLOCTEXT("LevelViewport", "CreateCameraHere", "Create Camera Here"));
 
 	// Set new camera to match viewport
-	ACameraActor* pNewCamera = ViewportClient->GetWorld()->SpawnActor<ACameraActor>();
+	ACameraActor* pNewCamera = Cast<ACameraActor>(ViewportClient->GetWorld()->SpawnActor(InClass));
 	pNewCamera->SetActorLocation( ViewportClient->GetViewLocation(), false );
 	pNewCamera->SetActorRotation( ViewportClient->GetViewRotation() );
-	pNewCamera->GetCameraComponent()->FieldOfView = ViewportClient->ViewFOV;
+	pNewCamera->GetCameraComponent()->SetFieldOfView( ViewportClient->ViewFOV );
 
 	// Deselect any currently selected actors
 	GUnrealEd->SelectNone( true, true );
diff --git a/Engine/Source/Editor/LevelEditor/Private/SLevelViewportToolBar.cpp b/Engine/Source/Editor/LevelEditor/Private/SLevelViewportToolBar.cpp
--- a/Engine/Source/Editor/LevelEditor/Private/SLevelViewportToolBar.cpp
+++ b/Engine/Source/Editor/LevelEditor/Private/SLevelViewportToolBar.cpp
@@ -1,4 +1,4 @@
-﻿// Copyright 1998-2019 Epic Games, Inc. All Rights Reserved.
+// Copyright 1998-2019 Epic Games, Inc. All Rights Reserved.
 
 #include "SLevelViewportToolBar.h"
 #include "Framework/Commands/UIAction.h"
@@ -654,8 +654,13 @@ TSharedRef<SWidget> SLevelViewportToolBar::GenerateOptionsMenu() const
 					);
 			}
 			OptionsMenuBuilder.EndSection();
-
-			OptionsMenuBuilder.AddMenuEntry( LevelViewportActions.CreateCamera );
+		
+			OptionsMenuBuilder.AddSubMenu(
+				LOCTEXT("CameraSubMeun", "Create Camera Here"),
+				LOCTEXT("CameraSubMenu_ToolTip", "Select a camera type to create at current viewport's location"),
+				FNewMenuDelegate::CreateSP(this, &SLevelViewportToolBar::GenerateCameraSpawnMenu)
+			);
+			
 		}
 
 		OptionsMenuBuilder.AddMenuEntry( LevelViewportActions.HighResScreenshot );
@@ -808,6 +813,16 @@ void SLevelViewportToolBar::GenerateViewportTypeMenu( FMenuBuilder& Builder ) co
 	});
 }
 
+void SLevelViewportToolBar::GenerateCameraSpawnMenu(FMenuBuilder & Builder) const
+{
+	const FLevelViewportCommands& Actions = FLevelViewportCommands::Get();
+
+	for (TSharedPtr<FUICommandInfo> Camera : Actions.CreateCameras)
+	{
+		Builder.AddMenuEntry(Camera);
+	}
+}
+
 TSharedRef<SWidget> SLevelViewportToolBar::GenerateCameraMenu() const
 {
 	Viewport.Pin()->OnFloatingButtonClicked();
diff --git a/Engine/Source/Editor/LevelEditor/Private/SLevelViewportToolBar.h b/Engine/Source/Editor/LevelEditor/Private/SLevelViewportToolBar.h
--- a/Engine/Source/Editor/LevelEditor/Private/SLevelViewportToolBar.h
+++ b/Engine/Source/Editor/LevelEditor/Private/SLevelViewportToolBar.h
@@ -128,6 +128,13 @@ private:
 	 */
 	void GenerateViewportTypeMenu( FMenuBuilder& Builder ) const;
 
+	/**
+	 * Generates menu entries for spawning cameras at the current viewport
+	 *
+	 * @param Builder	The menu builder to add menu entries to
+	 */
+	void GenerateCameraSpawnMenu(FMenuBuilder& Builder) const;
+
 	/**
 	 * Generates the toolbar view menu content 
 	 *
diff --git a/Engine/Source/Editor/LevelEditor/Public/LevelViewportActions.h b/Engine/Source/Editor/LevelEditor/Public/LevelViewportActions.h
--- a/Engine/Source/Editor/LevelEditor/Public/LevelViewportActions.h
+++ b/Engine/Source/Editor/LevelEditor/Public/LevelViewportActions.h
@@ -93,8 +93,8 @@ PRAGMA_ENABLE_DEPRECATION_WARNINGS
 	/** Toggles maximize mode in the viewport */
 	TSharedPtr< FUICommandInfo > ToggleMaximize;
 
-	/** Creates a cameraActor at current perspective viewport position */
-	TSharedPtr< FUICommandInfo > CreateCamera;
+	/** Creates a subclass of ACameraActor at current perspective viewport position */
+	TArray< TSharedPtr< FUICommandInfo > > CreateCameras;
 
 	/** Opens the control panel for high resolution screenshots */
 	TSharedPtr< FUICommandInfo > HighResScreenshot;
diff --git a/Engine/Source/Editor/LevelEditor/Public/SLevelViewport.h b/Engine/Source/Editor/LevelEditor/Public/SLevelViewport.h
--- a/Engine/Source/Editor/LevelEditor/Public/SLevelViewport.h
+++ b/Engine/Source/Editor/LevelEditor/Public/SLevelViewport.h
@@ -434,7 +434,7 @@ private:
 	/**
 	 * Called to create a cameraActor in the currently selected perspective viewport
 	 */
-	void OnCreateCameraActor();
+	void OnCreateCameraActor(UClass *InClass);
 
 	/**
 	 * Called to bring up the screenshot UI
diff --git a/Engine/Source/Runtime/CinematicCamera/Private/CineCameraComponent.cpp b/Engine/Source/Runtime/CinematicCamera/Private/CineCameraComponent.cpp
--- a/Engine/Source/Runtime/CinematicCamera/Private/CineCameraComponent.cpp
+++ b/Engine/Source/Runtime/CinematicCamera/Private/CineCameraComponent.cpp
@@ -178,6 +178,13 @@ void UCineCameraComponent::ResetProxyMeshTransform()
 
 #endif	// WITH_EDITORONLY_DATA
 
+void UCineCameraComponent::SetFieldOfView(float InFieldOfView)
+{
+	Super::SetFieldOfView(InFieldOfView);
+
+	CurrentFocalLength = (FilmbackSettings.SensorWidth / 2.f) / FMath::Tan(FMath::DegreesToRadians(InFieldOfView / 2.f));
+}
+
 float UCineCameraComponent::GetHorizontalFieldOfView() const
 {
 	return (CurrentFocalLength > 0.f)
diff --git a/Engine/Source/Runtime/CinematicCamera/Public/CineCameraComponent.h b/Engine/Source/Runtime/CinematicCamera/Public/CineCameraComponent.h
--- a/Engine/Source/Runtime/CinematicCamera/Public/CineCameraComponent.h
+++ b/Engine/Source/Runtime/CinematicCamera/Public/CineCameraComponent.h
@@ -254,6 +254,9 @@ public:
 	UPROPERTY(VisibleAnywhere, Category = "Current Camera Settings")
 	float CurrentHorizontalFOV;
 #endif
+
+	/** Override setting FOV to manipulate Focal Length. */
+	virtual void SetFieldOfView(float InFieldOfView) override;
 	
 	/** Returns the horizonal FOV of the camera with current settings. */
 	UFUNCTION(BlueprintCallable, Category = "Cine Camera")
diff --git a/Engine/Source/Runtime/Engine/Classes/Camera/CameraComponent.h b/Engine/Source/Runtime/Engine/Classes/Camera/CameraComponent.h
--- a/Engine/Source/Runtime/Engine/Classes/Camera/CameraComponent.h
+++ b/Engine/Source/Runtime/Engine/Classes/Camera/CameraComponent.h
@@ -26,7 +26,7 @@ class ENGINE_API UCameraComponent : public USceneComponent
 	UPROPERTY(Interp, EditAnywhere, BlueprintReadWrite, Category = CameraSettings, meta = (UIMin = "5.0", UIMax = "170", ClampMin = "0.001", ClampMax = "360.0", Units = deg))
 	float FieldOfView;
 	UFUNCTION(BlueprintCallable, Category = Camera)
-	void SetFieldOfView(float InFieldOfView) { FieldOfView = InFieldOfView; }
+	virtual void SetFieldOfView(float InFieldOfView) { FieldOfView = InFieldOfView; }
 
 	/** The desired width (in world units) of the orthographic view (ignored in Perspective mode) */
 	UPROPERTY(Interp, EditAnywhere, BlueprintReadWrite, Category = CameraSettings)
