Prevent favoriting a level when already at max favorites
Increase max favorited levels from 10 to 20

diff --git a/Engine/Source/Editor/LevelEditor/Private/LevelEditorActions.cpp b/Engine/Source/Editor/LevelEditor/Private/LevelEditorActions.cpp
--- a/Engine/Source/Editor/LevelEditor/Private/LevelEditorActions.cpp
+++ b/Engine/Source/Editor/LevelEditor/Private/LevelEditorActions.cpp
@@ -365,7 +365,7 @@ bool FLevelEditorActionCallbacks::ToggleFavorite_CanExecute()
 		const int32 NumFavorites = MRUFavorites.GetNumFavorites();
 		// Disable the favorites button if the map isn't associated to a file yet (new map, never before saved, etc.)
 		const FString PackageName = GetWorld()->GetOutermost()->GetName();
-		return (NumFavorites <= FLevelEditorCommands::Get().OpenFavoriteFileCommands.Num() || MRUFavorites.ContainsFavoritesItem(PackageName));
+		return (NumFavorites < FLevelEditorCommands::Get().OpenFavoriteFileCommands.Num() || MRUFavorites.ContainsFavoritesItem(PackageName));
 	}
 	return false;
 }
@@ -3015,7 +3015,7 @@ void FLevelEditorCommands::RegisterCommands()
 			.DefaultChord( FInputChord() );
 		OpenRecentFileCommands.Add( OpenRecentFile );
 	}
-	for (int32 CurFavoriteIndex = 0; CurFavoriteIndex < FLevelEditorCommands::MaxRecentFiles; ++CurFavoriteIndex)
+	for (int32 CurFavoriteIndex = 0; CurFavoriteIndex < FLevelEditorCommands::MaxFavoriteFiles; ++CurFavoriteIndex)
 	{
 		// NOTE: The actual label and tool-tip will be overridden at runtime when the command is bound to a menu item, however
 		// we still need to set one here so that the key bindings UI can function properly
diff --git a/Engine/Source/Editor/LevelEditor/Public/LevelEditorActions.h b/Engine/Source/Editor/LevelEditor/Public/LevelEditorActions.h
--- a/Engine/Source/Editor/LevelEditor/Public/LevelEditorActions.h
+++ b/Engine/Source/Editor/LevelEditor/Public/LevelEditorActions.h
@@ -59,7 +59,7 @@ public:
 
 	static const int32 MaxRecentFiles = 10;
 	TArray< TSharedPtr< FUICommandInfo > > OpenRecentFileCommands;
-	static const int32 MaxFavoriteFiles = 10;
+	static const int32 MaxFavoriteFiles = 20;
 	TArray< TSharedPtr< FUICommandInfo > > OpenFavoriteFileCommands;
 
 	TSharedPtr< FUICommandInfo > ToggleFavorite;
