Fixed Copy and paste error in SM4 preview options

diff --git a/Engine/Source/Editor/LevelEditor/Private/LevelEditorToolBar.cpp b/Engine/Source/Editor/LevelEditor/Private/LevelEditorToolBar.cpp
--- a/Engine/Source/Editor/LevelEditor/Private/LevelEditorToolBar.cpp
+++ b/Engine/Source/Editor/LevelEditor/Private/LevelEditorToolBar.cpp
@@ -1353,7 +1353,7 @@ TSharedRef< SWidget > FLevelEditorToolBar::MakeLevelEditorToolBar( const TShared
 					}
 					case ERHIFeatureLevel::SM4:
 					{
-						return FSlateIcon(FEditorStyle::GetStyleSetName(), GEditor->IsFeatureLevelPreviewActive() ? "LevelEditor.PreviewMode.SM4.Enabled" : "LevelEditor.PreviewMode.SM4.Enabled");
+						return FSlateIcon(FEditorStyle::GetStyleSetName(), GEditor->IsFeatureLevelPreviewActive() ? "LevelEditor.PreviewMode.SM4.Enabled" : "LevelEditor.PreviewMode.SM4.Disabled");
 					}
 					case ERHIFeatureLevel::ES2:
 					{
