Display of current level now is a button that allows you to change between filtered levels. Also, the level of the currently selected actor is displayed above.
#rnx

diff --git a/Engine/Plugins/MovieScene/LevelSequenceEditor/Source/LevelSequenceEditor/Private/CinematicViewport/SCinematicLevelViewport.cpp b/Engine/Plugins/MovieScene/LevelSequenceEditor/Source/LevelSequenceEditor/Private/CinematicViewport/SCinematicLevelViewport.cpp
--- a/Engine/Plugins/MovieScene/LevelSequenceEditor/Source/LevelSequenceEditor/Private/CinematicViewport/SCinematicLevelViewport.cpp
+++ b/Engine/Plugins/MovieScene/LevelSequenceEditor/Source/LevelSequenceEditor/Private/CinematicViewport/SCinematicLevelViewport.cpp
@@ -125,6 +125,7 @@ class SCinematicPreviewViewport : public SLevelViewport
 public:
 	virtual const FSlateBrush* OnGetViewportBorderBrush() const override { return nullptr; }
 	virtual EVisibility GetCurrentLevelTextVisibility() const override { return EVisibility::Collapsed; }
+	virtual EVisibility GetSelectedActorsCurrentLevelTextVisibility() const override { return EVisibility::Collapsed; }
 	virtual EVisibility GetViewportControlsVisibility() const override { return EVisibility::Collapsed; }
 
 	virtual TSharedPtr<SWidget> MakeViewportToolbar() { return nullptr; }
diff --git a/Engine/Source/Editor/LevelEditor/Private/SLevelViewport.cpp b/Engine/Source/Editor/LevelEditor/Private/SLevelViewport.cpp
--- a/Engine/Source/Editor/LevelEditor/Private/SLevelViewport.cpp
+++ b/Engine/Source/Editor/LevelEditor/Private/SLevelViewport.cpp
@@ -70,6 +70,9 @@
 #include "AssetRegistryModule.h"
 #include "IAssetRegistry.h"
 #include "BufferVisualizationMenuCommands.h"
+#include "EditorLevelUtils.h"
+#include "Engine/LevelStreaming.h"
+#include "Editor/WorldBrowser/Public/WorldBrowserModule.h"
 
 static const FName LevelEditorName("LevelEditor");
 
@@ -98,7 +101,7 @@ public:
 
 
 SLevelViewport::SLevelViewport()
-	: HighResScreenshotDialog( NULL )
+	: HighResScreenshotDialog( nullptr )
 	, ViewTransitionType( EViewTransition::None )
 	, bViewTransitionAnimPending( false )
 	, DeviceProfile("Default")
@@ -292,29 +295,68 @@ void SLevelViewport::ConstructViewportOverlayContent()
 		.Padding(2.0f, 1.0f, 2.0f, 1.0f)
 		[
 			SNew(SHorizontalBox)
-			.Visibility(this, &SLevelViewport::GetCurrentLevelTextVisibility)
+			.Visibility(this, &SLevelViewport::GetSelectedActorsCurrentLevelTextVisibility)
 			// Current level label
 			+ SHorizontalBox::Slot()
 			.AutoWidth()
 			.Padding(2.0f, 1.0f, 2.0f, 1.0f)
 			[
 				SNew(STextBlock)
-				.Text(this, &SLevelViewport::GetCurrentLevelText, true)
+				.Text(this, &SLevelViewport::GetSelectedActorsCurrentLevelText, true)
 				.Font(FEditorStyle::GetFontStyle(TEXT("MenuItem.Font")))
 				.ShadowOffset(FVector2D(1, 1))
 			]
-
 			// Current level
 			+ SHorizontalBox::Slot()
-				.AutoWidth()
-				.Padding(4.0f, 1.0f, 2.0f, 1.0f)
+			.AutoWidth()
+			.Padding(4.0f, 1.0f, 2.0f, 1.0f)
+			[
+				SNew(STextBlock)
+				.Text(this, &SLevelViewport::GetSelectedActorsCurrentLevelText, false)
+				.Font(FEditorStyle::GetFontStyle(TEXT("MenuItem.Font")))
+				.ShadowOffset(FVector2D(1, 1))
+			]
+		]
+		+ SVerticalBox::Slot()
+		.AutoHeight()
+		.Padding(2.0f, 1.0f, 2.0f, 1.0f)
+		[
+			SAssignNew(LevelMenuAnchor, SMenuAnchor)
+			.Placement(MenuPlacement_AboveAnchor)
+			[
+				SNew(SButton)
+				// Allows users to drag with the mouse to select options after opening the menu */
+				.ClickMethod(EButtonClickMethod::MouseDown)
+				.ContentPadding(FMargin(5.0f, 2.0f))
+				.VAlign(VAlign_Center)
+				.ButtonStyle(FEditorStyle::Get(), "EditorViewportToolBar.MenuButton")
+				.OnClicked(this, &SLevelViewport::OnMenuClicked)
 				[
-					SNew(STextBlock)
-					.Text(this, &SLevelViewport::GetCurrentLevelText, false)
-					.Font(FEditorStyle::GetFontStyle(TEXT("MenuItem.Font")))
-					.ColorAndOpacity(FLinearColor(0.4f, 1.0f, 1.0f))
-					.ShadowOffset(FVector2D(1, 1))
+					SNew(SHorizontalBox)
+					.Visibility(this, &SLevelViewport::GetCurrentLevelTextVisibility)
+					// Current level label
+					+ SHorizontalBox::Slot()
+					.AutoWidth()
+					.Padding(2.0f, 1.0f, 2.0f, 1.0f)
+					[
+						SNew(STextBlock)
+						.Text(this, &SLevelViewport::GetCurrentLevelText, true)
+						.Font(FEditorStyle::GetFontStyle(TEXT("MenuItem.Font")))
+						.ShadowOffset(FVector2D(1, 1))
+					]
+					// Current level
+					+ SHorizontalBox::Slot()
+					.AutoWidth()
+					.Padding(4.0f, 1.0f, 2.0f, 1.0f)
+					[
+						SNew(STextBlock)
+						.Text(this, &SLevelViewport::GetCurrentLevelText, false)
+						.Font(FEditorStyle::GetFontStyle(TEXT("MenuItem.Font")))
+						.ShadowOffset(FVector2D(1, 1))
+					]
 				]
+			]
+			.OnGetMenuContent(this, &SLevelViewport::GenerateLevelMenu)
 		]
 	];
 
@@ -323,11 +365,55 @@ void SLevelViewport::ConstructViewportOverlayContent()
 	.VAlign( VAlign_Fill )
 	.HAlign( HAlign_Fill )
 	.Padding( 0 )
-		[
-			SAssignNew(CaptureRegionWidget, SCaptureRegionWidget)
-		];
+	[
+		SAssignNew(CaptureRegionWidget, SCaptureRegionWidget)
+	];
 }
 
+
+TSharedRef<SWidget> SLevelViewport::GenerateLevelMenu() const
+{
+	FWorldBrowserModule& WorldBrowserModule = FModuleManager::LoadModuleChecked<FWorldBrowserModule>("WorldBrowser");
+	// Get all menu extenders for this context menu from the level editor module
+	FLevelEditorModule& LevelEditorModule = FModuleManager::GetModuleChecked<FLevelEditorModule>(TEXT("LevelEditor"));
+	TArray<FLevelEditorModule::FLevelEditorMenuExtender> MenuExtenderDelegates = LevelEditorModule.GetAllLevelEditorLevelMenuExtenders();
+	TArray<TSharedPtr<FExtender>> Extenders;
+
+	TSharedPtr<FUICommandList> InCommandList = GetCommandList();
+	for (int32 i = 0; i < MenuExtenderDelegates.Num(); ++i)
+	{
+		if (MenuExtenderDelegates[i].IsBound())
+		{
+			Extenders.Add(MenuExtenderDelegates[i].Execute(InCommandList.ToSharedRef()));
+		}
+	}
+	TSharedPtr<FExtender> MenuExtender = FExtender::Combine(Extenders);
+
+	// Create the menu
+	const bool bShouldCloseWindowAfterMenuSelection = true;
+	FMenuBuilder LevelMenuBuilder(bShouldCloseWindowAfterMenuSelection, InCommandList, MenuExtender);
+	
+	LevelMenuBuilder.BeginSection("LevelListing", LOCTEXT("Levels", "Levels"));
+	LevelMenuBuilder.EndSection();
+
+	return LevelMenuBuilder.MakeWidget();
+}
+
+FReply SLevelViewport::OnMenuClicked()
+{
+	OnFloatingButtonClicked();
+	// If the menu button is clicked toggle the state of the menu anchor which will open or close the menu
+	if (LevelMenuAnchor->ShouldOpenDueToClick())
+	{
+		LevelMenuAnchor->SetIsOpen(true);
+	}
+	else
+	{
+		LevelMenuAnchor->SetIsOpen(false);
+	}
+
+	return FReply::Handled();
+}
 void SLevelViewport::ConstructLevelEditorViewportClient( const FArguments& InArgs )
 {
 	if (InArgs._LevelEditorViewportClient.IsValid())
@@ -520,7 +606,7 @@ void SLevelViewport::OnDragEnter( const FGeometry& MyGeometry, const FDragDropEv
 		bDragEnterReentranceGuard = true;
 		// Don't execute the dragdrop op if the current level is locked.
 		// This prevents duplicate warning messages firing on DragEnter and Placement.
-		ULevel* CurrentLevel = (GetWorld()) ? GetWorld()->GetCurrentLevel() : NULL;
+		ULevel* CurrentLevel = (GetWorld()) ? GetWorld()->GetCurrentLevel() : nullptr;
 
 		if ( CurrentLevel && !FLevelUtils::IsLevelLocked(CurrentLevel) )
 		{
@@ -650,7 +736,7 @@ bool SLevelViewport::HandleDragObjects(const FGeometry& MyGeometry, const FDragD
 	if ( Operation->IsOfType<FAssetDragDropOp>() )
 	{
 		auto AssetOperation = StaticCastSharedPtr<FAssetDragDropOp>(DragDropEvent.GetOperation());
-		AssetOperation->SetToolTip(HintText, NULL);
+		AssetOperation->SetToolTip(HintText, nullptr);
 	}
 
 	return bValidDrag;
@@ -670,7 +756,7 @@ bool SLevelViewport::HandlePlaceDraggedObjects(const FGeometry& MyGeometry, cons
 {
 	bool bAllAssetWereLoaded = false;
 	bool bValidDrop = false;
-	UActorFactory* ActorFactory = NULL;
+	UActorFactory* ActorFactory = nullptr;
 
 	TSharedPtr< FDragDropOperation > Operation = DragDropEvent.GetOperation();
 	if (!Operation.IsValid())
@@ -722,7 +808,7 @@ bool SLevelViewport::HandlePlaceDraggedObjects(const FGeometry& MyGeometry, cons
 		for (const FAssetData& AssetData : DragDropOp->GetAssets())
 		{
 			UObject* Asset = AssetData.GetAsset();
-			if ( Asset != NULL )
+			if ( Asset != nullptr )
 			{
 				DroppedObjects.Add( Asset );
 			}
@@ -747,7 +833,7 @@ bool SLevelViewport::HandlePlaceDraggedObjects(const FGeometry& MyGeometry, cons
 			const FAssetData& AssetData = DroppedAssetDatas[AssetIdx];
 
 			UObject* Asset = AssetData.GetAsset();
-			if ( Asset != NULL )
+			if ( Asset != nullptr )
 			{
 				DroppedObjects.Add( Asset );
 			}
@@ -838,7 +924,7 @@ bool SLevelViewport::HandlePlaceDraggedObjects(const FGeometry& MyGeometry, cons
 
 FReply SLevelViewport::OnDrop( const FGeometry& MyGeometry, const FDragDropEvent& DragDropEvent )
 {
-	ULevel* CurrentLevel = (GetWorld()) ? GetWorld()->GetCurrentLevel() : NULL;
+	ULevel* CurrentLevel = (GetWorld()) ? GetWorld()->GetCurrentLevel() : nullptr;
 
 	if (CurrentLevel && !FLevelUtils::IsLevelLocked(CurrentLevel))
 	{
@@ -1508,7 +1594,7 @@ void SLevelViewport::BindStatCommand(const TSharedPtr<FUICommandInfo> InMenuItem
 
 const FSlateBrush* SLevelViewport::OnGetViewportBorderBrush() const
 {
-	const FSlateBrush* BorderBrush = NULL;
+	const FSlateBrush* BorderBrush = nullptr;
 	if( FSlateApplication::Get().IsNormalExecution() )
 	{
 		// Only show the active border if we have a valid client, its the current client being edited and we arent in immersive (in immersive there is only one visible viewport)
@@ -1724,7 +1810,7 @@ void SLevelViewport::OnCreateCameraActor(UClass* InClass)
 {		
 	// Find the perspective viewport we were using
 	FViewport* pViewPort = GEditor->GetActiveViewport();
-	FLevelEditorViewportClient* ViewportClient = NULL;
+	FLevelEditorViewportClient* ViewportClient = nullptr;
 	for( FLevelEditorViewportClient* LevelViewport : GEditor->GetLevelViewportClients())
 	{		
 		if( LevelViewport->IsPerspective() && LevelViewport->Viewport == pViewPort )
@@ -1734,7 +1820,7 @@ void SLevelViewport::OnCreateCameraActor(UClass* InClass)
 		}
 	}
 
-	if( ViewportClient == NULL )
+	if( ViewportClient == nullptr )
 	{
 		// May fail to find viewport if shortcut key was pressed on an ortho viewport, if so early out.
 		// This function only works on perspective viewports so new camera can match perspective camera.
@@ -1812,7 +1898,7 @@ void SLevelViewport::OnToggleAllVolumeActors( bool bVisible )
 
 	// Update visibility based on the new state
 	// All volume actor types should be taken since the user clicked on show or hide all to get here
-	GUnrealEd->UpdateVolumeActorVisibility( NULL, LevelViewportClient.Get() );
+	GUnrealEd->UpdateVolumeActorVisibility( nullptr, LevelViewportClient.Get() );
 }
 
 /** Called when the user toggles a volume visibility from Volumes sub-menu. **/
@@ -2024,7 +2110,7 @@ void SLevelViewport::OnUseDefaultShowFlags(bool bUseSavedDefaults)
 	if (!bUseSavedDefaults)
 	{
 		LevelViewportClient->InitializeVisibilityFlags();
-		GUnrealEd->UpdateVolumeActorVisibility(NULL, LevelViewportClient.Get());
+		GUnrealEd->UpdateVolumeActorVisibility(nullptr, LevelViewportClient.Get());
 		GEditor->Layers->UpdatePerViewVisibility(LevelViewportClient.Get());
 	}
 
@@ -2074,7 +2160,7 @@ void SLevelViewport::SaveConfig(const FString& ConfigName) const
 
 		if(GetDefault<ULevelEditorViewportSettings>()->bSaveEngineStats)
 		{
-			const TArray<FString>* EnabledStats = NULL;
+			const TArray<FString>* EnabledStats = nullptr;
 
 			// If the selected viewport is currently hosting a PIE session, we need to make sure we copy to stats from the active viewport
 			// Note: This happens if you close the editor while it's running because SwapStatCommands gets called after the config save when shutting down.
@@ -2252,7 +2338,7 @@ void SLevelViewport::OnDecrementRotationGridSize()
 
 void SLevelViewport::OnActorLockToggleFromMenu(AActor* Actor)
 {
-	if (Actor != NULL)
+	if (Actor != nullptr)
 	{
 		const bool bLockNewActor = Actor != LevelViewportClient->GetActiveActorLock().Get();
 
@@ -2295,7 +2381,7 @@ void SLevelViewport::FindSelectedInLevelScript()
 bool SLevelViewport::CanFindSelectedInLevelScript() const
 {
 	AActor* Actor = GEditor->GetSelectedActors()->GetTop<AActor>();
-	return (Actor != NULL);
+	return (Actor != nullptr);
 }
 
 void SLevelViewport::OnActorUnlock()
@@ -2868,7 +2954,7 @@ FReply SActorPreview::OnTogglePinnedButtonClicked()
 
 const FSlateBrush* SActorPreview::GetPinButtonIconBrush() const
 {
-	const FSlateBrush* IconBrush = NULL;
+	const FSlateBrush* IconBrush = nullptr;
 
 	TSharedPtr<SLevelViewport> ParentViewportPtr = ParentViewport.Pin();
 
@@ -3101,13 +3187,13 @@ void SLevelViewport::PreviewActors( const TArray< AActor* >& InActorsToPreview,
 	for( auto ActorPreviewIt = ActorPreviews.CreateConstIterator(); ActorPreviewIt; ++ActorPreviewIt )
 	{
 		auto ExistingActor = ActorPreviewIt->Actor.Get();
-		if( ExistingActor != NULL )
+		if( ExistingActor != nullptr )
 		{
 			auto bShouldKeepActor = false;
 			for( auto ActorIt = ActorsToPreview.CreateConstIterator(); ActorIt; ++ActorIt )
 			{
 				auto CurActor = *ActorIt;
-				if( CurActor != NULL && CurActor == ExistingActor )
+				if( CurActor != nullptr && CurActor == ExistingActor )
 				{
 					bShouldKeepActor = true;
 					break;
@@ -3133,7 +3219,7 @@ void SLevelViewport::PreviewActors( const TArray< AActor* >& InActorsToPreview,
 		{
 			// There could be null actors in this list as we haven't actually removed them yet.
 			auto ExistingActor = ExistingPreviewIt->Actor.Get();
-			if( ExistingActor != NULL && CurActor == ExistingActor )
+			if( ExistingActor != nullptr && CurActor == ExistingActor )
 			{
 				// Already previewing this actor.  Ignore it.
 				bIsAlreadyPreviewed = true;
@@ -3153,7 +3239,7 @@ void SLevelViewport::PreviewActors( const TArray< AActor* >& InActorsToPreview,
 	for( int32 PreviewIndex = 0; PreviewIndex < ActorPreviews.Num(); ++PreviewIndex )
 	{
 		AActor* ExistingActor = ActorPreviews[PreviewIndex].Actor.Get();
-		if ( ExistingActor == NULL )
+		if ( ExistingActor == nullptr )
 		{
 			// decrement index so we don't miss next preview after deleting
 			RemoveActorPreview( PreviewIndex-- , nullptr, bPreviewInDesktopViewport);
@@ -3482,6 +3568,85 @@ EVisibility SLevelViewport::GetCurrentLevelTextVisibility() const
 	return (&GetLevelViewportClient() == GCurrentLevelEditingViewportClient) ? ContentVisibility : EVisibility::Collapsed;
 }
 
+EVisibility SLevelViewport::GetSelectedActorsCurrentLevelTextVisibility() const
+{
+	EVisibility ContentVisibility = OnGetViewportContentVisibility();
+	if (ContentVisibility == EVisibility::Visible)
+	{
+		ContentVisibility = EVisibility::SelfHitTestInvisible;
+	}
+	return ((&GetLevelViewportClient() == GCurrentLevelEditingViewportClient) && (GEditor->GetSelectedActorCount() > 0)) ? ContentVisibility : EVisibility::Collapsed;
+}
+
+FText SLevelViewport::GetSelectedActorsCurrentLevelText(bool bDrawOnlyLabel) const
+{
+	// Display the current level and current level grid volume in the status bar
+	FText LabelName;
+	FText CurrentLevelName;
+
+
+	if (ActiveViewport.IsValid() && (&GetLevelViewportClient() == GCurrentLevelEditingViewportClient) && GetWorld())
+	{
+		if (ActiveViewport->GetPlayInEditorIsSimulate() || !ActiveViewport->GetClient()->GetWorld()->IsGameWorld())
+		{
+			if (bDrawOnlyLabel)
+			{
+				LabelName = LOCTEXT("SelectedActorsCurrentLevelLabel", "Selected Actor(s) in:");
+			}
+			else
+			{	
+				ULevel* LevelToMakeCurrent = nullptr;
+
+				// Look to the selected actors for the level to make current.
+				// If actors from multiple levels are selected, do nothing.
+				for (FSelectionIterator It(GEditor->GetSelectedActorIterator()); It; ++It)
+				{
+					AActor* Actor = static_cast<AActor*>(*It);
+					checkSlow(Actor->IsA(AActor::StaticClass()));
+
+					ULevel* ActorLevel = Actor->GetLevel();
+
+					if (!LevelToMakeCurrent)
+					{
+						// First assignment.
+						LevelToMakeCurrent = ActorLevel;
+					}
+					else if (LevelToMakeCurrent != ActorLevel)
+					{
+						// Actors from multiple levels are selected -- abort.
+						LevelToMakeCurrent = nullptr;
+						break;
+					}
+				}
+
+				FText ActualLevelName = LOCTEXT("MultipleLevelValues", "Multiple Levels");
+				if (LevelToMakeCurrent)
+				{
+					ActualLevelName = FText::FromString(FPackageName::GetShortFName(LevelToMakeCurrent->GetOutermost()->GetFName()).GetPlainNameString());
+				}
+				if (LevelToMakeCurrent == GetWorld()->PersistentLevel)
+				{
+					
+					FFormatNamedArguments Args;
+					Args.Add(TEXT("ActualLevelName"), ActualLevelName);
+					CurrentLevelName = FText::Format(LOCTEXT("LevelName", "{0} (Persistent)"), ActualLevelName);
+				}
+				else
+				{
+					CurrentLevelName = ActualLevelName;
+				}
+			}
+
+			if (bDrawOnlyLabel)
+			{
+				return LabelName;
+			}
+		}
+	}
+
+	return CurrentLevelName;
+}
+
 EVisibility SLevelViewport::GetCurrentScreenPercentageVisibility() const
 {
 	bool Visible = !IsPlayInEditorViewportActive() &&
@@ -4003,7 +4168,7 @@ FText SLevelViewport::GetLockedIconToolTip() const
 
 UWorld* SLevelViewport::GetWorld() const
 {
-	return ParentLevelEditor.IsValid() ? ParentLevelEditor.Pin()->GetWorld() : NULL;
+	return ParentLevelEditor.IsValid() ? ParentLevelEditor.Pin()->GetWorld() : nullptr;
 }
 
 void SLevelViewport::RemoveActorPreview( int32 PreviewIndex, AActor* Actor, const bool bRemoveFromDesktopViewport /*=true */ )
@@ -4023,7 +4188,7 @@ void SLevelViewport::RemoveActorPreview( int32 PreviewIndex, AActor* Actor, cons
 		// Clean up our level viewport client
 		if (ActorPreviews[PreviewIndex].LevelViewportClient.IsValid())
 		{
-			ActorPreviews[PreviewIndex].LevelViewportClient->Viewport = NULL;
+			ActorPreviews[PreviewIndex].LevelViewportClient->Viewport = nullptr;
 		}
 
 		// Remove from our list of actor previews.  This will destroy our level viewport client and viewport widget.
@@ -4059,7 +4224,7 @@ bool SLevelViewport::CanProduceActionForCommand(const TSharedRef<const FUIComman
 
 void SLevelViewport::LockActorInternal(AActor* NewActorToLock)
 {
-	if (NewActorToLock != NULL)
+	if (NewActorToLock != nullptr)
 	{
 		LevelViewportClient->SetActorLock(NewActorToLock);
 		if (LevelViewportClient->IsPerspective() && LevelViewportClient->GetActiveActorLock().IsValid())
diff --git a/Engine/Source/Editor/LevelEditor/Public/LevelEditor.h b/Engine/Source/Editor/LevelEditor/Public/LevelEditor.h
--- a/Engine/Source/Editor/LevelEditor/Public/LevelEditor.h
+++ b/Engine/Source/Editor/LevelEditor/Public/LevelEditor.h
@@ -232,6 +232,7 @@ public:
 	virtual TArray<FLevelEditorMenuExtender>& GetAllLevelEditorToolbarCreateMenuExtenders() { return LevelEditorToolbarCreateMenuExtenders; }
 	virtual TArray<FLevelEditorMenuExtender>& GetAllLevelEditorToolbarPlayMenuExtenders() { return LevelEditorToolbarPlayMenuExtenders; }
 	virtual TArray<TSharedPtr<FExtender>>& GetAllLevelEditorToolbarCinematicsMenuExtenders() {return LevelEditorToolbarCinematicsMenuExtenders;}
+	virtual TArray<FLevelEditorMenuExtender>& GetAllLevelEditorLevelMenuExtenders() { return LevelEditorLevelMenuExtenders; }
 	
 	/** Gets the extensibility managers for outside entities to extend static mesh editor's menus and toolbars */
 	virtual TSharedPtr<FExtensibilityManager> GetMenuExtensibilityManager() override {return MenuExtensibilityManager;}
@@ -383,6 +384,7 @@ private:
 	TArray<FLevelEditorMenuExtender> LevelEditorToolbarCreateMenuExtenders;
 	TArray<FLevelEditorMenuExtender> LevelEditorToolbarPlayMenuExtenders;
 	TArray<TSharedPtr<FExtender>> LevelEditorToolbarCinematicsMenuExtenders;
+	TArray<FLevelEditorMenuExtender> LevelEditorLevelMenuExtenders;
 
 	/* Pointer to the current level Editor instance */
 	TWeakPtr<class SLevelEditor> LevelEditorInstancePtr;
diff --git a/Engine/Source/Editor/LevelEditor/Public/SLevelViewport.h b/Engine/Source/Editor/LevelEditor/Public/SLevelViewport.h
--- a/Engine/Source/Editor/LevelEditor/Public/SLevelViewport.h
+++ b/Engine/Source/Editor/LevelEditor/Public/SLevelViewport.h
@@ -63,6 +63,9 @@ public:
 	 */
 	void ConstructViewportOverlayContent();
 
+	TSharedRef<SWidget> GenerateLevelMenu() const;
+	FReply OnMenuClicked();
+
 	/**
 	 * Constructs the level editor viewport client
 	 */
@@ -292,6 +295,12 @@ public:
 	/** @return The visibility of the current level text display */
 	virtual EVisibility GetCurrentLevelTextVisibility() const;
 
+	/** @return The visibility of the current level text display */
+	virtual EVisibility GetSelectedActorsCurrentLevelTextVisibility() const;
+
+	/** Called to get the text for the level the currently selected actor or actors are in. */
+	FText GetSelectedActorsCurrentLevelText(bool bDrawOnlyLabel) const;
+
 	/** @return The visibility of the current screen percentage text display */
 	EVisibility GetCurrentScreenPercentageVisibility() const;
 
@@ -847,6 +856,8 @@ private:
 	/** Whether to show a full toolbar, or a compact one */
 	bool bShowFullToolbar;
 
+	TSharedPtr<class SMenuAnchor> LevelMenuAnchor;
+
 protected:
 	void LockActorInternal(AActor* NewActorToLock);
 
diff --git a/Engine/Source/Editor/WorldBrowser/Private/WorldBrowserModule.cpp b/Engine/Source/Editor/WorldBrowser/Private/WorldBrowserModule.cpp
--- a/Engine/Source/Editor/WorldBrowser/Private/WorldBrowserModule.cpp
+++ b/Engine/Source/Editor/WorldBrowser/Private/WorldBrowserModule.cpp
@@ -15,12 +15,53 @@
 #include "SWorldHierarchy.h"
 #include "SWorldDetails.h"
 #include "Tiles/SWorldComposition.h"
+#include "Framework/MultiBox/MultiBoxExtender.h"
+#include "LevelEditor.h"
+#include "EditorLevelUtils.h"
 
 
 IMPLEMENT_MODULE( FWorldBrowserModule, WorldBrowser );
 
 #define LOCTEXT_NAMESPACE "WorldBrowser"
 
+/**
+ * Fills the level menu with extra options
+ */
+TSharedRef<FExtender> FWorldBrowserModule::BindLevelMenu(const TSharedRef<FUICommandList> CommandList)
+{
+	TSharedRef<FExtender> Extender(new FExtender());
+	Extender->AddMenuExtension("LevelListing", EExtensionHook::Before, CommandList, FMenuExtensionDelegate::CreateRaw(this, &FWorldBrowserModule::BuildLevelMenu));
+	return Extender;
+}
+
+
+void FWorldBrowserModule::BuildLevelMenu(FMenuBuilder& MenuBuilder)
+{
+	UWorld* EditorWorld = GEditor->GetEditorWorldContext().World();
+	TSharedPtr<FLevelCollectionModel> Model = SharedWorldModel(EditorWorld);
+	if (Model.IsValid())
+	{
+		FLevelModelList ModelList = Model->GetFilteredLevels();
+		for (TSharedPtr<FLevelModel> LevelModel : ModelList)
+		{
+			FUIAction Action(FExecuteAction::CreateRaw(this, &FWorldBrowserModule::SetCurrentSublevel, LevelModel->GetLevelObject()),
+				FCanExecuteAction(),
+				FIsActionChecked::CreateRaw(this, &FWorldBrowserModule::IsCurrentSublevel, LevelModel->GetLevelObject()));
+			MenuBuilder.AddMenuEntry(FText::FromString(LevelModel->GetDisplayName()), FText::GetEmpty(), FSlateIcon(), Action, NAME_None, EUserInterfaceActionType::Button);
+		}
+	}
+}
+
+bool FWorldBrowserModule::IsCurrentSublevel(ULevel* InLevel)
+{
+	return InLevel->IsCurrentLevel();
+}
+
+void FWorldBrowserModule::SetCurrentSublevel(ULevel* InLevel)
+{
+	EditorLevelUtils::MakeLevelCurrent(InLevel);
+}
+
 void FWorldBrowserModule::StartupModule()
 {
 	FLevelCollectionCommands::Register();
@@ -38,6 +79,19 @@ void FWorldBrowserModule::StartupModule()
 
 	UWorldComposition::WorldCompositionChangedEvent.AddRaw(this, &FWorldBrowserModule::OnWorldCompositionChanged);
 
+	// Have to check GIsEditor because right now editor modules can be loaded by the game
+	// Once LoadModule is guaranteed to return NULL for editor modules in game, this can be removed
+	// Without this check, loading the level editor in the game will crash
+	if (GIsEditor)
+	{
+		// Extend the level viewport levels menu
+		FLevelEditorModule& LevelEditorModule = FModuleManager::Get().LoadModuleChecked<FLevelEditorModule>("LevelEditor");
+		LevelMenuExtender = FLevelEditorModule::FLevelEditorMenuExtender::CreateRaw(this, &FWorldBrowserModule::BindLevelMenu);
+		auto& MenuExtenders = LevelEditorModule.GetAllLevelEditorLevelMenuExtenders();
+		MenuExtenders.Add(LevelMenuExtender);
+		LevelMenuExtenderHandle = MenuExtenders.Last().GetHandle();
+	}
+
 	FLevelFolders::Init();
 }
 
diff --git a/Engine/Source/Editor/WorldBrowser/Public/WorldBrowserModule.h b/Engine/Source/Editor/WorldBrowser/Public/WorldBrowserModule.h
--- a/Engine/Source/Editor/WorldBrowser/Public/WorldBrowserModule.h
+++ b/Engine/Source/Editor/WorldBrowser/Public/WorldBrowserModule.h
@@ -3,6 +3,8 @@
 
 #include "CoreMinimal.h"
 #include "Modules/ModuleInterface.h"
+#include "Framework/MultiBox/MultiBoxBuilder.h"
+#include "LevelEditor.h"
 
 class FLevelCollectionModel;
 
@@ -54,7 +56,21 @@ private:
 	void OnWorldCreated(UWorld* InWorld);
 	void OnWorldDestroyed(UWorld* InWorld);
 	void OnWorldCompositionChanged(UWorld* InWorld);
+
+	/** Bind world browser command delegate to the level viewport */
+	TSharedRef<FExtender> BindLevelMenu(const TSharedRef<FUICommandList> CommandList);
+
+	/** Fill out the level menu with entries for level operations */
+	void BuildLevelMenu(FMenuBuilder& MenuBuilder);
+
+	bool IsCurrentSublevel(ULevel* InLevel);
+	void SetCurrentSublevel(ULevel* InLevel);
 				
 private:
 	TWeakPtr<class FLevelCollectionModel>	WorldModel;
+
+	/** Extender for the level menu */
+	FLevelEditorModule::FLevelEditorMenuExtender LevelMenuExtender;
+	/** Delegate called when the menu is created */
+	FDelegateHandle LevelMenuExtenderHandle;
 };
