Fixes for UI appearance after level preview changes
#rnx

diff --git a/Engine/Source/Editor/LevelEditor/Private/SLevelViewport.cpp b/Engine/Source/Editor/LevelEditor/Private/SLevelViewport.cpp
--- a/Engine/Source/Editor/LevelEditor/Private/SLevelViewport.cpp
+++ b/Engine/Source/Editor/LevelEditor/Private/SLevelViewport.cpp
@@ -235,6 +235,7 @@ void SLevelViewport::ConstructViewportOverlayContent()
 
 	ViewportOverlay->AddSlot( SlotIndex )
 	.HAlign(HAlign_Right)
+	.Padding(0.0f, 0.0f, 0.0f, 35.0f)
 	[
 		SAssignNew( ActorPreviewHorizontalBox, SHorizontalBox )
 	];
@@ -341,8 +342,7 @@ void SLevelViewport::ConstructViewportOverlayContent()
 					[
 						SNew(STextBlock)
 						.Text(this, &SLevelViewport::GetCurrentLevelText, true)
-						.Font(FEditorStyle::GetFontStyle(TEXT("MenuItem.Font")))
-						.ShadowOffset(FVector2D(1, 1))
+						.Font(FEditorStyle::GetFontStyle("EditorViewportToolBar.Font"))
 					]
 					// Current level
 					+ SHorizontalBox::Slot()
@@ -351,8 +351,7 @@ void SLevelViewport::ConstructViewportOverlayContent()
 					[
 						SNew(STextBlock)
 						.Text(this, &SLevelViewport::GetCurrentLevelText, false)
-						.Font(FEditorStyle::GetFontStyle(TEXT("MenuItem.Font")))
-						.ShadowOffset(FVector2D(1, 1))
+						.Font(FEditorStyle::GetFontStyle("EditorViewportToolBar.Font"))
 					]
 				]
 			]
@@ -3580,7 +3579,7 @@ EVisibility SLevelViewport::GetSelectedActorsCurrentLevelTextVisibility() const
 
 FText SLevelViewport::GetSelectedActorsCurrentLevelText(bool bDrawOnlyLabel) const
 {
-	// Display the current level and current level grid volume in the status bar
+	// Display the currently selected actor's level
 	FText LabelName;
 	FText CurrentLevelName;
 
