Fix for appearance of level text in viewport
#rnx

diff --git a/Engine/Source/Editor/LevelEditor/Private/SLevelViewport.cpp b/Engine/Source/Editor/LevelEditor/Private/SLevelViewport.cpp
--- a/Engine/Source/Editor/LevelEditor/Private/SLevelViewport.cpp
+++ b/Engine/Source/Editor/LevelEditor/Private/SLevelViewport.cpp
@@ -324,6 +324,7 @@ void SLevelViewport::ConstructViewportOverlayContent()
 		[
 			SAssignNew(LevelMenuAnchor, SMenuAnchor)
 			.Placement(MenuPlacement_AboveAnchor)
+			.Visibility(this, &SLevelViewport::GetCurrentLevelTextVisibility)
 			[
 				SNew(SButton)
 				// Allows users to drag with the mouse to select options after opening the menu */
@@ -332,6 +333,7 @@ void SLevelViewport::ConstructViewportOverlayContent()
 				.VAlign(VAlign_Center)
 				.ButtonStyle(FEditorStyle::Get(), "EditorViewportToolBar.MenuButton")
 				.OnClicked(this, &SLevelViewport::OnMenuClicked)
+				.Visibility(this, &SLevelViewport::GetCurrentLevelTextVisibility)
 				[
 					SNew(SHorizontalBox)
 					.Visibility(this, &SLevelViewport::GetCurrentLevelTextVisibility)
@@ -3532,8 +3534,8 @@ FText SLevelViewport::GetCurrentLevelText( bool bDrawOnlyLabel ) const
 			}
 			else
 			{
-				// Get the level name (without the number at the end)
-				FText ActualLevelName = FText::FromString(FPackageName::GetShortFName(GetWorld()->GetCurrentLevel()->GetOutermost()->GetFName()).GetPlainNameString());
+				// Get the level name 
+				FText ActualLevelName = FText::FromName(FPackageName::GetShortFName(GetWorld()->GetCurrentLevel()->GetOutermost()->GetFName()));
 
 				if(GetWorld()->GetCurrentLevel() == GetWorld()->PersistentLevel)
 				{
@@ -3564,7 +3566,7 @@ EVisibility SLevelViewport::GetCurrentLevelTextVisibility() const
 	{
 		ContentVisibility = EVisibility::SelfHitTestInvisible;
 	}
-	return (&GetLevelViewportClient() == GCurrentLevelEditingViewportClient) ? ContentVisibility : EVisibility::Collapsed;
+	return (&GetLevelViewportClient() == GCurrentLevelEditingViewportClient) && !IsPlayInEditorViewportActive() ? ContentVisibility : EVisibility::Collapsed;
 }
 
 EVisibility SLevelViewport::GetSelectedActorsCurrentLevelTextVisibility() const
@@ -3574,7 +3576,7 @@ EVisibility SLevelViewport::GetSelectedActorsCurrentLevelTextVisibility() const
 	{
 		ContentVisibility = EVisibility::SelfHitTestInvisible;
 	}
-	return ((&GetLevelViewportClient() == GCurrentLevelEditingViewportClient) && (GEditor->GetSelectedActorCount() > 0)) ? ContentVisibility : EVisibility::Collapsed;
+	return (&GetLevelViewportClient() == GCurrentLevelEditingViewportClient) && (GEditor->GetSelectedActorCount() > 0) && !IsPlayInEditorViewportActive() ? ContentVisibility : EVisibility::Collapsed;
 }
 
 FText SLevelViewport::GetSelectedActorsCurrentLevelText(bool bDrawOnlyLabel) const
@@ -3621,7 +3623,7 @@ FText SLevelViewport::GetSelectedActorsCurrentLevelText(bool bDrawOnlyLabel) con
 				FText ActualLevelName = LOCTEXT("MultipleLevelValues", "Multiple Levels");
 				if (LevelToMakeCurrent)
 				{
-					ActualLevelName = FText::FromString(FPackageName::GetShortFName(LevelToMakeCurrent->GetOutermost()->GetFName()).GetPlainNameString());
+					ActualLevelName = FText::FromName(FPackageName::GetShortFName(LevelToMakeCurrent->GetOutermost()->GetFName()));
 				}
 				if (LevelToMakeCurrent == GetWorld()->PersistentLevel)
 				{
