Show correct HTML5 preview icon with PC editor.

diff --git a/Engine/Source/Editor/LevelEditor/Private/LevelEditorToolBar.cpp b/Engine/Source/Editor/LevelEditor/Private/LevelEditorToolBar.cpp
--- a/Engine/Source/Editor/LevelEditor/Private/LevelEditorToolBar.cpp
+++ b/Engine/Source/Editor/LevelEditor/Private/LevelEditorToolBar.cpp
@@ -1329,6 +1329,9 @@ TSharedRef< SWidget > FLevelEditorToolBar::MakeLevelEditorToolBar( const TShared
 					{
 						return FSlateIcon(FEditorStyle::GetStyleSetName(), GEditor->IsFeatureLevelPreviewActive() ? "LevelEditor.PreviewMode.iOS.Enabled" : "LevelEditor.PreviewMode.iOS.Disabled");
 					}
+					case SP_VULKAN_PCES3_1:
+					case SP_OPENGL_PCES2:
+					case SP_PCD3D_ES2:
 					case SP_METAL_MACES2:
                 	case SP_OPENGL_ES2_WEBGL:
 					{
