Disable featurelevel switching when a lighting build is active.

diff --git a/Engine/Source/Editor/LevelEditor/Private/LevelEditorActions.cpp b/Engine/Source/Editor/LevelEditor/Private/LevelEditorActions.cpp
--- a/Engine/Source/Editor/LevelEditor/Private/LevelEditorActions.cpp
+++ b/Engine/Source/Editor/LevelEditor/Private/LevelEditorActions.cpp
@@ -574,7 +574,7 @@ void FLevelEditorActionCallbacks::ToggleFeatureLevelPreview()
 
 bool FLevelEditorActionCallbacks::IsFeatureLevelPreviewEnabled()
 {
-	if (GEditor->PlayWorld || GUnrealEd->bIsSimulatingInEditor)
+	if (GEditor->PlayWorld || GUnrealEd->bIsSimulatingInEditor || GUnrealEd->IsLightingBuildCurrentlyRunning())
 	{
 		return false;
 	}
