Fixed SM4 preview icon

diff --git a/Engine/Source/Editor/EditorStyle/Private/SlateEditorStyle.cpp b/Engine/Source/Editor/EditorStyle/Private/SlateEditorStyle.cpp
--- a/Engine/Source/Editor/EditorStyle/Private/SlateEditorStyle.cpp
+++ b/Engine/Source/Editor/EditorStyle/Private/SlateEditorStyle.cpp
@@ -4774,8 +4774,8 @@ void FSlateEditorStyle::FStyle::SetupLevelEditorStyle()
 		Set("LevelEditor.PreviewMode.SM5.Disabled.Small", new IMAGE_BRUSH("Icons/icon_PreviewMode_SM5_Enabled_40x", Icon20x20));
 		Set("LevelEditor.PreviewMode.SM4.Enabled", new IMAGE_BRUSH("Icons/icon_PreviewMode_SM4_Enabled_40x", Icon40x40));
 		Set("LevelEditor.PreviewMode.SM4.Enabled.Small", new IMAGE_BRUSH("Icons/icon_PreviewMode_SM4_Enabled_40x", Icon20x20));
-		Set("LevelEditor.PreviewMode.SM4.Disabled", new IMAGE_BRUSH("Icons/icon_PreviewMode_SM4_Enabled_40x", Icon40x40));
-		Set("LevelEditor.PreviewMode.SM4.Disabled.Small", new IMAGE_BRUSH("Icons/icon_PreviewMode_SM4_Enabled_40x", Icon20x20));
+		Set("LevelEditor.PreviewMode.SM4.Disabled", new IMAGE_BRUSH("Icons/icon_PreviewMode_SM4_Disabled_40x", Icon40x40));
+		Set("LevelEditor.PreviewMode.SM4.Disabled.Small", new IMAGE_BRUSH("Icons/icon_PreviewMode_SM4_Disabled_40x", Icon20x20));
 		Set("LevelEditor.PreviewMode.AndroidES2.Enabled", new IMAGE_BRUSH("Icons/icon_PreviewMode_AndroidES2_Enabled_40x", Icon40x40));
 		Set("LevelEditor.PreviewMode.AndroidES2.Enabled.Small", new IMAGE_BRUSH("Icons/icon_PreviewMode_AndroidES2_Enabled_40x", Icon20x20));
 		Set("LevelEditor.PreviewMode.AndroidES2.Disabled", new IMAGE_BRUSH("Icons/icon_PreviewMode_AndroidES2_Disabled_40x", Icon40x40));
diff --git a/Engine/Source/Editor/LevelEditor/Private/LevelEditorToolBar.cpp b/Engine/Source/Editor/LevelEditor/Private/LevelEditorToolBar.cpp
--- a/Engine/Source/Editor/LevelEditor/Private/LevelEditorToolBar.cpp
+++ b/Engine/Source/Editor/LevelEditor/Private/LevelEditorToolBar.cpp
@@ -1278,6 +1278,31 @@ TSharedRef< SWidget > FLevelEditorToolBar::MakeLevelEditorToolBar( const TShared
 	{
 		struct FPreviewModeFunctionality
 		{
+
+			static FText GetPreviewModeText()
+			{
+				switch (GEditor->PreviewFeatureLevel)
+				{
+					case ERHIFeatureLevel::SM4:
+					{
+						return LOCTEXT("PreviewModeSM4_Text", "SM4 Preview");
+					}
+					case ERHIFeatureLevel::ES2:
+					{
+						return LOCTEXT("PreviewModeES2_Text", "ES2 Preview");
+					}
+					case ERHIFeatureLevel::ES3_1:
+					{
+						return LOCTEXT("PreviewModeES3_1_Text", "ES3.1 Preview");
+					}
+					default:
+					{
+						return LOCTEXT("PreviewModeGeneric", "Preview Mode");
+					}
+				}
+			}
+
+
 			static FText GetPreviewModeTooltip()
 			{
 				UMaterialShaderQualitySettings* MaterialShaderQualitySettings = UMaterialShaderQualitySettings::Get();
@@ -1363,7 +1388,7 @@ TSharedRef< SWidget > FLevelEditorToolBar::MakeLevelEditorToolBar( const TShared
 		ToolbarBuilder.AddToolBarButton(
 			FLevelEditorCommands::Get().ToggleFeatureLevelPreview,
 			NAME_None,
-			LOCTEXT("PreviewModeActivate", "Preview Mode"),
+			TAttribute<FText>::Create(&FPreviewModeFunctionality::GetPreviewModeText),
         	TAttribute<FText>::Create(&FPreviewModeFunctionality::GetPreviewModeTooltip),
         	TAttribute<FSlateIcon>::Create(&FPreviewModeFunctionality::GetPreviewModeIcon)
 			);
