LiveCode: Don't show the dropdown next to the compile button when the editor has a blueprint project open.

#rnx

diff --git a/Engine/Source/Editor/LevelEditor/Private/LevelEditorToolBar.cpp b/Engine/Source/Editor/LevelEditor/Private/LevelEditorToolBar.cpp
--- a/Engine/Source/Editor/LevelEditor/Private/LevelEditorToolBar.cpp
+++ b/Engine/Source/Editor/LevelEditor/Private/LevelEditorToolBar.cpp
@@ -1449,7 +1449,11 @@ TSharedRef< SWidget > FLevelEditorToolBar::MakeLevelEditorToolBar( const TShared
 
 #if WITH_LIVE_CODING
 			ToolbarBuilder.AddComboButton(
-				FUIAction(),
+				FUIAction(
+					FExecuteAction(),
+					FCanExecuteAction(),
+					FIsActionChecked(),
+					FIsActionButtonVisible::CreateStatic(FLevelEditorActionCallbacks::CanShowSourceCodeActions)), 
 				FOnGetContent::CreateStatic( &FLevelEditorToolBar::GenerateCompileMenuContent, InCommandList ),
 				LOCTEXT( "CompileCombo_Label", "Compile Options" ),
 				LOCTEXT( "CompileComboToolTip", "Compile options menu" ),
