Fix "Compile" button in toolbar using legacy hot reload when live coding startup is set to manual.

diff --git a/Engine/Source/Developer/Windows/LiveCoding/Private/LiveCodingModule.cpp b/Engine/Source/Developer/Windows/LiveCoding/Private/LiveCodingModule.cpp
--- a/Engine/Source/Developer/Windows/LiveCoding/Private/LiveCodingModule.cpp
+++ b/Engine/Source/Developer/Windows/LiveCoding/Private/LiveCodingModule.cpp
@@ -25,6 +25,7 @@ FString GLiveCodingConsoleArguments;
 
 FLiveCodingModule::FLiveCodingModule()
 	: bEnabledLastTick(false)
+	, bEnabledForSession(false)
 	, bStarted(false)
 {
 }
@@ -135,9 +136,15 @@ void FLiveCodingModule::EnableForSession(bool bEnable)
 			LppSetVisible(false);
 		}
 	}
+	bEnabledForSession = bEnable;
 }
 
 bool FLiveCodingModule::IsEnabledForSession() const
+{
+	return bEnabledForSession;
+}
+
+bool FLiveCodingModule::HasStarted() const
 {
 	return bStarted;
 }
diff --git a/Engine/Source/Developer/Windows/LiveCoding/Private/LiveCodingModule.h b/Engine/Source/Developer/Windows/LiveCoding/Private/LiveCodingModule.h
--- a/Engine/Source/Developer/Windows/LiveCoding/Private/LiveCodingModule.h
+++ b/Engine/Source/Developer/Windows/LiveCoding/Private/LiveCodingModule.h
@@ -26,6 +26,7 @@ public:
 	virtual bool IsEnabledByDefault() const override;
 	virtual void EnableForSession(bool bInEnabled) override;
 	virtual bool IsEnabledForSession() const override;
+	virtual bool HasStarted() const override;
 	virtual void ShowConsole() override;
 	virtual void Compile() override;
 	virtual bool IsCompiling() const override;
@@ -35,6 +36,7 @@ private:
 	ULiveCodingSettings* Settings;
 	TSharedPtr<ISettingsSection> SettingsSection;
 	bool bEnabledLastTick;
+	bool bEnabledForSession;
 	bool bStarted;
 	TSet<FString> EnabledModules;
 
diff --git a/Engine/Source/Developer/Windows/LiveCoding/Public/ILiveCodingModule.h b/Engine/Source/Developer/Windows/LiveCoding/Public/ILiveCodingModule.h
--- a/Engine/Source/Developer/Windows/LiveCoding/Public/ILiveCodingModule.h
+++ b/Engine/Source/Developer/Windows/LiveCoding/Public/ILiveCodingModule.h
@@ -12,8 +12,12 @@ class ILiveCodingModule : public IModuleInterface
 public:
 	virtual void EnableByDefault(bool bEnabled) = 0;
 	virtual bool IsEnabledByDefault() const = 0;
+
 	virtual void EnableForSession(bool bEnabled) = 0;
 	virtual bool IsEnabledForSession() const = 0;
+
+	virtual bool HasStarted() const = 0;
+
 	virtual void ShowConsole() = 0;
 	virtual void Compile() = 0;
 	virtual bool IsCompiling() const = 0;
diff --git a/Engine/Source/Editor/LevelEditor/Private/LevelEditorActions.cpp b/Engine/Source/Editor/LevelEditor/Private/LevelEditorActions.cpp
--- a/Engine/Source/Editor/LevelEditor/Private/LevelEditorActions.cpp
+++ b/Engine/Source/Editor/LevelEditor/Private/LevelEditorActions.cpp
@@ -1033,7 +1033,7 @@ void FLevelEditorActionCallbacks::RecompileGameCode_Clicked()
 {
 #if WITH_LIVE_CODING
 	ILiveCodingModule* LiveCoding = FModuleManager::GetModulePtr<ILiveCodingModule>(LIVE_CODING_MODULE_NAME);
-	if (LiveCoding != nullptr && LiveCoding->IsEnabledForSession())
+	if (LiveCoding != nullptr && LiveCoding->IsEnabledByDefault())
 	{
 		LiveCoding->Compile();
 		return;
@@ -1053,7 +1053,7 @@ bool FLevelEditorActionCallbacks::Recompile_CanExecute()
 {
 #if WITH_LIVE_CODING
 	ILiveCodingModule* LiveCoding = FModuleManager::GetModulePtr<ILiveCodingModule>(LIVE_CODING_MODULE_NAME);
-	if (LiveCoding != nullptr && LiveCoding->IsEnabledForSession())
+	if (LiveCoding != nullptr && LiveCoding->IsEnabledByDefault())
 	{
 		return !LiveCoding->IsCompiling();
 	}
@@ -1099,7 +1099,7 @@ void FLevelEditorActionCallbacks::LiveCoding_StartSession_Clicked()
 bool FLevelEditorActionCallbacks::LiveCoding_CanStartSession()
 {
 	ILiveCodingModule* LiveCoding = FModuleManager::GetModulePtr<ILiveCodingModule>(LIVE_CODING_MODULE_NAME);
-	return LiveCoding != nullptr && LiveCoding->IsEnabledByDefault() && !LiveCoding->IsEnabledForSession();
+	return LiveCoding != nullptr && LiveCoding->IsEnabledByDefault() && !LiveCoding->HasStarted();
 }
 
 void FLevelEditorActionCallbacks::LiveCoding_ShowConsole_Clicked()
