Viewport Add Camera list no longer can contain Abstract classes.

diff --git a/Engine/Source/Editor/LevelEditor/Private/SLevelViewport.cpp b/Engine/Source/Editor/LevelEditor/Private/SLevelViewport.cpp
--- a/Engine/Source/Editor/LevelEditor/Private/SLevelViewport.cpp
+++ b/Engine/Source/Editor/LevelEditor/Private/SLevelViewport.cpp
@@ -1333,6 +1333,12 @@ void SLevelViewport::BindOptionCommands( FUICommandList& OutCommandList )
 		{
 			DisplayName = FName(*DisplayName.ToString().LeftChop(2));
 		}
+
+		UClass* CameraClass = FindObject<UClass>(ANY_PACKAGE, *Name.ToString());
+		if (CameraClass->HasAllClassFlags(CLASS_Abstract))
+		{
+			continue;
+		}
 		
 		// Look for existing UI Command info so one isn't created for every viewport
 		TSharedPtr<FUICommandInfo> * FoundCamera = FLevelViewportCommands::Get().CreateCameras.FindByPredicate([Name](TSharedPtr<FUICommandInfo> Camera) { return Camera->GetCommandName() == Name; });
@@ -1341,7 +1347,7 @@ void SLevelViewport::BindOptionCommands( FUICommandList& OutCommandList )
 		{
 			OutCommandList.MapAction(
 				*FoundCamera,
-				FExecuteAction::CreateSP(this, &SLevelViewport::OnCreateCameraActor, FindObject<UClass>(ANY_PACKAGE, *Name.ToString()))
+				FExecuteAction::CreateSP(this, &SLevelViewport::OnCreateCameraActor, CameraClass)
 			);
 		}
 		else
@@ -1351,7 +1357,7 @@ void SLevelViewport::BindOptionCommands( FUICommandList& OutCommandList )
 			
 			OutCommandList.MapAction(
 				NewCamera,
-				FExecuteAction::CreateSP(this, &SLevelViewport::OnCreateCameraActor, FindObject<UClass>(ANY_PACKAGE, *Name.ToString()))
+				FExecuteAction::CreateSP(this, &SLevelViewport::OnCreateCameraActor, CameraClass)
 			);
 			
 			FLevelViewportCommands::Get().CreateCameras.Add(NewCamera);
