Fix for unloaded BP classes causing crash when trying to create a new camera.

diff --git a/Engine/Source/Editor/LevelEditor/Private/SLevelViewport.cpp b/Engine/Source/Editor/LevelEditor/Private/SLevelViewport.cpp
--- a/Engine/Source/Editor/LevelEditor/Private/SLevelViewport.cpp
+++ b/Engine/Source/Editor/LevelEditor/Private/SLevelViewport.cpp
@@ -1335,7 +1335,7 @@ void SLevelViewport::BindOptionCommands( FUICommandList& OutCommandList )
 		}
 
 		UClass* CameraClass = FindObject<UClass>(ANY_PACKAGE, *Name.ToString());
-		if (CameraClass->HasAllClassFlags(CLASS_Abstract))
+		if (!CameraClass || CameraClass->HasAllClassFlags(CLASS_Abstract))
 		{
 			continue;
 		}
