PR #2940: Sort the BP sub-level selection dropdown menu alphabetically (Contributed by [REDACTED_PERSON])

diff --git a/Engine/Source/Editor/LevelEditor/Private/LevelEditorToolBar.cpp b/Engine/Source/Editor/LevelEditor/Private/LevelEditorToolBar.cpp
--- a/Engine/Source/Editor/LevelEditor/Private/LevelEditorToolBar.cpp
+++ b/Engine/Source/Editor/LevelEditor/Private/LevelEditorToolBar.cpp
@@ -631,6 +631,21 @@ namespace LevelEditorActionHelpers
 		);
 #undef LOCTEXT_NAMESPACE
 	}
+
+	struct FLevelSortByName
+	{
+		bool operator ()(const ULevel* LHS, const ULevel* RHS) const
+		{
+			if (LHS != NULL && LHS->GetOutermost() != NULL && RHS != NULL && RHS->GetOutermost() != NULL)
+			{
+				return FPaths::GetCleanFilename(LHS->GetOutermost()->GetName()) < FPaths::GetCleanFilename(RHS->GetOutermost()->GetName());
+			}
+			else
+			{
+				return false;
+			}
+		}
+	};
 }
 
 void LevelEditorActionHelpers::GetBlueprintSettingsSubMenu(FMenuBuilder& InMenuBuilder, const TSharedRef<FUICommandList> InCommandList, FBlueprintMenuSettings InSettingsData)
@@ -2142,21 +2157,21 @@ TSharedRef< SWidget > FLevelEditorToolBar::GenerateOpenBlueprintMenuContent( TSh
 			InMenuBuilder.BeginSection(NAME_None, LOCTEXT("SubLevelsHeading", "Sub-Level Blueprints"));
 			{
 				UWorld* World = InLvlEditor.Pin()->GetWorld();
-				for (int32 iLevel = 0; iLevel < World->GetNumLevels(); iLevel++)
+				// Sort the levels alphabetically 
+				TArray<ULevel*> SortedLevels = World->GetLevels();
+				Algo::Sort(SortedLevels, LevelEditorActionHelpers::FLevelSortByName());
+
+				for (ULevel* const Level : SortedLevels)
 				{
-					ULevel* Level = World->GetLevel(iLevel);
-					if (Level != NULL && Level->GetOutermost() != NULL)
+					if (Level != NULL && Level->GetOutermost() != NULL && !Level->IsPersistentLevel())
 					{
-						if (!Level->IsPersistentLevel())
-						{
-							FUIAction UIAction
-								(
-								FExecuteAction::CreateStatic(&FLevelEditorToolBar::OnOpenSubLevelBlueprint, Level)
-								);
+						FUIAction UIAction
+						(
+							FExecuteAction::CreateStatic(&FLevelEditorToolBar::OnOpenSubLevelBlueprint, Level)
+						);
 
-							FText DisplayName = FText::Format(LOCTEXT("SubLevelBlueprintItem", "Edit {0}"), FText::FromString(FPaths::GetCleanFilename(Level->GetOutermost()->GetName())));
-							InMenuBuilder.AddMenuEntry(DisplayName, FText::GetEmpty(), EditBP, UIAction);
-						}
+						FText DisplayName = FText::Format(LOCTEXT("SubLevelBlueprintItem", "Edit {0}"), FText::FromString(FPaths::GetCleanFilename(Level->GetOutermost()->GetName())));
+						InMenuBuilder.AddMenuEntry(DisplayName, FText::GetEmpty(), EditBP, UIAction);
 					}
 				}
 			}
