Disable PIE/Simpulate Enter/Exit sound by default as it can flush async loading the first time the sound is loaded

diff --git a/Engine/Source/Editor/LevelEditor/Private/SLevelViewport.cpp b/Engine/Source/Editor/LevelEditor/Private/SLevelViewport.cpp
--- a/Engine/Source/Editor/LevelEditor/Private/SLevelViewport.cpp
+++ b/Engine/Source/Editor/LevelEditor/Private/SLevelViewport.cpp
@@ -4128,7 +4128,10 @@ void SLevelViewport::SwapViewportsForSimulateInEditor()
 	ViewTransitionType = EViewTransition::StartingSimulate;
 	ViewTransitionAnim = FCurveSequence( 0.0f, 1.5f, ECurveEaseFunction::CubicOut );
 	bViewTransitionAnimPending = true;
-	GEditor->PlayEditorSound( TEXT( "/Engine/EditorSounds/GamePreview/PossessPlayer_Cue.PossessPlayer_Cue" ) );
+	if (GetDefault<ULevelEditorPlaySettings>()->EnablePIEEnterAndExitSounds)
+	{
+		GEditor->PlayEditorSound(TEXT("/Engine/EditorSounds/GamePreview/PossessPlayer_Cue.PossessPlayer_Cue"));
+	}
 }
 
 
