Added extension points to level editor toolbar blueprint menu

diff --git a/Engine/Source/Editor/LevelEditor/Private/LevelEditorToolBar.cpp b/Engine/Source/Editor/LevelEditor/Private/LevelEditorToolBar.cpp
--- a/Engine/Source/Editor/LevelEditor/Private/LevelEditorToolBar.cpp
+++ b/Engine/Source/Editor/LevelEditor/Private/LevelEditorToolBar.cpp
@@ -2124,11 +2124,13 @@ TSharedRef< SWidget > FLevelEditorToolBar::GenerateOpenBlueprintMenuContent( TSh
 		}
 	};
 
+	FLevelEditorModule& LevelEditorModule = FModuleManager::LoadModuleChecked<FLevelEditorModule>("LevelEditor");
+	TSharedPtr<FExtender> Extender = FExtender::Combine(LevelEditorModule.GetAllLevelEditorToolbarBlueprintsMenuExtenders());
 
 	const bool bShouldCloseWindowAfterMenuSelection = true;
-	FMenuBuilder MenuBuilder( bShouldCloseWindowAfterMenuSelection, InCommandList );
+	FMenuBuilder MenuBuilder( bShouldCloseWindowAfterMenuSelection, InCommandList, Extender);
 
-	MenuBuilder.BeginSection(NAME_None, LOCTEXT("BlueprintClass", "Blueprint Class"));
+	MenuBuilder.BeginSection("BlueprintClass", LOCTEXT("BlueprintClass", "Blueprint Class"));
 	{
 		// Create a blank BP
 		MenuBuilder.AddMenuEntry(FLevelEditorCommands::Get().CreateBlankBlueprintClass);
@@ -2148,7 +2150,7 @@ TSharedRef< SWidget > FLevelEditorToolBar::GenerateOpenBlueprintMenuContent( TSh
 	}
 	MenuBuilder.EndSection();
 
-	MenuBuilder.BeginSection(NAME_None, LOCTEXT("LevelScriptBlueprints", "Level Blueprints"));
+	MenuBuilder.BeginSection("LevelScriptBlueprints", LOCTEXT("LevelScriptBlueprints", "Level Blueprints"));
 	{
 		MenuBuilder.AddMenuEntry( FLevelEditorCommands::Get().OpenLevelBlueprint );
 
@@ -2165,7 +2167,7 @@ TSharedRef< SWidget > FLevelEditorToolBar::GenerateOpenBlueprintMenuContent( TSh
 	}
 	MenuBuilder.EndSection();
 
-	MenuBuilder.BeginSection(NAME_None, LOCTEXT("ProjectSettingsClasses", "Project Settings"));
+	MenuBuilder.BeginSection("ProjectSettingsClasses", LOCTEXT("ProjectSettingsClasses", "Project Settings"));
 	{
 		// If source control is enabled, queue up a query to the status of the config file so it is (hopefully) ready before we get to the sub-menu
 		if(ISourceControlModule::Get().IsEnabled())
@@ -2179,7 +2181,7 @@ TSharedRef< SWidget > FLevelEditorToolBar::GenerateOpenBlueprintMenuContent( TSh
 	}
 	MenuBuilder.EndSection();
 
-	MenuBuilder.BeginSection(NAME_None, LOCTEXT("WorldSettingsClasses", "World Override"));
+	MenuBuilder.BeginSection("WorldSettingsClasses", LOCTEXT("WorldSettingsClasses", "World Override"));
 	{
 		LevelEditorActionHelpers::CreateGameModeSubMenu(MenuBuilder, InCommandList, InLevelEditor, false);
 	}
diff --git a/Engine/Source/Editor/LevelEditor/Public/LevelEditor.h b/Engine/Source/Editor/LevelEditor/Public/LevelEditor.h
--- a/Engine/Source/Editor/LevelEditor/Public/LevelEditor.h
+++ b/Engine/Source/Editor/LevelEditor/Public/LevelEditor.h
@@ -231,6 +231,7 @@ public:
 	virtual TArray<FLevelEditorMenuExtender>& GetAllLevelEditorToolbarSourceControlMenuExtenders() { return LevelEditorToolbarSourceControlMenuExtenders; }
 	virtual TArray<FLevelEditorMenuExtender>& GetAllLevelEditorToolbarCreateMenuExtenders() { return LevelEditorToolbarCreateMenuExtenders; }
 	virtual TArray<FLevelEditorMenuExtender>& GetAllLevelEditorToolbarPlayMenuExtenders() { return LevelEditorToolbarPlayMenuExtenders; }
+	virtual TArray<TSharedPtr<FExtender>>& GetAllLevelEditorToolbarBlueprintsMenuExtenders() { return LevelEditorToolbarBlueprintsMenuExtenders; }
 	virtual TArray<TSharedPtr<FExtender>>& GetAllLevelEditorToolbarCinematicsMenuExtenders() {return LevelEditorToolbarCinematicsMenuExtenders;}
 	virtual TArray<FLevelEditorMenuExtender>& GetAllLevelEditorLevelMenuExtenders() { return LevelEditorLevelMenuExtenders; }
 	
@@ -384,6 +385,7 @@ private:
 	TArray<FLevelEditorMenuExtender> LevelEditorToolbarCreateMenuExtenders;
 	TArray<FLevelEditorMenuExtender> LevelEditorToolbarPlayMenuExtenders;
 	TArray<TSharedPtr<FExtender>> LevelEditorToolbarCinematicsMenuExtenders;
+	TArray<TSharedPtr<FExtender>> LevelEditorToolbarBlueprintsMenuExtenders;
 	TArray<FLevelEditorMenuExtender> LevelEditorLevelMenuExtenders;
 
 	/* Pointer to the current level Editor instance */
