Fix level editor toolbar missing the 'game' section
#rnx

diff --git a/Engine/Source/Editor/LevelEditor/Private/LevelEditorToolBar.cpp b/Engine/Source/Editor/LevelEditor/Private/LevelEditorToolBar.cpp
--- a/Engine/Source/Editor/LevelEditor/Private/LevelEditorToolBar.cpp
+++ b/Engine/Source/Editor/LevelEditor/Private/LevelEditorToolBar.cpp
@@ -1482,11 +1482,15 @@ void FLevelEditorToolBar::RegisterLevelEditorToolBar( const TSharedRef<FUIComman
 #endif
 		}
 	}
-	
+
 	ToolbarBuilder->AddDynamicSection("Game", FNewToolBarDelegateLegacy::CreateLambda([](FToolBarBuilder& InBuilder, UToolMenu* InData)
 	{
+		InBuilder.BeginSection("Game");
+
 		// Add the shared play-world commands that will be shown on the Kismet toolbar as well
 		FPlayWorldCommands::BuildToolbar(InBuilder, true);
+
+		InBuilder.EndSection();
 	}));
 
 #undef LOCTEXT_NAMESPACE
