Fixed ray tracing view mode is not saved correctly between sessions

diff --git a/Engine/Source/Editor/LevelEditor/Private/SLevelViewport.cpp b/Engine/Source/Editor/LevelEditor/Private/SLevelViewport.cpp
--- a/Engine/Source/Editor/LevelEditor/Private/SLevelViewport.cpp
+++ b/Engine/Source/Editor/LevelEditor/Private/SLevelViewport.cpp
@@ -497,6 +497,7 @@ void SLevelViewport::ConstructLevelEditorViewportClient( const FArguments& InArg
 	LevelViewportClient->EngineShowFlags = EditorShowFlags;
 	LevelViewportClient->LastEngineShowFlags = GameShowFlags;
 	LevelViewportClient->CurrentBufferVisualizationMode = ViewportInstanceSettings.BufferVisualizationMode;
+	LevelViewportClient->CurrentRayTracingDebugVisualizationMode = ViewportInstanceSettings.RayTracingDebugVisualizationMode;
 	LevelViewportClient->ExposureSettings = ViewportInstanceSettings.ExposureSettings;
 	if(InArgs._ViewportType == LVT_Perspective)
 	{
@@ -2148,6 +2149,7 @@ void SLevelViewport::SaveConfig(const FString& ConfigName) const
 		ViewportInstanceSettings.EditorShowFlagsString = EditorShowFlagsToSave.ToString();
 		ViewportInstanceSettings.GameShowFlagsString = GameShowFlagsToSave.ToString();
 		ViewportInstanceSettings.BufferVisualizationMode = LevelViewportClient->CurrentBufferVisualizationMode;
+		ViewportInstanceSettings.RayTracingDebugVisualizationMode = LevelViewportClient->CurrentRayTracingDebugVisualizationMode;
 		ViewportInstanceSettings.ExposureSettings = LevelViewportClient->ExposureSettings;
 		ViewportInstanceSettings.FOVAngle = LevelViewportClient->FOVAngle;
 		if (!FPlatformMisc::IsRemoteSession())
diff --git a/Engine/Source/Editor/UnrealEd/Classes/Settings/LevelEditorViewportSettings.h b/Engine/Source/Editor/UnrealEd/Classes/Settings/LevelEditorViewportSettings.h
--- a/Engine/Source/Editor/UnrealEd/Classes/Settings/LevelEditorViewportSettings.h
+++ b/Engine/Source/Editor/UnrealEd/Classes/Settings/LevelEditorViewportSettings.h
@@ -126,6 +126,10 @@ struct UNREALED_API FLevelEditorViewportInstanceSettings
 	UPROPERTY(config)
 	FName BufferVisualizationMode;
 
+	/** The buffer visualization mode for the viewport. */
+	UPROPERTY(config)
+	FName RayTracingDebugVisualizationMode;
+
 	/** Setting to allow designers to override the automatic expose. */
 	UPROPERTY(config)
 	FExposureSettings ExposureSettings;
diff --git a/Engine/Source/Editor/UnrealEd/Private/EditorViewportClient.cpp b/Engine/Source/Editor/UnrealEd/Private/EditorViewportClient.cpp
--- a/Engine/Source/Editor/UnrealEd/Private/EditorViewportClient.cpp
+++ b/Engine/Source/Editor/UnrealEd/Private/EditorViewportClient.cpp
@@ -311,6 +311,7 @@ FEditorViewportClient::FEditorViewportClient(FEditorModeTools* InModeTools, FPre
 	, LastEngineShowFlags(ESFIM_Game)
 	, ExposureSettings()
 	, CurrentBufferVisualizationMode(NAME_None)
+	, CurrentRayTracingDebugVisualizationMode(NAME_None)
 	, FramesSinceLastDraw(0)
 	, ViewIndex(INDEX_NONE)
 	, ViewFOV(EditorViewportDefs::DefaultPerspectiveFOVAngle)
