Fixed small issue with Level Editor CreateActor menu; Place Actor and Replace With Actor were being added in the same section with the same name, resulting in one of them being overwritten. Put them in different sections, for consistency with the block below.

diff --git a/Engine/Source/Editor/LevelEditor/Private/LevelEditorCreateActorMenu.cpp b/Engine/Source/Editor/LevelEditor/Private/LevelEditorCreateActorMenu.cpp
--- a/Engine/Source/Editor/LevelEditor/Private/LevelEditorCreateActorMenu.cpp
+++ b/Engine/Source/Editor/LevelEditor/Private/LevelEditorCreateActorMenu.cpp
@@ -361,21 +361,22 @@ void LevelEditorCreateActorMenu::FillAddReplaceViewportContextMenuSections(UTool
 	if ( AssetMenuOptions.Num() == 0 )
 	{
 		{
-			FToolMenuSection& Section = Menu->AddSection("ActorType");
+			FToolMenuSection& Section = Menu->AddSection("AddActor");
 			Section.AddSubMenu(
 				"LevelViewportContextMenu",
 				NSLOCTEXT("LevelViewportContextMenu", "AddActorHeading", "Place Actor") , 
 				NSLOCTEXT("LevelViewportContextMenu", "AddActorMenu_ToolTip", "Templates for adding a new actor to the world"),
 				FNewToolMenuDelegate::CreateStatic(&LevelEditorCreateActorMenu::FillAddReplaceActorMenu, EActorCreateMode::Add));
+		}
 
-			if ( CanReplaceActors() )
-			{
-				Section.AddSubMenu(
-					"LevelViewportContextMenu",
-					NSLOCTEXT("LevelViewportContextMenu", "ReplaceActorHeading", "Replace Selected Actors with") , 
-					NSLOCTEXT("LevelViewportContextMenu", "ReplaceActorMenu_ToolTip", "Templates for replacing selected with new actors in the world"),
-					FNewToolMenuDelegate::CreateStatic(&LevelEditorCreateActorMenu::FillAddReplaceActorMenu, EActorCreateMode::Replace));
-			}
+		if ( CanReplaceActors() )
+		{
+			FToolMenuSection& Section = Menu->AddSection("ReplaceActor");
+			Section.AddSubMenu(
+				"LevelViewportContextMenu",
+				NSLOCTEXT("LevelViewportContextMenu", "ReplaceActorHeading", "Replace Selected Actors with") , 
+				NSLOCTEXT("LevelViewportContextMenu", "ReplaceActorMenu_ToolTip", "Templates for replacing selected with new actors in the world"),
+				FNewToolMenuDelegate::CreateStatic(&LevelEditorCreateActorMenu::FillAddReplaceActorMenu, EActorCreateMode::Replace));
 		}
 	}
 	else
