Follow up to Level Editor Create Actor menu fix  CL 9844484

diff --git a/Engine/Source/Editor/LevelEditor/Private/LevelEditorCreateActorMenu.cpp b/Engine/Source/Editor/LevelEditor/Private/LevelEditorCreateActorMenu.cpp
--- a/Engine/Source/Editor/LevelEditor/Private/LevelEditorCreateActorMenu.cpp
+++ b/Engine/Source/Editor/LevelEditor/Private/LevelEditorCreateActorMenu.cpp
@@ -363,7 +363,7 @@ void LevelEditorCreateActorMenu::FillAddReplaceViewportContextMenuSections(UTool
 		{
 			FToolMenuSection& Section = Menu->AddSection("AddActor");
 			Section.AddSubMenu(
-				"LevelViewportContextMenu",
+				"AddActor",
 				NSLOCTEXT("LevelViewportContextMenu", "AddActorHeading", "Place Actor") , 
 				NSLOCTEXT("LevelViewportContextMenu", "AddActorMenu_ToolTip", "Templates for adding a new actor to the world"),
 				FNewToolMenuDelegate::CreateStatic(&LevelEditorCreateActorMenu::FillAddReplaceActorMenu, EActorCreateMode::Add));
@@ -373,7 +373,7 @@ void LevelEditorCreateActorMenu::FillAddReplaceViewportContextMenuSections(UTool
 		{
 			FToolMenuSection& Section = Menu->AddSection("ReplaceActor");
 			Section.AddSubMenu(
-				"LevelViewportContextMenu",
+				"ReplaceActor",
 				NSLOCTEXT("LevelViewportContextMenu", "ReplaceActorHeading", "Replace Selected Actors with") , 
 				NSLOCTEXT("LevelViewportContextMenu", "ReplaceActorMenu_ToolTip", "Templates for replacing selected with new actors in the world"),
 				FNewToolMenuDelegate::CreateStatic(&LevelEditorCreateActorMenu::FillAddReplaceActorMenu, EActorCreateMode::Replace));
