Brought over changes related to Tool Menus
#rnx

diff --git a/Engine/Plugins/Experimental/ControlRig/Source/ControlRigEditor/Private/Editor/ControlRigEditorMode.cpp b/Engine/Plugins/Experimental/ControlRig/Source/ControlRigEditor/Private/Editor/ControlRigEditorMode.cpp
--- a/Engine/Plugins/Experimental/ControlRig/Source/ControlRigEditor/Private/Editor/ControlRigEditorMode.cpp
+++ b/Engine/Plugins/Experimental/ControlRig/Source/ControlRigEditor/Private/Editor/ControlRigEditorMode.cpp
@@ -9,6 +9,7 @@
 #include "RigHierarchyTabSummoner.h"
 #include "RigStackTabSummoner.h"
 #include "RigCurveContainerTabSummoner.h"
+#include "ToolMenus.h"
 
 FControlRigEditorMode::FControlRigEditorMode(const TSharedRef<FControlRigEditor>& InControlRigEditor)
 	: FBlueprintEditorApplicationMode(InControlRigEditor, FControlRigEditorModes::ControlRigEditorMode, FControlRigEditorModes::GetLocalizedMode, false, false)
@@ -116,12 +117,12 @@ FControlRigEditorMode::FControlRigEditorMode(const TSharedRef<FControlRigEditor>
 			)
 		);
 
-	// setup toolbar - clear existing toolbar extender from the BP mode
-	ToolbarExtender = MakeShareable(new FExtender);
-	InControlRigEditor->GetToolbarBuilder()->AddCompileToolbar(ToolbarExtender);
-	InControlRigEditor->GetToolbarBuilder()->AddScriptingToolbar(ToolbarExtender);
-	InControlRigEditor->GetToolbarBuilder()->AddBlueprintGlobalOptionsToolbar(ToolbarExtender);
-//	InControlRigEditor->GetToolbarBuilder()->AddDebuggingToolbar(ToolbarExtender);
+	if (UToolMenu* Toolbar = InControlRigEditor->RegisterModeToolbarIfUnregistered(GetModeName()))
+	{
+		InControlRigEditor->GetToolbarBuilder()->AddCompileToolbar(Toolbar);
+		InControlRigEditor->GetToolbarBuilder()->AddScriptingToolbar(Toolbar);
+		InControlRigEditor->GetToolbarBuilder()->AddBlueprintGlobalOptionsToolbar(Toolbar);
+	}
 }
 
 void FControlRigEditorMode::RegisterTabFactories(TSharedPtr<FTabManager> InTabManager)
diff --git a/Engine/Plugins/Media/MediaPlayerEditor/Source/MediaPlayerEditor/Private/Widgets/SMediaPlayerEditorMedia.cpp b/Engine/Plugins/Media/MediaPlayerEditor/Source/MediaPlayerEditor/Private/Widgets/SMediaPlayerEditorMedia.cpp
--- a/Engine/Plugins/Media/MediaPlayerEditor/Source/MediaPlayerEditor/Private/Widgets/SMediaPlayerEditorMedia.cpp
+++ b/Engine/Plugins/Media/MediaPlayerEditor/Source/MediaPlayerEditor/Private/Widgets/SMediaPlayerEditorMedia.cpp
@@ -139,11 +139,11 @@ void SMediaPlayerEditorMedia::RegisterMenus()
 			FSlateIcon(FEditorStyle::GetStyleSetName(), "ContentBrowser.AssetActions.Edit"),
 			FToolMenuExecuteAction::CreateLambda([](const FToolMenuContext& InContext)
 			{
-				if (UMediaPlayerEditorMediaContext* Context = InContext.Find<UMediaPlayerEditorMediaContext>())
+				if (UMediaPlayerEditorMediaContext* Context = InContext.FindContext<UMediaPlayerEditorMediaContext>())
 				{
 					if (Context->SelectedAsset)
 					{
-						GEditor->GetEditorSubsystem<UAssetEditorSubsystem>()->OpenEditorForAsset(Context->SelectedAsset);
+						FAssetEditorManager::Get().OpenEditorForAsset(Context->SelectedAsset);
 					}
 				}
 			})
@@ -161,7 +161,7 @@ void SMediaPlayerEditorMedia::RegisterMenus()
 					FSlateIcon(Context->StyleSetName, "MediaPlayerEditor.NextMedia.Small"),
 					FToolMenuExecuteAction::CreateLambda([=](const FToolMenuContext& InContext)
 					{
-						if (UMediaPlayerEditorMediaContext* Context = InContext.Find<UMediaPlayerEditorMediaContext>())
+						if (UMediaPlayerEditorMediaContext* Context = InContext.FindContext<UMediaPlayerEditorMediaContext>())
 						{
 							if (Context->MediaPlayerEditorMedia.IsValid())
 							{
@@ -183,7 +183,7 @@ void SMediaPlayerEditorMedia::RegisterMenus()
 			FSlateIcon(FEditorStyle::GetStyleSetName(), "SystemWideCommands.FindInContentBrowser"),
 			FToolMenuExecuteAction::CreateLambda([](const FToolMenuContext& InContext)
 			{
-				if (UMediaPlayerEditorMediaContext* Context = InContext.Find<UMediaPlayerEditorMediaContext>())
+				if (UMediaPlayerEditorMediaContext* Context = InContext.FindContext<UMediaPlayerEditorMediaContext>())
 				{
 					TArray<UObject*> AssetsToSync = TArrayBuilder<UObject*>().Add(Context->SelectedAsset);
 					GEditor->SyncBrowserToObjects(AssetsToSync);
diff --git a/Engine/Plugins/Runtime/SoundUtilities/Source/SoundUtilitiesEditor/Private/SoundWaveAssetActionExtender.cpp b/Engine/Plugins/Runtime/SoundUtilities/Source/SoundUtilitiesEditor/Private/SoundWaveAssetActionExtender.cpp
--- a/Engine/Plugins/Runtime/SoundUtilities/Source/SoundUtilitiesEditor/Private/SoundWaveAssetActionExtender.cpp
+++ b/Engine/Plugins/Runtime/SoundUtilities/Source/SoundUtilitiesEditor/Private/SoundWaveAssetActionExtender.cpp
@@ -38,7 +38,7 @@ void FSoundWaveAssetActionExtender::RegisterMenus()
 
 void FSoundWaveAssetActionExtender::ExecuteCreateSimpleSound(const FToolMenuContext& MenuContext)
 {
-	UContentBrowserAssetContextMenuContext* Context = MenuContext.Find<UContentBrowserAssetContextMenuContext>();
+	UContentBrowserAssetContextMenuContext* Context = MenuContext.FindContext<UContentBrowserAssetContextMenuContext>();
 	if (!Context || Context->SelectedObjects.Num() == 0)
 	{
 		return;
diff --git a/Engine/Plugins/Runtime/TimeSynth/Source/TimeSynthEditor/Private/TimeSynthSoundWaveAssetActionExtender.cpp b/Engine/Plugins/Runtime/TimeSynth/Source/TimeSynthEditor/Private/TimeSynthSoundWaveAssetActionExtender.cpp
--- a/Engine/Plugins/Runtime/TimeSynth/Source/TimeSynthEditor/Private/TimeSynthSoundWaveAssetActionExtender.cpp
+++ b/Engine/Plugins/Runtime/TimeSynth/Source/TimeSynthEditor/Private/TimeSynthSoundWaveAssetActionExtender.cpp
@@ -57,7 +57,7 @@ void FTimeSynthSoundWaveAssetActionExtender::RegisterMenus()
 
 void FTimeSynthSoundWaveAssetActionExtender::ExecuteCreateTimeSyncClip(const FToolMenuContext& MenuContext)
 {
-	UContentBrowserAssetContextMenuContext* Context = MenuContext.Find<UContentBrowserAssetContextMenuContext>();
+	UContentBrowserAssetContextMenuContext* Context = MenuContext.FindContext<UContentBrowserAssetContextMenuContext>();
 	if (!Context || Context->SelectedObjects.Num() == 0)
 	{
 		return;
@@ -84,7 +84,7 @@ void FTimeSynthSoundWaveAssetActionExtender::ExecuteCreateTimeSyncClip(const FTo
 
 void FTimeSynthSoundWaveAssetActionExtender::ExecuteCreateTimeSyncClipSet(const FToolMenuContext& MenuContext)
 {
-	UContentBrowserAssetContextMenuContext* Context = MenuContext.Find<UContentBrowserAssetContextMenuContext>();
+	UContentBrowserAssetContextMenuContext* Context = MenuContext.FindContext<UContentBrowserAssetContextMenuContext>();
 	if (!Context || Context->SelectedObjects.Num() == 0)
 	{
 		return;
diff --git a/Engine/Source/Developer/MeshMergeUtilities/MeshMergeUtilities.Build.cs b/Engine/Source/Developer/MeshMergeUtilities/MeshMergeUtilities.Build.cs
--- a/Engine/Source/Developer/MeshMergeUtilities/MeshMergeUtilities.Build.cs
+++ b/Engine/Source/Developer/MeshMergeUtilities/MeshMergeUtilities.Build.cs
@@ -24,6 +24,7 @@ public class MeshMergeUtilities : ModuleRules
                 "SkeletalMeshEditor",
                 "MaterialBaking",
                 "MeshUtilitiesCommon",
+				"ToolMenus",
             }
 		);
         
diff --git a/Engine/Source/Developer/MeshMergeUtilities/Private/MeshMergeEditorExtensions.cpp b/Engine/Source/Developer/MeshMergeUtilities/Private/MeshMergeEditorExtensions.cpp
--- a/Engine/Source/Developer/MeshMergeUtilities/Private/MeshMergeEditorExtensions.cpp
+++ b/Engine/Source/Developer/MeshMergeUtilities/Private/MeshMergeEditorExtensions.cpp
@@ -17,10 +17,11 @@
 #include "MeshMergeUtilities.h"
 #include "Modules/ModuleManager.h"
 #include "MeshMergeModule.h"
+#include "ToolMenus.h"
+#include "SkeletalMeshToolMenuContext.h"
 
 #define LOCTEXT_NAMESPACE "MeshMergeEditorExtensions"
 
-FDelegateHandle FMeshMergeEditorExtensions::SkeletalMeshEditorExtenderHandle;
 FDelegateHandle FMeshMergeEditorExtensions::StaticMeshEditorExtenderHandle;
 
 void FMeshMergeEditorExtensions::OnModulesChanged(FName InModuleName, EModuleChangeReason InChangeReason)
@@ -28,11 +29,7 @@ void FMeshMergeEditorExtensions::OnModulesChanged(FName InModuleName, EModuleCha
 	// If one of the modules we are interested in is loaded apply editor extensions
 	if (InChangeReason == EModuleChangeReason::ModuleLoaded)
 	{		
-		if (InModuleName == "SkeletalMeshEditor")
-		{
-			AddSkeletalMeshEditorToolbarExtender();
-		}
-		else if (InModuleName == "StaticMeshEditor")
+		if (InModuleName == "StaticMeshEditor")
 		{
 			AddStaticMeshEditorToolbarExtender();
 		}
@@ -41,7 +38,6 @@ void FMeshMergeEditorExtensions::OnModulesChanged(FName InModuleName, EModuleCha
 
 void FMeshMergeEditorExtensions::RemoveExtenders()
 {
-	RemoveSkeletalMeshEditorToolbarExtender();
 	RemoveStaticMeshEditorToolbarExtender();
 }
 
@@ -77,7 +73,7 @@ void FMeshMergeEditorExtensions::RemoveStaticMeshEditorToolbarExtender()
 
 	if (StaticMeshEditorModule)
 	{
-		StaticMeshEditorModule->GetToolBarExtensibilityManager()->GetExtenderDelegates().RemoveAll([=](const auto& In) { return In.GetHandle() == SkeletalMeshEditorExtenderHandle; });
+		StaticMeshEditorModule->GetToolBarExtensibilityManager()->GetExtenderDelegates().RemoveAll([=](const auto& In) { return In.GetHandle() == StaticMeshEditorExtenderHandle; });
 	}
 }
 
@@ -96,57 +92,30 @@ void FMeshMergeEditorExtensions::HandleAddStaticMeshActionExtenderToToolbar(FToo
 	);
 }
 
-void FMeshMergeEditorExtensions::AddSkeletalMeshEditorToolbarExtender()
-{
-	ISkeletalMeshEditorModule& SkeletalMeshEditorModule = FModuleManager::Get().LoadModuleChecked<ISkeletalMeshEditorModule>("SkeletalMeshEditor");
-	auto& ToolbarExtenders = SkeletalMeshEditorModule.GetAllSkeletalMeshEditorToolbarExtenders();
-
-	ToolbarExtenders.Add(ISkeletalMeshEditorModule::FSkeletalMeshEditorToolbarExtender::CreateStatic(&FMeshMergeEditorExtensions::GetSkeletalMeshEditorToolbarExtender));
-	SkeletalMeshEditorExtenderHandle = ToolbarExtenders.Last().GetHandle();
-}
-
-void FMeshMergeEditorExtensions::RemoveSkeletalMeshEditorToolbarExtender()
+void FMeshMergeEditorExtensions::RegisterMenus()
 {
-	ISkeletalMeshEditorModule* SkeletalMeshEditorModule = FModuleManager::Get().GetModulePtr<ISkeletalMeshEditorModule>("SkeletalMeshEditor");
-	if (SkeletalMeshEditorModule)
 	{
-		typedef ISkeletalMeshEditorModule::FSkeletalMeshEditorToolbarExtender DelegateType;
-		SkeletalMeshEditorModule->GetAllSkeletalMeshEditorToolbarExtenders().RemoveAll([=](const DelegateType& In) { return In.GetHandle() == SkeletalMeshEditorExtenderHandle; });
-	}
-}
-
-TSharedRef<FExtender> FMeshMergeEditorExtensions::GetSkeletalMeshEditorToolbarExtender(const TSharedRef<FUICommandList> CommandList, TSharedRef<ISkeletalMeshEditor> InSkeletalMeshEditor)
-{
-	TSharedRef<FExtender> Extender = MakeShareable(new FExtender);
-	UMeshComponent* MeshComponent = InSkeletalMeshEditor->GetPersonaToolkit()->GetPreviewMeshComponent();
-
-	// Add button on skeletal mesh editor toolbar
-	Extender->AddToolBarExtension(
-		"Asset",
-		EExtensionHook::After,
-		CommandList,
-		FToolBarExtensionDelegate::CreateStatic(&FMeshMergeEditorExtensions::HandleAddSkeletalMeshActionExtenderToToolbar, MeshComponent)
-	);
-
-	return Extender;
-}
-
-void FMeshMergeEditorExtensions::HandleAddSkeletalMeshActionExtenderToToolbar(FToolBarBuilder& ParentToolbarBuilder, UMeshComponent* InMeshComponent)
-{
-	ParentToolbarBuilder.AddToolBarButton(
-		FUIAction(FExecuteAction::CreateLambda([InMeshComponent]()
-		{
-			if (USkeletalMeshComponent* SkeletalMeshComponent = Cast<USkeletalMeshComponent>(InMeshComponent))
+		UToolMenu* Toolbar = UToolMenus::Get()->ExtendMenu("AssetEditor.SkeletalMeshEditor.ToolBar");
+		FToolMenuSection& Section = Toolbar->FindOrAddSection("SkeletalMesh");
+		Section.AddEntry(FToolMenuEntry::InitToolBarButton(
+			"BakeMaterials",
+			FToolMenuExecuteAction::CreateLambda([](const FToolMenuContext& InMenuContext)
 			{
-				IMeshMergeModule& Module = FModuleManager::Get().LoadModuleChecked<IMeshMergeModule>("MeshMergeUtilities");				
-				Module.GetUtilities().BakeMaterialsForComponent(SkeletalMeshComponent);
-			}
-		})),
-		NAME_None,
-		LOCTEXT("BakeMaterials", "Bake out Materials"),
-		LOCTEXT("BakeMaterialsTooltip", "Bake out Materials for given LOD(s)."),
-		FSlateIcon("EditorStyle", "Persona.BakeMaterials")
-	);
+				USkeletalMeshToolMenuContext* Context = InMenuContext.FindContext<USkeletalMeshToolMenuContext>();
+				if (Context && Context->SkeletalMeshEditor.IsValid())
+				{
+					if (UDebugSkelMeshComponent* SkeletalMeshComponent = Context->SkeletalMeshEditor.Pin()->GetPersonaToolkit()->GetPreviewMeshComponent())
+					{
+						IMeshMergeModule& Module = FModuleManager::Get().LoadModuleChecked<IMeshMergeModule>("MeshMergeUtilities");
+						Module.GetUtilities().BakeMaterialsForComponent(SkeletalMeshComponent);
+					}
+				}
+			}),
+			LOCTEXT("BakeMaterials", "Bake out Materials"),
+			LOCTEXT("BakeMaterialsTooltip", "Bake out Materials for given LOD(s)."),
+			FSlateIcon("EditorStyle", "Persona.BakeMaterials")
+		));
+	}
 }
 
 #undef LOCTEXT_NAMESPACE // "MeshMergeEditorExtensions"
\ No newline at end of file
diff --git a/Engine/Source/Developer/MeshMergeUtilities/Private/MeshMergeEditorExtensions.h b/Engine/Source/Developer/MeshMergeUtilities/Private/MeshMergeEditorExtensions.h
--- a/Engine/Source/Developer/MeshMergeUtilities/Private/MeshMergeEditorExtensions.h
+++ b/Engine/Source/Developer/MeshMergeUtilities/Private/MeshMergeEditorExtensions.h
@@ -20,19 +20,13 @@ class FMeshMergeEditorExtensions
 public:
 	static void OnModulesChanged(FName InModuleName, EModuleChangeReason InChangeReason);
 	static void RemoveExtenders();
+	static void RegisterMenus();
 	
 	/** Callback functionality for static mesh editor extension */
 	static void AddStaticMeshEditorToolbarExtender();
 	static void RemoveStaticMeshEditorToolbarExtender();
 	static void HandleAddStaticMeshActionExtenderToToolbar(FToolBarBuilder& ParentToolbarBuilder, UStaticMesh* Mesh);
 	static TSharedRef<FExtender> GetStaticMeshEditorToolbarExtender(const TSharedRef<FUICommandList> CommandList, const TArray<UObject*> Objects);
-	
-	/** Callback functionality for skeletal mesh editor extension */
-	static void AddSkeletalMeshEditorToolbarExtender();
-	static void RemoveSkeletalMeshEditorToolbarExtender();
-	static TSharedRef<FExtender> GetSkeletalMeshEditorToolbarExtender(const TSharedRef<FUICommandList> CommandList, TSharedRef<ISkeletalMeshEditor> InSkeletalMeshEditor);
-	static void HandleAddSkeletalMeshActionExtenderToToolbar(FToolBarBuilder& ParentToolbarBuilder, UMeshComponent* InMeshComponent);
-	
-	static FDelegateHandle SkeletalMeshEditorExtenderHandle;
+
 	static FDelegateHandle StaticMeshEditorExtenderHandle;
 };
\ No newline at end of file
diff --git a/Engine/Source/Developer/MeshMergeUtilities/Private/MeshMergeModule.cpp b/Engine/Source/Developer/MeshMergeUtilities/Private/MeshMergeModule.cpp
--- a/Engine/Source/Developer/MeshMergeUtilities/Private/MeshMergeModule.cpp
+++ b/Engine/Source/Developer/MeshMergeUtilities/Private/MeshMergeModule.cpp
@@ -3,6 +3,8 @@
 #include "MeshMergeModule.h"
 #include "MeshMergeUtilities.h"
 #include "Modules/ModuleManager.h"
+#include "MeshMergeEditorExtensions.h"
+#include "ToolMenus.h"
 
 class FMeshMergeModule : public IMeshMergeModule
 {
@@ -16,8 +18,27 @@ public:
 	{
 		return Utilities;
 	}
+
+	virtual void StartupModule() override
+	{
+		UToolMenus::RegisterStartupCallback(FSimpleMulticastDelegate::FDelegate::CreateRaw(this, &FMeshMergeModule::RegisterMenus));
+	}
+
+	virtual void ShutdownModule() override
+	{
+		UToolMenus::UnRegisterStartupCallback(this);
+		UToolMenus::UnregisterOwner(this);
+	}
+
 protected:
 	FMeshMergeUtilities Utilities;
+
+private:
+	void RegisterMenus()
+	{
+		FToolMenuOwnerScoped OwnerScoped(this);
+		FMeshMergeEditorExtensions::RegisterMenus();
+	}
 };
 
 
diff --git a/Engine/Source/Developer/MeshUtilities/MeshUtilities.Build.cs b/Engine/Source/Developer/MeshUtilities/MeshUtilities.Build.cs
--- a/Engine/Source/Developer/MeshUtilities/MeshUtilities.Build.cs
+++ b/Engine/Source/Developer/MeshUtilities/MeshUtilities.Build.cs
@@ -36,8 +36,8 @@ public class MeshUtilities : ModuleRules
                 "MeshBuilderCommon",
                 "MeshUtilitiesCommon",
                 "MeshDescription",
-				"StaticMeshDescription",
-                "MeshDescriptionOperations"
+                "MeshDescriptionOperations",
+				"ToolMenus",
             }
 		);
 
diff --git a/Engine/Source/Developer/MeshUtilities/Private/MeshUtilities.cpp b/Engine/Source/Developer/MeshUtilities/Private/MeshUtilities.cpp
--- a/Engine/Source/Developer/MeshUtilities/Private/MeshUtilities.cpp
+++ b/Engine/Source/Developer/MeshUtilities/Private/MeshUtilities.cpp
@@ -17,6 +17,8 @@
 #include "Framework/Commands/UICommandList.h"
 #include "Framework/MultiBox/MultiBoxExtender.h"
 #include "Framework/MultiBox/MultiBoxBuilder.h"
+#include "ToolMenus.h"
+#include "SkeletalMeshToolMenuContext.h"
 #include "Components/MeshComponent.h"
 #include "RawIndexBuffer.h"
 #include "Components/StaticMeshComponent.h"
@@ -5553,20 +5555,21 @@ void FMeshUtilities::StartupModule()
 			{
 				AddAnimationEditorToolbarExtender();
 			}
-			else if (InModuleName == "SkeletalMeshEditor")
-			{
-				AddSkeletalMeshEditorToolbarExtender();
-			}
 			else if (InModuleName == "SkeletonEditor")
 			{
 				AddSkeletonEditorToolbarExtender();
 			}
 		}
 	});
+
+	UToolMenus::RegisterStartupCallback(FSimpleMulticastDelegate::FDelegate::CreateRaw(this, &FMeshUtilities::RegisterMenus));
 }
 
 void FMeshUtilities::ShutdownModule()
 {
+	UToolMenus::UnRegisterStartupCallback(this);
+	UToolMenus::UnregisterOwner(this);
+
 	static const FName PropertyEditorModuleName("PropertyEditor");
 	if(FModuleManager::Get().IsModuleLoaded(PropertyEditorModuleName))
 	{
@@ -5580,12 +5583,42 @@ void FMeshUtilities::ShutdownModule()
 	RemoveLevelViewportMenuExtender();
 	RemoveAnimationBlueprintEditorToolbarExtender();
 	RemoveAnimationEditorToolbarExtender();
-	RemoveSkeletalMeshEditorToolbarExtender();
 	RemoveSkeletonEditorToolbarExtender();
 	FModuleManager::Get().OnModulesChanged().Remove(ModuleLoadedDelegateHandle);
 	VersionString.Empty();
 }
 
+void FMeshUtilities::RegisterMenus()
+{
+	FToolMenuOwnerScoped OwnerScoped(this);
+
+	static auto AddMakeStaticMeshToolbarButton = [this](FToolMenuSection& InSection, const FToolMenuExecuteAction& InAction)
+	{
+		InSection.AddEntry(FToolMenuEntry::InitToolBarButton(
+			"MakeStaticMesh",
+			InAction,
+			LOCTEXT("MakeStaticMesh", "Make Static Mesh"),
+			LOCTEXT("MakeStaticMeshTooltip", "Make a new static mesh out of the preview's current pose."),
+			FSlateIcon("EditorStyle", "Persona.ConvertToStaticMesh")
+		));
+	};
+
+	{
+		UToolMenu* Toolbar = UToolMenus::Get()->ExtendMenu("AssetEditor.SkeletalMeshEditor.ToolBar");
+		FToolMenuSection& Section = Toolbar->FindOrAddSection("SkeletalMesh");
+		AddMakeStaticMeshToolbarButton(Section, FToolMenuExecuteAction::CreateLambda([this](const FToolMenuContext& InMenuContext)
+		{
+			USkeletalMeshToolMenuContext* Context = InMenuContext.FindContext<USkeletalMeshToolMenuContext>();
+			if (Context && Context->SkeletalMeshEditor.IsValid())
+			{
+				if (UMeshComponent* MeshComponent = Context->SkeletalMeshEditor.Pin()->GetPersonaToolkit()->GetPreviewMeshComponent())
+				{
+					ConvertMeshesToStaticMesh(TArray<UMeshComponent*>({ MeshComponent }), MeshComponent->GetComponentToWorld());
+				}
+			}
+		}));
+	}
+}
 
 bool FMeshUtilities::GenerateUniqueUVsForSkeletalMesh(const FSkeletalMeshLODModel& LODModel, int32 TextureResolution, TArray<FVector2D>& OutTexCoords) const
 {
@@ -5845,25 +5878,6 @@ TSharedRef<FExtender> FMeshUtilities::GetAnimationEditorToolbarExtender(const TS
 	return Extender;
 }
 
-void FMeshUtilities::AddSkeletalMeshEditorToolbarExtender()
-{
-	ISkeletalMeshEditorModule& SkeletalMeshEditorModule = FModuleManager::Get().LoadModuleChecked<ISkeletalMeshEditorModule>("SkeletalMeshEditor");
-	auto& ToolbarExtenders = SkeletalMeshEditorModule.GetAllSkeletalMeshEditorToolbarExtenders();
-
-	ToolbarExtenders.Add(ISkeletalMeshEditorModule::FSkeletalMeshEditorToolbarExtender::CreateRaw(this, &FMeshUtilities::GetSkeletalMeshEditorToolbarExtender));
-	SkeletalMeshEditorExtenderHandle = ToolbarExtenders.Last().GetHandle();
-}
-
-void FMeshUtilities::RemoveSkeletalMeshEditorToolbarExtender()
-{
-	ISkeletalMeshEditorModule* SkeletalMeshEditorModule = FModuleManager::Get().GetModulePtr<ISkeletalMeshEditorModule>("SkeletalMeshEditor");
-	if (SkeletalMeshEditorModule)
-	{
-		typedef ISkeletalMeshEditorModule::FSkeletalMeshEditorToolbarExtender DelegateType;
-		SkeletalMeshEditorModule->GetAllSkeletalMeshEditorToolbarExtenders().RemoveAll([=](const DelegateType& In) { return In.GetHandle() == SkeletalMeshEditorExtenderHandle; });
-	}
-}
-
 TSharedRef<FExtender> FMeshUtilities::GetSkeletalMeshEditorToolbarExtender(const TSharedRef<FUICommandList> CommandList, TSharedRef<ISkeletalMeshEditor> InSkeletalMeshEditor)
 {
 	TSharedRef<FExtender> Extender = MakeShareable(new FExtender);
diff --git a/Engine/Source/Developer/MeshUtilities/Private/MeshUtilitiesPrivate.h b/Engine/Source/Developer/MeshUtilities/Private/MeshUtilitiesPrivate.h
--- a/Engine/Source/Developer/MeshUtilities/Private/MeshUtilitiesPrivate.h
+++ b/Engine/Source/Developer/MeshUtilities/Private/MeshUtilitiesPrivate.h
@@ -166,6 +166,8 @@ private:
 
 	virtual void GenerateRuntimeSkinWeightData(const FSkeletalMeshLODModel* ImportedModel, const TArray<FRawSkinWeight>& InRawSkinWeights, FRuntimeSkinWeightProfileData& InOutSkinWeightOverrideData) const override;
 
+	void RegisterMenus();
+
 	// Need to call some members from this class, (which is internal to this module)
 	friend class FStaticMeshUtilityBuilder;
 
@@ -182,10 +184,6 @@ protected:
 
 	TSharedRef<FExtender> GetAnimationEditorToolbarExtender(const TSharedRef<FUICommandList> CommandList, TSharedRef<IAnimationEditor> InAnimationEditor);
 
-	void AddSkeletalMeshEditorToolbarExtender();
-
-	void RemoveSkeletalMeshEditorToolbarExtender();
-
 	TSharedRef<FExtender> GetSkeletalMeshEditorToolbarExtender(const TSharedRef<FUICommandList> CommandList, TSharedRef<ISkeletalMeshEditor> InSkeletalMeshEditor);
 
 	void AddSkeletonEditorToolbarExtender();
@@ -208,7 +206,6 @@ protected:
 	FDelegateHandle LevelViewportExtenderHandle;
 	FDelegateHandle AnimationBlueprintEditorExtenderHandle;
 	FDelegateHandle AnimationEditorExtenderHandle;
-	FDelegateHandle SkeletalMeshEditorExtenderHandle;
 	FDelegateHandle SkeletonEditorExtenderHandle;
 };
 
diff --git a/Engine/Source/Editor/AnimationBlueprintEditor/AnimationBlueprintEditor.Build.cs b/Engine/Source/Editor/AnimationBlueprintEditor/AnimationBlueprintEditor.Build.cs
--- a/Engine/Source/Editor/AnimationBlueprintEditor/AnimationBlueprintEditor.Build.cs
+++ b/Engine/Source/Editor/AnimationBlueprintEditor/AnimationBlueprintEditor.Build.cs
@@ -39,6 +39,7 @@ public class AnimationBlueprintEditor : ModuleRules
                 "BlueprintGraph",
                 "RHI",
                 "KismetCompiler",
+				"ToolMenus",
             }
 		);
 
diff --git a/Engine/Source/Editor/AnimationBlueprintEditor/Private/AnimationBlueprintEditorMode.cpp b/Engine/Source/Editor/AnimationBlueprintEditor/Private/AnimationBlueprintEditorMode.cpp
--- a/Engine/Source/Editor/AnimationBlueprintEditor/Private/AnimationBlueprintEditorMode.cpp
+++ b/Engine/Source/Editor/AnimationBlueprintEditor/Private/AnimationBlueprintEditorMode.cpp
@@ -14,6 +14,8 @@
 #include "IPersonaPreviewScene.h"
 #include "PersonaUtils.h"
 
+#include "ToolMenus.h"
+
 /////////////////////////////////////////////////////
 // FAnimationBlueprintEditorMode
 
@@ -132,7 +134,10 @@ FAnimationBlueprintEditorMode::FAnimationBlueprintEditorMode(const TSharedRef<FA
 	PersonaModule.RegisterPersonaViewportTabFactories(TabFactories, InAnimationBlueprintEditor, ViewportArgs);
 
 	TabFactories.RegisterFactory(PersonaModule.CreateAdvancedPreviewSceneTabFactory(InAnimationBlueprintEditor, InAnimationBlueprintEditor->GetPersonaToolkit()->GetPreviewScene()));
+#pragma message("TODO: CLEAN UP MERGED CODE")
+#if 0
 	TabFactories.RegisterFactory(PersonaModule.CreateAnimationAssetBrowserTabFactory(InAnimationBlueprintEditor, InAnimationBlueprintEditor->GetPersonaToolkit(), FOnOpenNewAsset::CreateSP(&InAnimationBlueprintEditor.Get(), &FAnimationBlueprintEditor::HandleOpenNewAsset), FOnAnimationSequenceBrowserCreated::CreateSP(&InAnimationBlueprintEditor.Get(), &FAnimationBlueprintEditor::HandleAnimationSequenceBrowserCreated), true));
+#endif
 	TabFactories.RegisterFactory(PersonaModule.CreateAnimBlueprintPreviewTabFactory(InAnimationBlueprintEditor, InAnimationBlueprintEditor->GetPersonaToolkit()->GetPreviewScene()));
 	TabFactories.RegisterFactory(PersonaModule.CreateAnimBlueprintAssetOverridesTabFactory(InAnimationBlueprintEditor, InAnimationBlueprintEditor->GetPersonaToolkit()->GetAnimBlueprint(), InAnimationBlueprintEditor->OnPostUndo));
 	TabFactories.RegisterFactory(PersonaModule.CreateSkeletonSlotNamesTabFactory(InAnimationBlueprintEditor, InAnimationBlueprintEditor->GetSkeletonTree()->GetEditableSkeleton(), InAnimationBlueprintEditor->OnPostUndo, FOnObjectSelected::CreateSP(&InAnimationBlueprintEditor.Get(), &FAnimationBlueprintEditor::HandleObjectSelected)));
@@ -141,10 +146,14 @@ FAnimationBlueprintEditorMode::FAnimationBlueprintEditorMode(const TSharedRef<FA
 	// setup toolbar - clear existing toolbar extender from the BP mode
 	//@TODO: Keep this in sync with BlueprintEditorModes.cpp
 	ToolbarExtender = MakeShareable(new FExtender);
-	InAnimationBlueprintEditor->GetToolbarBuilder()->AddCompileToolbar(ToolbarExtender);
-	InAnimationBlueprintEditor->GetToolbarBuilder()->AddScriptingToolbar(ToolbarExtender);
-	InAnimationBlueprintEditor->GetToolbarBuilder()->AddBlueprintGlobalOptionsToolbar(ToolbarExtender);
-	InAnimationBlueprintEditor->GetToolbarBuilder()->AddDebuggingToolbar(ToolbarExtender);
+
+	if (UToolMenu* Toolbar = InAnimationBlueprintEditor->RegisterModeToolbarIfUnregistered(GetModeName()))
+	{
+		InAnimationBlueprintEditor->GetToolbarBuilder()->AddCompileToolbar(Toolbar);
+		InAnimationBlueprintEditor->GetToolbarBuilder()->AddScriptingToolbar(Toolbar);
+		InAnimationBlueprintEditor->GetToolbarBuilder()->AddBlueprintGlobalOptionsToolbar(Toolbar);
+		InAnimationBlueprintEditor->GetToolbarBuilder()->AddDebuggingToolbar(Toolbar);
+	}
 
 	PersonaModule.OnRegisterTabs().Broadcast(TabFactories, InAnimationBlueprintEditor);
 	LayoutExtender = MakeShared<FLayoutExtender>();
diff --git a/Engine/Source/Editor/AnimationBlueprintEditor/Private/AnimationBlueprintInterfaceEditorMode.cpp b/Engine/Source/Editor/AnimationBlueprintEditor/Private/AnimationBlueprintInterfaceEditorMode.cpp
--- a/Engine/Source/Editor/AnimationBlueprintEditor/Private/AnimationBlueprintInterfaceEditorMode.cpp
+++ b/Engine/Source/Editor/AnimationBlueprintEditor/Private/AnimationBlueprintInterfaceEditorMode.cpp
@@ -9,6 +9,7 @@
 #include "PersonaModule.h"
 #include "SBlueprintEditorToolbar.h"
 #include "IPersonaPreviewScene.h"
+#include "ToolMenus.h"
 
 FAnimationBlueprintInterfaceEditorMode::FAnimationBlueprintInterfaceEditorMode(const TSharedRef<FAnimationBlueprintEditor>& InAnimationBlueprintEditor)
 	: FBlueprintInterfaceApplicationMode(InAnimationBlueprintEditor, FAnimationBlueprintEditorModes::AnimationBlueprintInterfaceEditorMode, FAnimationBlueprintEditorModes::GetLocalizedMode)
@@ -83,8 +84,12 @@ FAnimationBlueprintInterfaceEditorMode::FAnimationBlueprintInterfaceEditorMode(c
 	FPersonaModule& PersonaModule = FModuleManager::LoadModuleChecked<FPersonaModule>("Persona");
 
 	ToolbarExtender = MakeShareable(new FExtender);
-	InAnimationBlueprintEditor->GetToolbarBuilder()->AddCompileToolbar(ToolbarExtender);
-	InAnimationBlueprintEditor->GetToolbarBuilder()->AddBlueprintGlobalOptionsToolbar(ToolbarExtender);
+
+	if (UToolMenu* Toolbar = InAnimationBlueprintEditor->RegisterModeToolbarIfUnregistered(GetModeName()))
+	{
+		InAnimationBlueprintEditor->GetToolbarBuilder()->AddCompileToolbar(Toolbar);
+		InAnimationBlueprintEditor->GetToolbarBuilder()->AddBlueprintGlobalOptionsToolbar(Toolbar);	
+	}
 
 	PersonaModule.OnRegisterTabs().Broadcast(TabFactories, InAnimationBlueprintEditor);
 	LayoutExtender = MakeShared<FLayoutExtender>();
diff --git a/Engine/Source/Editor/AudioEditor/Classes/SoundSubmixGraph/SoundSubmixGraphSchema.h b/Engine/Source/Editor/AudioEditor/Classes/SoundSubmixGraph/SoundSubmixGraphSchema.h
--- a/Engine/Source/Editor/AudioEditor/Classes/SoundSubmixGraph/SoundSubmixGraphSchema.h
+++ b/Engine/Source/Editor/AudioEditor/Classes/SoundSubmixGraph/SoundSubmixGraphSchema.h
@@ -91,7 +91,7 @@ struct AUDIOEDITOR_API FSoundSubmixGraphSchemaAction_NewNode : public FEdGraphSc
 
 	//~ Begin FEdGraphSchemaAction Interface
 	virtual FName GetTypeId() const override { return StaticGetTypeId(); } 
-	virtual UEdGraphNode* PerformAction(UEdGraph* ParentGraph, UEdGraphPin* FromPin, const FVector2D Location, bool bSelectNewNode = true) override;
+	virtual UEdGraphNode* PerformAction(class UEdGraph* ParentGraph, UEdGraphPin* FromPin, const FVector2D Location, bool bSelectNewNode = true) override;
 	//~ End FEdGraphSchemaAction Interface
 
 	/** Name for the new SoundSubmix */
@@ -120,7 +120,7 @@ class USoundSubmixGraphSchema : public UEdGraphSchema
 	virtual void BreakNodeLinks(UEdGraphNode& TargetNode) const override;
 	virtual void BreakPinLinks(UEdGraphPin& TargetPin, bool bSendsNodeNotifcation) const override;
 	virtual void BreakSinglePinLink(UEdGraphPin* SourcePin, UEdGraphPin* TargetPin) const override;
-	virtual void DroppedAssetsOnGraph(const TArray<FAssetData>& Assets, const FVector2D& GraphPosition, UEdGraph* Graph) const override;
+	virtual void DroppedAssetsOnGraph(const TArray<struct FAssetData>& Assets, const FVector2D& GraphPosition, UEdGraph* Graph) const override;
 	//~ End EdGraphSchema Interface
 };
 
diff --git a/Engine/Source/Editor/AudioEditor/Private/SoundSubmixGraphSchema.cpp b/Engine/Source/Editor/AudioEditor/Private/SoundSubmixGraphSchema.cpp
--- a/Engine/Source/Editor/AudioEditor/Private/SoundSubmixGraphSchema.cpp
+++ b/Engine/Source/Editor/AudioEditor/Private/SoundSubmixGraphSchema.cpp
@@ -8,6 +8,8 @@
 #include "EdGraphUtilities.h"
 #include "Framework/Commands/GenericCommands.h"
 #include "Framework/MultiBox/MultiBoxBuilder.h"
+#include "ToolMenus.h"
+
 #include "ScopedTransaction.h"
 #include "Sound/SoundSubmix.h"
 #include "SoundSubmixGraph/SoundSubmixGraphNode.h"
@@ -114,7 +116,7 @@ void USoundSubmixGraphSchema::GetBreakLinkToSubMenuActions(UToolMenu* Menu, cons
 	TMap<FString, uint32> LinkTitleCount;
 
 	// Add all the links we could break from
-	for (TArray<UEdGraphPin*>::TConstIterator Links(InGraphPin->LinkedTo); Links; ++Links)
+	for(TArray<class UEdGraphPin*>::TConstIterator Links(InGraphPin->LinkedTo); Links; ++Links)
 	{
 		UEdGraphPin* Pin = *Links;
 		FString TitleString = Pin->GetOwningNode()->GetNodeTitle(ENodeTitleType::ListView).ToString();
diff --git a/Engine/Source/Editor/ClothPainter/ClothPainter.Build.cs b/Engine/Source/Editor/ClothPainter/ClothPainter.Build.cs
--- a/Engine/Source/Editor/ClothPainter/ClothPainter.Build.cs
+++ b/Engine/Source/Editor/ClothPainter/ClothPainter.Build.cs
@@ -31,7 +31,8 @@ public class ClothPainter : ModuleRules
                 "ClothingSystemRuntimeInterface",
                 "ClothingSystemEditorInterface",
                 "SkeletalMeshEditor",
-                "AdvancedPreviewScene"
+                "AdvancedPreviewScene",
+				"ToolMenus",
             }
 		);
     }
diff --git a/Engine/Source/Editor/ClothPainter/Private/ClothPaintingModule.cpp b/Engine/Source/Editor/ClothPainter/Private/ClothPaintingModule.cpp
--- a/Engine/Source/Editor/ClothPainter/Private/ClothPaintingModule.cpp
+++ b/Engine/Source/Editor/ClothPainter/Private/ClothPaintingModule.cpp
@@ -18,6 +18,8 @@
 #include "ClothPaintToolCommands.h"
 #include "ISkeletalMeshEditorModule.h"
 #include "Framework/MultiBox/MultiBoxBuilder.h"
+#include "ToolMenus.h"
+#include "SkeletalMeshToolMenuContext.h"
 #include "ClothPainterCommands.h"
 #include "Widgets/Docking/SDockTab.h"
 
@@ -68,6 +70,8 @@ void FClothPaintingModule::StartupModule()
 	ClothPaintToolCommands::RegisterClothPaintToolCommands();
 	FClothPainterCommands::Register();
 
+	UToolMenus::RegisterStartupCallback(FSimpleMulticastDelegate::FDelegate::CreateRaw(this, &FClothPaintingModule::RegisterMenus));
+
 	if(!SkelMeshEditorExtenderHandle.IsValid())
 	{
 		ISkeletalMeshEditorModule& SkelMeshEditorModule = FModuleManager::Get().LoadModuleChecked<ISkeletalMeshEditorModule>("SkeletalMeshEditor");
@@ -112,21 +116,32 @@ TSharedRef<FExtender> FClothPaintingModule::ExtendSkelMeshEditorToolbar(const TS
 		FIsActionChecked::CreateRaw(this, &FClothPaintingModule::GetIsPaintToolsButtonChecked, Ptr)
 	);
 
-	ToolbarExtender->AddToolBarExtension("Asset", EExtensionHook::After, InCommandList, FToolBarExtensionDelegate::CreateLambda(
-		[this, Ptr](FToolBarBuilder& Builder)
-	{
-		Builder.AddToolBarButton(
-			FClothPainterCommands::Get().TogglePaintMode, 
-			NAME_None, 
-			TAttribute<FText>::Create(TAttribute<FText>::FGetter::CreateRaw(this, &FClothPaintingModule::GetPaintToolsButtonText, Ptr)), 
-			TAttribute<FText>(), 
-			FSlateIcon(FEditorStyle::GetStyleSetName(), 
-			"LevelEditor.MeshPaintMode.TexturePaint"));
-	}));
-
 	return ToolbarExtender.ToSharedRef();
 }
 
+void FClothPaintingModule::RegisterMenus()
+{
+	FToolMenuOwnerScoped OwnerScoped(this);
+
+	UToolMenu* Toolbar = UToolMenus::Get()->ExtendMenu("AssetEditor.SkeletalMeshEditor.ToolBar");
+	{
+		FToolMenuSection& Section = Toolbar->FindOrAddSection("SkeletalMesh");
+		Section.AddDynamicEntry("TogglePaintMode", FNewToolMenuSectionDelegate::CreateLambda([this](FToolMenuSection& InSection)
+		{
+			USkeletalMeshToolMenuContext* Context = InSection.FindContext<USkeletalMeshToolMenuContext>();
+			if (Context && Context->SkeletalMeshEditor.IsValid())
+			{
+				InSection.AddEntry(FToolMenuEntry::InitToolBarButton(
+					FClothPainterCommands::Get().TogglePaintMode,
+					TAttribute<FText>::Create(TAttribute<FText>::FGetter::CreateRaw(this, &FClothPaintingModule::GetPaintToolsButtonText, Context->SkeletalMeshEditor)),
+					FText(),
+					FSlateIcon(FEditorStyle::GetStyleSetName(), "LevelEditor.MeshPaintMode.TexturePaint")
+				));	
+			}
+		}));
+	}
+}
+
 FText FClothPaintingModule::GetPaintToolsButtonText(TWeakPtr<ISkeletalMeshEditor> InSkeletalMeshEditor) const
 {
 	TSharedPtr<SClothPaintTab> ClothTab = GetActiveClothTab(InSkeletalMeshEditor, false);
@@ -194,6 +209,9 @@ void FClothPaintingModule::ShutdownModule()
 {
 	ShutdownMode();
 
+	UToolMenus::UnRegisterStartupCallback(this);
+	UToolMenus::UnregisterOwner(this);
+
 	// Remove skel mesh editor extenders
 	ISkeletalMeshEditorModule* SkelMeshEditorModule = FModuleManager::GetModulePtr<ISkeletalMeshEditorModule>("SkeletalMeshEditor");
 	if (SkelMeshEditorModule)
diff --git a/Engine/Source/Editor/ClothPainter/Public/ClothPaintingModule.h b/Engine/Source/Editor/ClothPainter/Public/ClothPaintingModule.h
--- a/Engine/Source/Editor/ClothPainter/Public/ClothPaintingModule.h
+++ b/Engine/Source/Editor/ClothPainter/Public/ClothPaintingModule.h
@@ -36,6 +36,9 @@ protected:
 
 private:
 
+	// Extend toolbars and menus
+	void RegisterMenus();
+
 	// Extends a skeletal mesh editor instance toolbar
 	TSharedRef<FExtender> ExtendSkelMeshEditorToolbar(const TSharedRef<FUICommandList> InCommandList, TSharedRef<ISkeletalMeshEditor> InSkeletalMeshEditor);
 
diff --git a/Engine/Source/Editor/Kismet/Private/BlueprintEditor.cpp b/Engine/Source/Editor/Kismet/Private/BlueprintEditor.cpp
--- a/Engine/Source/Editor/Kismet/Private/BlueprintEditor.cpp
+++ b/Engine/Source/Editor/Kismet/Private/BlueprintEditor.cpp
@@ -99,6 +99,8 @@
 
 #include "BlueprintEditorTabs.h"
 
+#include "ToolMenus.h"
+#include "BlueprintEditorContext.h"
 
 #include "Interfaces/IProjectManager.h"
 
@@ -1895,6 +1897,15 @@ void FBlueprintEditor::InitBlueprintEditor(
 	}
 }
 
+void FBlueprintEditor::InitToolMenuContext(FToolMenuContext& MenuContext)
+{
+	FAssetEditorToolkit::InitToolMenuContext(MenuContext);
+
+	UBlueprintEditorToolMenuContext* Context = NewObject<UBlueprintEditorToolMenuContext>();
+	Context->BlueprintEditor = SharedThis(this);
+	MenuContext.AddObject(Context);
+}
+
 void FBlueprintEditor::InitalizeExtenders()
 {
 	TSharedPtr<FExtender> MenuExtender = MakeShareable(new FExtender);
@@ -9406,4 +9417,9 @@ void FBlueprintEditor::ClearAllGraphEditorQuickJumps()
 /////////////////////////////////////////////////////
 /////////////////////////////////////////////////////
 
+UBlueprint* UBlueprintEditorToolMenuContext::GetBlueprintObj() const
+{
+	return BlueprintEditor.IsValid() ? BlueprintEditor.Pin()->GetBlueprintObj() : nullptr;
+}
+
 #undef LOCTEXT_NAMESPACE
diff --git a/Engine/Source/Editor/Kismet/Private/BlueprintEditorModes.cpp b/Engine/Source/Editor/Kismet/Private/BlueprintEditorModes.cpp
--- a/Engine/Source/Editor/Kismet/Private/BlueprintEditorModes.cpp
+++ b/Engine/Source/Editor/Kismet/Private/BlueprintEditorModes.cpp
@@ -2,7 +2,7 @@
 
 #include "BlueprintEditorModes.h"
 #include "Settings/EditorExperimentalSettings.h"
-
+#include "ToolMenus.h"
 
 // Core kismet tabs
 #include "SSCSEditor.h"
@@ -114,14 +114,6 @@ FBlueprintEditorApplicationMode::FBlueprintEditorApplicationMode(TSharedPtr<clas
 	CoreTabFactories.RegisterFactory(MakeShareable(new FSelectionDetailsSummoner(InBlueprintEditor)));
 
 	TabLayout = GetDefaltEditorLayout(InBlueprintEditor);
-	
-	// setup toolbar
-	//@TODO: Keep this in sync with AnimBlueprintMode.cpp
-	InBlueprintEditor->GetToolbarBuilder()->AddBlueprintEditorModesToolbar(ToolbarExtender);
-	InBlueprintEditor->GetToolbarBuilder()->AddCompileToolbar(ToolbarExtender);
-	InBlueprintEditor->GetToolbarBuilder()->AddScriptingToolbar(ToolbarExtender);
-	InBlueprintEditor->GetToolbarBuilder()->AddBlueprintGlobalOptionsToolbar(ToolbarExtender);
-	InBlueprintEditor->GetToolbarBuilder()->AddDebuggingToolbar(ToolbarExtender);
 
 	FBlueprintEditorModule& BlueprintEditorModule = FModuleManager::LoadModuleChecked<FBlueprintEditorModule>("Kismet");
 	BlueprintEditorModule.OnRegisterTabsForEditor().Broadcast(BlueprintEditorTabFactories, InModeName, InBlueprintEditor);
@@ -194,8 +186,10 @@ FBlueprintDefaultsApplicationMode::FBlueprintDefaultsApplicationMode(TSharedPtr<
 		);
 
 	// setup toolbar
-	InBlueprintEditor->GetToolbarBuilder()->AddCompileToolbar(ToolbarExtender);
-	InBlueprintEditor->GetToolbarBuilder()->AddBlueprintEditorModesToolbar(ToolbarExtender);
+	if (UToolMenu* Toolbar = InBlueprintEditor->RegisterModeToolbarIfUnregistered(GetModeName()))
+	{
+		InBlueprintEditor->GetToolbarBuilder()->AddCompileToolbar(Toolbar);
+	}
 }
 
 void FBlueprintDefaultsApplicationMode::RegisterTabFactories(TSharedPtr<FTabManager> InTabManager)
@@ -276,10 +270,12 @@ FBlueprintComponentsApplicationMode::FBlueprintComponentsApplicationMode(TShared
 		);
 
 	// setup toolbar
-	InBlueprintEditor->GetToolbarBuilder()->AddBlueprintEditorModesToolbar(ToolbarExtender);
-	InBlueprintEditor->GetToolbarBuilder()->AddBlueprintGlobalOptionsToolbar(ToolbarExtender);
-	InBlueprintEditor->GetToolbarBuilder()->AddCompileToolbar(ToolbarExtender);
-	InBlueprintEditor->GetToolbarBuilder()->AddComponentsToolbar(ToolbarExtender);
+	if (UToolMenu* Toolbar = InBlueprintEditor->RegisterModeToolbarIfUnregistered(GetModeName()))
+	{
+		InBlueprintEditor->GetToolbarBuilder()->AddBlueprintGlobalOptionsToolbar(Toolbar);
+		InBlueprintEditor->GetToolbarBuilder()->AddCompileToolbar(Toolbar);
+		InBlueprintEditor->GetToolbarBuilder()->AddComponentsToolbar(Toolbar);
+	}
 }
 
 void FBlueprintComponentsApplicationMode::RegisterTabFactories(TSharedPtr<FTabManager> InTabManager)
@@ -416,8 +412,11 @@ FBlueprintInterfaceApplicationMode::FBlueprintInterfaceApplicationMode(TSharedPt
 		);
 
 	// setup toolbar
-	InBlueprintEditor->GetToolbarBuilder()->AddCompileToolbar(ToolbarExtender);
-	InBlueprintEditor->GetToolbarBuilder()->AddBlueprintGlobalOptionsToolbar(ToolbarExtender);
+	if (UToolMenu* Toolbar = InBlueprintEditor->RegisterModeToolbarIfUnregistered(GetModeName()))
+	{
+		InBlueprintEditor->GetToolbarBuilder()->AddCompileToolbar(Toolbar);
+		InBlueprintEditor->GetToolbarBuilder()->AddBlueprintGlobalOptionsToolbar(Toolbar);
+	}
 }
 
 void FBlueprintInterfaceApplicationMode::RegisterTabFactories(TSharedPtr<FTabManager> InTabManager)
@@ -521,12 +520,14 @@ FBlueprintMacroApplicationMode::FBlueprintMacroApplicationMode(TSharedPtr<class
 		);
 
 	// setup toolbar
-	InBlueprintEditor->GetToolbarBuilder()->AddCompileToolbar(ToolbarExtender);
-	InBlueprintEditor->GetToolbarBuilder()->AddScriptingToolbar(ToolbarExtender);
-	InBlueprintEditor->GetToolbarBuilder()->AddBlueprintGlobalOptionsToolbar(ToolbarExtender);
-	InBlueprintEditor->GetToolbarBuilder()->AddDebuggingToolbar(ToolbarExtender);
+	if (UToolMenu* Toolbar = InBlueprintEditor->RegisterModeToolbarIfUnregistered(GetModeName()))
+	{
+		InBlueprintEditor->GetToolbarBuilder()->AddCompileToolbar(Toolbar);
+		InBlueprintEditor->GetToolbarBuilder()->AddScriptingToolbar(Toolbar);
+		InBlueprintEditor->GetToolbarBuilder()->AddBlueprintGlobalOptionsToolbar(Toolbar);
+		InBlueprintEditor->GetToolbarBuilder()->AddDebuggingToolbar(Toolbar);
+	}
 }
-
 void FBlueprintMacroApplicationMode::RegisterTabFactories(TSharedPtr<FTabManager> InTabManager)
 {
 	TSharedPtr<FBlueprintEditor> BP = MyBlueprintEditor.Pin();
@@ -716,17 +717,19 @@ FBlueprintEditorUnifiedMode::FBlueprintEditorUnifiedMode(TSharedPtr<class FBluep
 	
 	// setup toolbar
 	//@TODO: Keep this in sync with AnimBlueprintMode.cpp
-	//InBlueprintEditor->GetToolbarBuilder()->AddBlueprintEditorModesToolbar(ToolbarExtender);
-	InBlueprintEditor->GetToolbarBuilder()->AddCompileToolbar(ToolbarExtender);
-	InBlueprintEditor->GetToolbarBuilder()->AddScriptingToolbar(ToolbarExtender);
-	InBlueprintEditor->GetToolbarBuilder()->AddBlueprintGlobalOptionsToolbar(ToolbarExtender);
-	
-	if ( bRegisterViewport )
+	if (UToolMenu* Toolbar = InBlueprintEditor->RegisterModeToolbarIfUnregistered(GetModeName()))
 	{
-		InBlueprintEditor->GetToolbarBuilder()->AddComponentsToolbar(ToolbarExtender);
+		InBlueprintEditor->GetToolbarBuilder()->AddCompileToolbar(Toolbar);
+		InBlueprintEditor->GetToolbarBuilder()->AddScriptingToolbar(Toolbar);
+		InBlueprintEditor->GetToolbarBuilder()->AddBlueprintGlobalOptionsToolbar(Toolbar);
+
+		if ( bRegisterViewport )
+		{
+			InBlueprintEditor->GetToolbarBuilder()->AddComponentsToolbar(Toolbar);
+		}
+
+		InBlueprintEditor->GetToolbarBuilder()->AddDebuggingToolbar(Toolbar);
 	}
-	
-	InBlueprintEditor->GetToolbarBuilder()->AddDebuggingToolbar(ToolbarExtender);
 
 	FBlueprintEditorModule& BlueprintEditorModule = FModuleManager::LoadModuleChecked<FBlueprintEditorModule>("Kismet");
 	BlueprintEditorModule.OnRegisterTabsForEditor().Broadcast(BlueprintEditorTabFactories, InModeName, InBlueprintEditor);
diff --git a/Engine/Source/Editor/Kismet/Private/Debugging/SKismetDebuggingView.cpp b/Engine/Source/Editor/Kismet/Private/Debugging/SKismetDebuggingView.cpp
--- a/Engine/Source/Editor/Kismet/Private/Debugging/SKismetDebuggingView.cpp
+++ b/Engine/Source/Editor/Kismet/Private/Debugging/SKismetDebuggingView.cpp
@@ -23,6 +23,7 @@
 #include "Kismet2/DebuggerCommands.h"
 #include "Debugging/KismetDebugCommands.h"
 #include "Widgets/Input/SHyperlink.h"
+#include "ToolMenus.h"
 
 #define LOCTEXT_NAMESPACE "DebugViewUI"
 
@@ -1062,8 +1063,19 @@ void SKismetDebuggingView::Construct(const FArguments& InArgs)
 	BlueprintToWatchPtr = InArgs._BlueprintToWatch;
 
 	// Build the debug toolbar
-	FToolBarBuilder DebugToolbarBuilder(FPlayWorldCommands::GlobalPlayWorldActions, FMultiBoxCustomization::None );
-	FPlayWorldCommands::BuildToolbar(DebugToolbarBuilder);
+	static const FName ToolbarName = "Kismet.DebuggingViewToolBar";
+	if (!UToolMenus::Get()->IsMenuRegistered(ToolbarName))
+	{
+		UToolMenu* ToolBar = UToolMenus::Get()->RegisterMenu(ToolbarName, NAME_None, EMultiBoxType::ToolBar);
+
+		{
+			FToolMenuSection& Section = ToolBar->AddSection("Debug");
+			FPlayWorldCommands::BuildToolbar(Section);
+		}
+	}
+
+	FToolMenuContext MenuContext(FPlayWorldCommands::GlobalPlayWorldActions);
+	TSharedRef<SWidget> ToolbarWidget = UToolMenus::Get()->GenerateWidget(ToolbarName, MenuContext);
 
 	this->ChildSlot
 	[
@@ -1075,7 +1087,7 @@ void SKismetDebuggingView::Construct(const FArguments& InArgs)
 			.Visibility( this, &SKismetDebuggingView::IsDebuggerVisible )
 			.BorderImage( FEditorStyle::GetBrush( TEXT("NoBorder") ) )
 			[
-				DebugToolbarBuilder.MakeWidget()
+				ToolbarWidget
 			]
 		]
 		+SVerticalBox::Slot()
diff --git a/Engine/Source/Editor/Kismet/Private/SBlueprintEditorToolbar.cpp b/Engine/Source/Editor/Kismet/Private/SBlueprintEditorToolbar.cpp
--- a/Engine/Source/Editor/Kismet/Private/SBlueprintEditorToolbar.cpp
+++ b/Engine/Source/Editor/Kismet/Private/SBlueprintEditorToolbar.cpp
@@ -28,6 +28,8 @@
 #include "SBlueprintEditorSelectedDebugObjectWidget.h"
 #include "DesktopPlatformModule.h"
 #include "SBlueprintRevisionMenu.h"
+#include "ToolMenus.h"
+#include "BlueprintEditorContext.h"
 #include "ISourceCodeAccessor.h"
 #include "ISourceCodeAccessModule.h"
 
@@ -397,222 +399,151 @@ static void BlueprintEditorToolbarImpl::MakeCompileDeveloperSubMenu(FMenuBuilder
 //////////////////////////////////////////////////////////////////////////
 // FBlueprintEditorToolbar
 
-void FBlueprintEditorToolbar::AddBlueprintEditorModesToolbar(TSharedPtr<FExtender> Extender)
+void FBlueprintEditorToolbar::AddBlueprintGlobalOptionsToolbar(UToolMenu* InMenu)
 {
-	
-}
-
-void FBlueprintEditorToolbar::AddBlueprintGlobalOptionsToolbar(TSharedPtr<FExtender> Extender)
-{
-	TSharedPtr<FBlueprintEditor> BlueprintEditorPtr = BlueprintEditor.Pin();
-
-	Extender->AddToolBarExtension(
-		"Asset",
-		EExtensionHook::After,
-		BlueprintEditorPtr->GetToolkitCommands(),
-		FToolBarExtensionDelegate::CreateSP( this, &FBlueprintEditorToolbar::FillBlueprintGlobalOptionsToolbar ) );
-}
+	FToolMenuSection& Section = InMenu->AddSection("Settings");
+	Section.InsertPosition = FToolMenuInsert("Asset", EToolMenuInsertType::After);
 
-void FBlueprintEditorToolbar::AddCompileToolbar(TSharedPtr<FExtender> Extender)
-{
-	TSharedPtr<FBlueprintEditor> BlueprintEditorPtr = BlueprintEditor.Pin();
-
-	Extender->AddToolBarExtension(
-		"Asset",
-		EExtensionHook::Before,
-		BlueprintEditorPtr->GetToolkitCommands(),
-		FToolBarExtensionDelegate::CreateSP( this, &FBlueprintEditorToolbar::FillCompileToolbar ) );
-}
-
-void FBlueprintEditorToolbar::AddNewToolbar(TSharedPtr<FExtender> Extender)
-{
-	TSharedPtr<FBlueprintEditor> BlueprintEditorPtr = BlueprintEditor.Pin();
-
-	Extender->AddToolBarExtension(
-		"MyBlueprint",
-		EExtensionHook::After,
-		BlueprintEditorPtr->GetToolkitCommands(),
-		FToolBarExtensionDelegate::CreateSP( this, &FBlueprintEditorToolbar::FillNewToolbar ) );
-}
-
-void FBlueprintEditorToolbar::AddScriptingToolbar(TSharedPtr<FExtender> Extender)
-{
-	TSharedPtr<FBlueprintEditor> BlueprintEditorPtr = BlueprintEditor.Pin();
-
-	Extender->AddToolBarExtension(
-		"Asset",
-		EExtensionHook::After,
-		BlueprintEditorPtr->GetToolkitCommands(),
-		FToolBarExtensionDelegate::CreateSP( this, &FBlueprintEditorToolbar::FillScriptingToolbar ) );
-}
-
-void FBlueprintEditorToolbar::AddDebuggingToolbar(TSharedPtr<FExtender> Extender)
-{
-	TSharedPtr<FBlueprintEditor> BlueprintEditorPtr = BlueprintEditor.Pin();
-
-	Extender->AddToolBarExtension(
-		"Asset",
-		EExtensionHook::After,
-		BlueprintEditorPtr->GetToolkitCommands(),
-		FToolBarExtensionDelegate::CreateSP( this, &FBlueprintEditorToolbar::FillDebuggingToolbar ) );
-}
-
-void FBlueprintEditorToolbar::AddComponentsToolbar(TSharedPtr<FExtender> Extender)
-{
-	TSharedPtr<FBlueprintEditor> BlueprintEditorPtr = BlueprintEditor.Pin();
-
-	Extender->AddToolBarExtension(
-		"Asset",
-		EExtensionHook::After,
-		BlueprintEditorPtr->GetToolkitCommands(),
-		FToolBarExtensionDelegate::CreateSP( this, &FBlueprintEditorToolbar::FillComponentsToolbar ) );
-}
-
-void FBlueprintEditorToolbar::FillBlueprintEditorModesToolbar(FToolBarBuilder& ToolbarBuilder)
-{
-	TSharedPtr<FBlueprintEditor> BlueprintEditorPtr = BlueprintEditor.Pin();
-	UBlueprint* BlueprintObj = BlueprintEditorPtr->GetBlueprintObj();
-
-	TAttribute<FName> GetActiveMode(BlueprintEditorPtr.ToSharedRef(), &FBlueprintEditor::GetCurrentMode);
-	FOnModeChangeRequested SetActiveMode = FOnModeChangeRequested::CreateSP(BlueprintEditorPtr.ToSharedRef(), &FBlueprintEditor::SetCurrentMode);
-
-	TArray< TSharedPtr< SWidget > > ToolbarWidgets = GenerateToolbarWidgets( BlueprintObj, GetActiveMode, SetActiveMode );
-	
-	for( const auto& Widget : ToolbarWidgets )
+	Section.AddDynamicEntry("BlueprintGlobalOptions", FNewToolMenuSectionDelegate::CreateLambda([](FToolMenuSection& InSection)
 	{
-		BlueprintEditorPtr->AddToolbarWidget( Widget.ToSharedRef() );
-	}
+		UBlueprintEditorToolMenuContext* Context = InSection.FindContext<UBlueprintEditorToolMenuContext>();
+		if (Context && Context->GetBlueprintObj())
+		{
+			const FFullBlueprintEditorCommands& Commands = FFullBlueprintEditorCommands::Get();
+			InSection.AddEntry(FToolMenuEntry::InitToolBarButton(Commands.EditGlobalOptions));
+			InSection.AddEntry(FToolMenuEntry::InitToolBarButton(Commands.EditClassDefaults));
+		}
+	}));
 }
 
-void FBlueprintEditorToolbar::FillBlueprintGlobalOptionsToolbar(FToolBarBuilder& ToolbarBuilder)
+void FBlueprintEditorToolbar::AddCompileToolbar(UToolMenu* InMenu)
 {
-	const FFullBlueprintEditorCommands& Commands = FFullBlueprintEditorCommands::Get();
-	TSharedPtr<FBlueprintEditor> BlueprintEditorPtr = BlueprintEditor.Pin();
-	UBlueprint* BlueprintObj = BlueprintEditorPtr->GetBlueprintObj();
-	
-	ToolbarBuilder.BeginSection("Settings");
-	
-	if(BlueprintObj != NULL)
+	FToolMenuSection& Section = InMenu->AddSection("Compile");
+	Section.InsertPosition = FToolMenuInsert("Asset", EToolMenuInsertType::Before);
+
+	Section.AddDynamicEntry("CompileCommands", FNewToolMenuSectionDelegate::CreateLambda([](FToolMenuSection& InSection)
 	{
-		ToolbarBuilder.AddToolBarButton(Commands.EditGlobalOptions);
-		ToolbarBuilder.AddToolBarButton(Commands.EditClassDefaults);
-	}
-	
-	ToolbarBuilder.EndSection();
+		UBlueprintEditorToolMenuContext* Context = InSection.FindContext<UBlueprintEditorToolMenuContext>();
+		if (Context && Context->BlueprintEditor.IsValid() && Context->GetBlueprintObj())
+		{
+			TSharedPtr<class FBlueprintEditorToolbar> BlueprintEditorToolbar = Context->BlueprintEditor.Pin()->GetToolbarBuilder();
+			if (BlueprintEditorToolbar.IsValid())
+			{
+				const FFullBlueprintEditorCommands& Commands = FFullBlueprintEditorCommands::Get();
+
+				InSection.AddEntry(FToolMenuEntry::InitToolBarButton(
+					Commands.Compile,
+					TAttribute<FText>(),
+					TAttribute<FText>(BlueprintEditorToolbar.ToSharedRef(), &FBlueprintEditorToolbar::GetStatusTooltip),
+					TAttribute<FSlateIcon>(BlueprintEditorToolbar.ToSharedRef(), &FBlueprintEditorToolbar::GetStatusImage),
+					"CompileBlueprint"
+				));
+
+				InSection.AddEntry(FToolMenuEntry::InitComboButton(
+					"BlueprintCompileOptions",
+					FUIAction(),
+					FOnGetContent::CreateStatic(&BlueprintEditorToolbarImpl::GenerateCompileOptionsWidget, Context->BlueprintEditor.Pin()->GetToolkitCommands()),
+					LOCTEXT("BlupeintCompileOptions_ToolbarName",    "Compile Options"),
+					LOCTEXT("BlupeintCompileOptions_ToolbarTooltip", "Options to customize how Blueprints compile"),
+					TAttribute<FSlateIcon>(),
+					/*bSimpleComboBox =*/true
+				));
+			}
+		}
+	}));
 }
 
-void FBlueprintEditorToolbar::FillCompileToolbar(FToolBarBuilder& ToolbarBuilder)
+void FBlueprintEditorToolbar::AddNewToolbar(UToolMenu* InMenu)
 {
-	const FFullBlueprintEditorCommands& Commands = FFullBlueprintEditorCommands::Get();
-	TSharedPtr<FBlueprintEditor> BlueprintEditorPtr = BlueprintEditor.Pin();
-	UBlueprint* BlueprintObj = BlueprintEditorPtr->GetBlueprintObj();
+	FToolMenuSection& Section = InMenu->AddSection("Add");
+	Section.InsertPosition = FToolMenuInsert("MyBlueprint", EToolMenuInsertType::After);
 
-	ToolbarBuilder.BeginSection("Compile");
-	if (BlueprintObj != NULL)
+	Section.AddDynamicEntry("AddCommands", FNewToolMenuSectionDelegate::CreateLambda([](FToolMenuSection& InSection)
 	{
-		ToolbarBuilder.AddToolBarButton( Commands.Compile,
-									 NAME_None, 
-									 TAttribute<FText>(),
-									 TAttribute<FText>(this, &FBlueprintEditorToolbar::GetStatusTooltip),
-									 TAttribute<FSlateIcon>(this, &FBlueprintEditorToolbar::GetStatusImage),
-									 FName(TEXT("CompileBlueprint")));
-
-		FUIAction TempCompileOptionsCommand;
-		ToolbarBuilder.AddComboButton(
-			TempCompileOptionsCommand,
-			FOnGetContent::CreateStatic(&BlueprintEditorToolbarImpl::GenerateCompileOptionsWidget, BlueprintEditorPtr->GetToolkitCommands()),
-			LOCTEXT("BlupeintCompileOptions_ToolbarName",    "Compile Options"),
-			LOCTEXT("BlupeintCompileOptions_ToolbarTooltip", "Options to customize how Blueprints compile"),
-			TAttribute<FSlateIcon>(),
-			/*bSimpleComboBox =*/true
-		);
-	}
-	ToolbarBuilder.EndSection();
+		UBlueprintEditorToolMenuContext* Context = InSection.FindContext<UBlueprintEditorToolMenuContext>();
+		if (Context && Context->BlueprintEditor.IsValid() && Context->GetBlueprintObj())
+		{
+			const FBlueprintEditorCommands& Commands = FBlueprintEditorCommands::Get();
+			InSection.AddEntry(FToolMenuEntry::InitToolBarButton(Commands.AddNewVariable, TAttribute<FText>(), TAttribute<FText>(), TAttribute<FSlateIcon>(), FName(TEXT("BPEAddNewVariable"))));
+			InSection.AddEntry(FToolMenuEntry::InitToolBarButton(Commands.AddNewFunction, TAttribute<FText>(), TAttribute<FText>(), TAttribute<FSlateIcon>(), FName(TEXT("BPEAddNewFunction"))));
+			InSection.AddEntry(FToolMenuEntry::InitToolBarButton(Commands.AddNewMacroDeclaration, TAttribute<FText>(), TAttribute<FText>(), TAttribute<FSlateIcon>(), FName(TEXT("BPEAddNewMacro"))));
+			// Add New Animation Graph isn't supported right now.
+			InSection.AddEntry(FToolMenuEntry::InitToolBarButton(Commands.AddNewEventGraph, TAttribute<FText>(), TAttribute<FText>(), TAttribute<FSlateIcon>(), FName(TEXT("BPEAddNewEventGraph"))));
+			InSection.AddEntry(FToolMenuEntry::InitToolBarButton(Commands.AddNewDelegate, TAttribute<FText>(), TAttribute<FText>(), TAttribute<FSlateIcon>(), FName(TEXT("BPEAddNewDelegate"))));
+		}
+	}));
 }
 
-void FBlueprintEditorToolbar::FillNewToolbar(FToolBarBuilder& ToolbarBuilder)
+void FBlueprintEditorToolbar::AddScriptingToolbar(UToolMenu* InMenu)
 {
-	const FBlueprintEditorCommands& Commands = FBlueprintEditorCommands::Get();
-
-	TSharedPtr<FBlueprintEditor> BlueprintEditorPtr = BlueprintEditor.Pin();
-	UBlueprint* BlueprintObj = BlueprintEditorPtr->GetBlueprintObj();
+	FToolMenuSection& Section = InMenu->AddSection("Script");
+	Section.InsertPosition = FToolMenuInsert("Asset", EToolMenuInsertType::After);
 
-	ToolbarBuilder.BeginSection("AddNew");
-	if (BlueprintObj != NULL)
+	Section.AddDynamicEntry("ScriptCommands", FNewToolMenuSectionDelegate::CreateLambda([](FToolMenuSection& InSection)
 	{
-		ToolbarBuilder.AddToolBarButton(Commands.AddNewVariable, NAME_None, TAttribute<FText>(), TAttribute<FText>(), TAttribute<FSlateIcon>(), FName(TEXT("BPEAddNewVariable")));
-		ToolbarBuilder.AddToolBarButton(Commands.AddNewFunction, NAME_None, TAttribute<FText>(), TAttribute<FText>(), TAttribute<FSlateIcon>(), FName(TEXT("BPEAddNewFunction")));
-		ToolbarBuilder.AddToolBarButton(Commands.AddNewMacroDeclaration, NAME_None, TAttribute<FText>(), TAttribute<FText>(), TAttribute<FSlateIcon>(), FName(TEXT("BPEAddNewMacro")));
-		// Add New Animation Graph isn't supported right now.
-		ToolbarBuilder.AddToolBarButton(Commands.AddNewEventGraph, NAME_None, TAttribute<FText>(), TAttribute<FText>(), TAttribute<FSlateIcon>(), FName(TEXT("BPEAddNewEventGraph")));
-		ToolbarBuilder.AddToolBarButton(Commands.AddNewDelegate, NAME_None, TAttribute<FText>(), TAttribute<FText>(), TAttribute<FSlateIcon>(), FName(TEXT("BPEAddNewDelegate")));
-	}
-	ToolbarBuilder.EndSection();
-}
-
-void FBlueprintEditorToolbar::FillScriptingToolbar(FToolBarBuilder& ToolbarBuilder)
-{
-	const FBlueprintEditorCommands& Commands = FBlueprintEditorCommands::Get();
-
-	TSharedPtr<FBlueprintEditor> BlueprintEditorPtr = BlueprintEditor.Pin();
-	UBlueprint* BlueprintObj = BlueprintEditorPtr->GetBlueprintObj();
-
-	ToolbarBuilder.BeginSection("Script");
-
-	ToolbarBuilder.AddToolBarButton(FBlueprintEditorCommands::Get().FindInBlueprint);
-
-	ToolbarBuilder.AddToolBarButton(
-		FBlueprintEditorCommands::Get().ToggleHideUnrelatedNodes,
-		NAME_None,
-		TAttribute<FText>(),
-		TAttribute<FText>(),
-		FSlateIcon(FEditorStyle::GetStyleSetName(), "GraphEditor.ToggleHideUnrelatedNodes")
-	);
-	ToolbarBuilder.AddComboButton(
-		FUIAction(),
-		FOnGetContent::CreateSP(BlueprintEditorPtr.Get(), &FBlueprintEditor::MakeHideUnrelatedNodesOptionsMenu),
-		LOCTEXT("HideUnrelatedNodesOptions", "Focus Related Nodes Options"),
-		LOCTEXT("HideUnrelatedNodesOptionsMenu", "Focus Related Nodes options menu"),
-		TAttribute<FSlateIcon>(),
-		true
-	);
-
-	ToolbarBuilder.EndSection();
+		UBlueprintEditorToolMenuContext* Context = InSection.FindContext<UBlueprintEditorToolMenuContext>();
+		if (Context && Context->BlueprintEditor.IsValid() && Context->GetBlueprintObj())
+		{
+			InSection.AddEntry(FToolMenuEntry::InitToolBarButton(FBlueprintEditorCommands::Get().FindInBlueprint));
+
+			InSection.AddEntry(FToolMenuEntry::InitToolBarButton(
+				FBlueprintEditorCommands::Get().ToggleHideUnrelatedNodes,
+				TAttribute<FText>(),
+				TAttribute<FText>(),
+				FSlateIcon(FEditorStyle::GetStyleSetName(), "GraphEditor.ToggleHideUnrelatedNodes")
+			));
+
+			InSection.AddEntry(FToolMenuEntry::InitComboButton(
+				"HideUnrelatedNodesOptions",
+				FUIAction(),
+				FOnGetContent::CreateSP(Context->BlueprintEditor.Pin().Get(), &FBlueprintEditor::MakeHideUnrelatedNodesOptionsMenu),
+				LOCTEXT("HideUnrelatedNodesOptions", "Focus Related Nodes Options"),
+				LOCTEXT("HideUnrelatedNodesOptionsMenu", "Focus Related Nodes options menu"),
+				TAttribute<FSlateIcon>(),
+				true
+			));
+		}
+	}));
 }
 
-void FBlueprintEditorToolbar::FillDebuggingToolbar(FToolBarBuilder& ToolbarBuilder)
+void FBlueprintEditorToolbar::AddDebuggingToolbar(UToolMenu* InMenu)
 {
-	TSharedPtr<FBlueprintEditor> BlueprintEditorPtr = BlueprintEditor.Pin();
-	UBlueprint* BlueprintObj = BlueprintEditorPtr->GetBlueprintObj();
+	FToolMenuSection& Section = InMenu->AddSection("Debugging");
+	Section.InsertPosition = FToolMenuInsert("Asset", EToolMenuInsertType::After);
 
-	ToolbarBuilder.BeginSection("Debugging");
-	if (BlueprintObj)
+	Section.AddDynamicEntry("DebuggingCommands", FNewToolMenuSectionDelegate::CreateLambda([](FToolMenuSection& InSection)
 	{
-		FPlayWorldCommands::BuildToolbar(ToolbarBuilder);
-
-		if (BlueprintObj->BlueprintType != BPTYPE_MacroLibrary)
+		UBlueprintEditorToolMenuContext* Context = InSection.FindContext<UBlueprintEditorToolMenuContext>();
+		if (Context && Context->BlueprintEditor.IsValid() && Context->GetBlueprintObj())
 		{
-			// Selected debug actor button
-			ToolbarBuilder.AddWidget(SNew(SBlueprintEditorSelectedDebugObjectWidget, BlueprintEditorPtr));
+			FPlayWorldCommands::BuildToolbar(InSection);
+
+			if (Context->GetBlueprintObj()->BlueprintType != BPTYPE_MacroLibrary)
+			{
+				// Selected debug actor button
+				InSection.AddEntry(FToolMenuEntry::InitWidget("SelectedDebugObjectWidget", SNew(SBlueprintEditorSelectedDebugObjectWidget, Context->BlueprintEditor.Pin()), FText::GetEmpty()));
+			}
 		}
-	}
-	ToolbarBuilder.EndSection();
+	}));
 }
 
-void FBlueprintEditorToolbar::FillComponentsToolbar(FToolBarBuilder& ToolbarBuilder)
+void FBlueprintEditorToolbar::AddComponentsToolbar(UToolMenu* InMenu)
 {
-	TSharedPtr<FBlueprintEditor> BlueprintEditorPtr = BlueprintEditor.Pin();
-	UBlueprint* BlueprintObj = BlueprintEditorPtr->GetBlueprintObj();
+	FToolMenuInsert InsertPosition("Asset", EToolMenuInsertType::After);
 
 #if 0 // restore this if we ever need the ability to toggle component editing on/off
-	ToolbarBuilder.BeginSection("Components");
-		ToolbarBuilder.AddToolBarButton(FSCSCommands::Get().ToggleComponentEditing);
-	ToolbarBuilder.EndSection();
+	{
+		FToolMenuSection& Section = InMenu->AddSection("Components");
+		Section.InsertPosition = InsertPosition;
+		Section.AddEntry(FToolMenuEntry::InitToolBarButton(FSCSCommands::Get().ToggleComponentEditing));
+	}
 #endif
 
-	ToolbarBuilder.BeginSection("ComponentsViewport");
-		ToolbarBuilder.AddToolBarButton(FBlueprintEditorCommands::Get().EnableSimulation);
-	ToolbarBuilder.EndSection();
+	{
+		FToolMenuSection& Section = InMenu->AddSection("ComponentsViewport");
+		Section.InsertPosition = InsertPosition;
+		Section.AddEntry(FToolMenuEntry::InitToolBarButton(FBlueprintEditorCommands::Get().EnableSimulation));
+	}
 }
 
 FSlateIcon FBlueprintEditorToolbar::GetStatusImage() const
@@ -668,72 +599,4 @@ FText FBlueprintEditorToolbar::GetStatusTooltip() const
 	}
 }
 
-TArray< TSharedPtr< SWidget> > FBlueprintEditorToolbar::GenerateToolbarWidgets(const UBlueprint* BlueprintObj, TAttribute<FName> ActiveModeGetter, FOnModeChangeRequested ActiveModeSetter)
-{
-	TArray< TSharedPtr< SWidget> > Ret;
-	if (!BlueprintObj ||
-		(!FBlueprintEditorUtils::IsLevelScriptBlueprint(BlueprintObj)
-		&& !FBlueprintEditorUtils::IsInterfaceBlueprint(BlueprintObj)
-		&& !BlueprintObj->bIsNewlyCreated)
-		)
-	{
-		// Left side padding
-		Ret.Add(SNew(SSpacer).Size(FVector2D(4.0f, 1.0f)));
-
-		Ret.Add(
-			SNew(SModeWidget, FBlueprintEditorApplicationModes::GetLocalizedMode(FBlueprintEditorApplicationModes::BlueprintDefaultsMode), FBlueprintEditorApplicationModes::BlueprintDefaultsMode)
-			.OnGetActiveMode(ActiveModeGetter)
-			.OnSetActiveMode(ActiveModeSetter)
-			.CanBeSelected(BlueprintObj ? FBlueprintEditorUtils::DoesSupportDefaults(BlueprintObj) : false)
-			.ToolTip(IDocumentation::Get()->CreateToolTip(
-				LOCTEXT("BlueprintDefaultsModeButtonTooltip", "Switch to Class Defaults Mode"),
-				NULL,
-				TEXT("Shared/Editors/BlueprintEditor"),
-				TEXT("DefaultsMode")))
-			.IconImage(FEditorStyle::GetBrush("FullBlueprintEditor.SwitchToBlueprintDefaultsMode"))
-			.SmallIconImage(FEditorStyle::GetBrush("FullBlueprintEditor.SwitchToBlueprintDefaultsMode.Small"))
-			.AddMetaData<FTagMetaData>(FTagMetaData(TEXT("DefaultsMode")))
-		);
-
-		Ret.Add(SNew(SBlueprintModeSeparator));
-
-		Ret.Add(
-			SNew(SModeWidget, FBlueprintEditorApplicationModes::GetLocalizedMode(FBlueprintEditorApplicationModes::BlueprintComponentsMode), FBlueprintEditorApplicationModes::BlueprintComponentsMode)
-			.OnGetActiveMode(ActiveModeGetter)
-			.OnSetActiveMode(ActiveModeSetter)
-			.CanBeSelected(BlueprintObj ? FBlueprintEditorUtils::DoesSupportComponents(BlueprintObj) : false)
-			.ToolTip(IDocumentation::Get()->CreateToolTip(
-				LOCTEXT("ComponentsModeButtonTooltip", "Switch to Components Mode"),
-				NULL,
-				TEXT("Shared/Editors/BlueprintEditor"),
-				TEXT("ComponentsMode")))
-			.IconImage(FEditorStyle::GetBrush("FullBlueprintEditor.SwitchToComponentsMode"))
-			.SmallIconImage(FEditorStyle::GetBrush("FullBlueprintEditor.SwitchToComponentsMode.Small"))
-			.AddMetaData<FTagMetaData>(FTagMetaData(TEXT("ComponentsMode")))
-		);
-
-		Ret.Add(SNew(SBlueprintModeSeparator));
-
-		Ret.Add(
-			SNew(SModeWidget, FBlueprintEditorApplicationModes::GetLocalizedMode(FBlueprintEditorApplicationModes::StandardBlueprintEditorMode), FBlueprintEditorApplicationModes::StandardBlueprintEditorMode)
-			.OnGetActiveMode(ActiveModeGetter)
-			.OnSetActiveMode(ActiveModeSetter)
-			.CanBeSelected(BlueprintObj != NULL)
-			.ToolTip(IDocumentation::Get()->CreateToolTip(
-				LOCTEXT("GraphModeButtonTooltip", "Switch to Graph Editing Mode"),
-				NULL,
-				TEXT("Shared/Editors/BlueprintEditor"),
-				TEXT("GraphMode")))
-			.ToolTipText(LOCTEXT("GraphModeButtonTooltip", "Switch to Graph Editing Mode"))
-			.IconImage(FEditorStyle::GetBrush("FullBlueprintEditor.SwitchToScriptingMode"))
-			.SmallIconImage(FEditorStyle::GetBrush("FullBlueprintEditor.SwitchToScriptingMode.Small"))
-			.AddMetaData<FTagMetaData>(FTagMetaData(TEXT("GraphMode")))
-		);
-
-		// Right side padding
-		Ret.Add(SNew(SSpacer).Size(FVector2D(4.0f, 1.0f)));
-	}
-	return Ret;
-}
-
 #undef LOCTEXT_NAMESPACE
diff --git a/Engine/Source/Editor/Kismet/Private/SMyBlueprint.cpp b/Engine/Source/Editor/Kismet/Private/SMyBlueprint.cpp
--- a/Engine/Source/Editor/Kismet/Private/SMyBlueprint.cpp
+++ b/Engine/Source/Editor/Kismet/Private/SMyBlueprint.cpp
@@ -300,10 +300,6 @@ void SMyBlueprint::Construct(const FArguments& InArgs, TWeakPtr<FBlueprintEditor
 			FCanExecuteAction(),
 			FIsActionChecked(),
 			FIsActionButtonVisible::CreateSP(this, &SMyBlueprint::IsNativeVariable) );
-
-		TSharedPtr<FBlueprintEditorToolbar> Toolbar = MakeShareable(new FBlueprintEditorToolbar(InBlueprintEditor.Pin()));
-		TSharedPtr<FExtender> Extender = MakeShareable(new FExtender);
-		Toolbar->AddNewToolbar(Extender);
 		ToolbarBuilderWidget = SNullWidget::NullWidget;
 	
 		ToolKitCommandList->MapAction(FGenericCommands::Get().Rename,
diff --git a/Engine/Source/Editor/Kismet/Private/WorkflowCentricApplication.cpp b/Engine/Source/Editor/Kismet/Private/WorkflowCentricApplication.cpp
--- a/Engine/Source/Editor/Kismet/Private/WorkflowCentricApplication.cpp
+++ b/Engine/Source/Editor/Kismet/Private/WorkflowCentricApplication.cpp
@@ -5,6 +5,7 @@
 #include "WorkflowOrientedApp/WorkflowTabFactory.h"
 #include "WorkflowOrientedApp/WorkflowTabManager.h"
 #include "WorkflowOrientedApp/ApplicationMode.h"
+#include "ToolMenus.h"
 
 #define LOCTEXT_NAMESPACE "WorkflowCentricApplication"
 
@@ -26,6 +27,34 @@ void FWorkflowCentricApplication::UnregisterTabSpawners(const TSharedRef<class F
 	InTabManager->UnregisterAllTabSpawners();
 }
 
+FName FWorkflowCentricApplication::GetToolMenuToolbarName(FName& OutParentName) const
+{
+	return GetToolMenuToolbarNameForMode(GetCurrentMode(), OutParentName);
+}
+
+FName FWorkflowCentricApplication::GetToolMenuToolbarNameForMode(const FName InModeName, FName& OutParentName) const
+{
+	const FName BaseMenuName = FAssetEditorToolkit::GetToolMenuToolbarName(OutParentName);
+	if (InModeName != NAME_None)
+	{
+		OutParentName = BaseMenuName;
+		return *(BaseMenuName.ToString() + TEXT(".") + InModeName.ToString());
+	}
+
+	return BaseMenuName;
+}
+
+UToolMenu* FWorkflowCentricApplication::RegisterModeToolbarIfUnregistered(const FName InModeName)
+{
+	FName ParentToolbarName;
+	const FName ModeSpecificToolbarName = GetToolMenuToolbarNameForMode(InModeName, ParentToolbarName);
+	if (!UToolMenus::Get()->IsMenuRegistered(ModeSpecificToolbarName))
+	{
+		return UToolMenus::Get()->RegisterMenu(ModeSpecificToolbarName, ParentToolbarName, EMultiBoxType::ToolBar);
+	}
+
+	return nullptr;
+}
 
 FName FWorkflowCentricApplication::GetCurrentMode() const
 {
diff --git a/Engine/Source/Editor/Kismet/Public/BlueprintEditor.h b/Engine/Source/Editor/Kismet/Public/BlueprintEditor.h
--- a/Engine/Source/Editor/Kismet/Public/BlueprintEditor.h
+++ b/Engine/Source/Editor/Kismet/Public/BlueprintEditor.h
@@ -45,6 +45,8 @@ class UUserDefinedEnum;
 class UUserDefinedStruct;
 class UBlueprintEditorOptions;
 struct Rect;
+class UToolMenu;
+struct FToolMenuContext;
 class UK2Node_FunctionEntry;
 class UK2Node_Event;
 
@@ -238,6 +240,9 @@ public:
 
 	virtual ~FBlueprintEditor();
 
+	/** Add context objects for menus and toolbars */
+	virtual void InitToolMenuContext(FToolMenuContext& MenuContext) override;
+
 	/** Check the Node Title is visible */
 	bool IsNodeTitleVisible(const UEdGraphNode* Node, bool bRequestRename);
 
diff --git a/Engine/Source/Editor/Kismet/Public/BlueprintEditorContext.h b/Engine/Source/Editor/Kismet/Public/BlueprintEditorContext.h
new file mode 100644
--- /dev/null
+++ b/Engine/Source/Editor/Kismet/Public/BlueprintEditorContext.h
@@ -0,0 +1,21 @@
+// Copyright 1998-2019 Epic Games, Inc. All Rights Reserved.
+
+#pragma once
+
+#include "CoreMinimal.h"
+
+#include "BlueprintEditorContext.generated.h"
+
+class UBlueprint;
+class FBlueprintEditor;
+
+UCLASS()
+class KISMET_API UBlueprintEditorToolMenuContext : public UObject
+{
+	GENERATED_BODY()
+public:
+
+	UBlueprint* GetBlueprintObj() const;
+
+	TWeakPtr<FBlueprintEditor> BlueprintEditor;
+};
diff --git a/Engine/Source/Editor/Kismet/Public/SBlueprintEditorToolbar.h b/Engine/Source/Editor/Kismet/Public/SBlueprintEditorToolbar.h
--- a/Engine/Source/Editor/Kismet/Public/SBlueprintEditorToolbar.h
+++ b/Engine/Source/Editor/Kismet/Public/SBlueprintEditorToolbar.h
@@ -14,6 +14,7 @@ class FBlueprintEditor;
 class FExtender;
 class FMenuBuilder;
 class FToolBarBuilder;
+class UToolMenu;
 
 /**
  * Kismet menu
@@ -79,13 +80,12 @@ public:
 	FBlueprintEditorToolbar(TSharedPtr<FBlueprintEditor> InBlueprintEditor)
 		: BlueprintEditor(InBlueprintEditor) {}
 
-	void AddBlueprintEditorModesToolbar(TSharedPtr<FExtender> Extender);
-	void AddBlueprintGlobalOptionsToolbar(TSharedPtr<FExtender> Extender);
-	void AddCompileToolbar(TSharedPtr<FExtender> Extender);
-	void AddNewToolbar(TSharedPtr<FExtender> Extender);
-	void AddScriptingToolbar(TSharedPtr<FExtender> Extender);
-	void AddDebuggingToolbar(TSharedPtr<FExtender> Extender);
-	void AddComponentsToolbar(TSharedPtr<FExtender> Extender);
+	void AddBlueprintGlobalOptionsToolbar(UToolMenu* InMenu);
+	void AddCompileToolbar(UToolMenu* InMenu);
+	void AddNewToolbar(UToolMenu* InMenu);
+	void AddScriptingToolbar(UToolMenu* InMenu);
+	void AddDebuggingToolbar(UToolMenu* InMenu);
+	void AddComponentsToolbar(UToolMenu* InMenu);
 
 	/** Returns the current status icon for the blueprint being edited */
 	FSlateIcon GetStatusImage() const;
@@ -93,18 +93,6 @@ public:
 	/** Returns the current status as text for the blueprint being edited */
 	FText GetStatusTooltip() const;
 
-	/** Helper function for generating the buttons in the toolbar, reused by merge and diff tools */
-	static TArray< TSharedPtr< class SWidget> > GenerateToolbarWidgets(const class UBlueprint* BlueprintObj, TAttribute<FName> ActiveModeGetter, FOnModeChangeRequested ActiveModeSetter);
-
-private:
-	void FillBlueprintEditorModesToolbar(FToolBarBuilder& ToolbarBuilder);
-	void FillBlueprintGlobalOptionsToolbar(FToolBarBuilder& ToolBarBuilder);
-	void FillCompileToolbar(FToolBarBuilder& ToolbarBuilder);
-	void FillNewToolbar(FToolBarBuilder& ToolbarBuilder);
-	void FillScriptingToolbar(FToolBarBuilder& ToolbarBuilder);
-	void FillDebuggingToolbar(FToolBarBuilder& ToolbarBuilder);
-	void FillComponentsToolbar(FToolBarBuilder& ToolbarBuilder);
-
 protected:
 	/** Pointer back to the blueprint editor tool that owns us */
 	TWeakPtr<FBlueprintEditor> BlueprintEditor;
diff --git a/Engine/Source/Editor/Kismet/Public/WorkflowOrientedApp/WorkflowCentricApplication.h b/Engine/Source/Editor/Kismet/Public/WorkflowOrientedApp/WorkflowCentricApplication.h
--- a/Engine/Source/Editor/Kismet/Public/WorkflowOrientedApp/WorkflowCentricApplication.h
+++ b/Engine/Source/Editor/Kismet/Public/WorkflowOrientedApp/WorkflowCentricApplication.h
@@ -8,6 +8,7 @@
 class FApplicationMode;
 class FWorkflowAllowedTabSet;
 class FWorkflowTabFactory;
+class UToolMenu;
 
 // Delegate for mutating a mode
 DECLARE_DELEGATE_RetVal_TwoParams(TSharedRef<FApplicationMode>, FWorkflowApplicationModeExtender, const FName /*ModeName*/, TSharedRef<FApplicationMode> /*InMode*/);
@@ -25,8 +26,12 @@ public:
 	
 	// FAssetEditorToolkit interface
 	virtual bool OnRequestClose() override;
+	virtual FName GetToolMenuToolbarName(FName& OutParentName) const override;
 	// End of FAssetEditorToolkit interface
 
+	// Returns the name of the toolbar for the given mode along with the name of the toolbar's parent.
+	FName GetToolMenuToolbarNameForMode(const FName InModeName, FName& OutParentName) const;
+
 	// Returns the current mode of this application
 	FName GetCurrentMode() const;
 	TSharedPtr<FApplicationMode> GetCurrentModePtr() const { return CurrentAppModePtr; }
@@ -39,6 +44,14 @@ public:
 
 	// Gets the mode extender list for all workflow applications (append to customize a specific mode)
 	static TArray<FWorkflowApplicationModeExtender>& GetModeExtenderList() { return ModeExtenderList; }
+
+	/**
+	* Registers toolbar for a mode if not already registered
+	*
+	* @param InModeName - Name of the mode to register toolbar for
+	* @return nullptr if toolbar already registered.
+	*/
+	virtual UToolMenu* RegisterModeToolbarIfUnregistered(const FName InModeName);
 protected:
 
 	virtual void AddApplicationMode(FName ModeName, TSharedRef<FApplicationMode> Mode);
diff --git a/Engine/Source/Editor/LevelEditor/Private/LevelEditorCreateActorMenu.cpp b/Engine/Source/Editor/LevelEditor/Private/LevelEditorCreateActorMenu.cpp
--- a/Engine/Source/Editor/LevelEditor/Private/LevelEditorCreateActorMenu.cpp
+++ b/Engine/Source/Editor/LevelEditor/Private/LevelEditorCreateActorMenu.cpp
@@ -361,22 +361,22 @@ void LevelEditorCreateActorMenu::FillAddReplaceViewportContextMenuSections(UTool
 	if ( AssetMenuOptions.Num() == 0 )
 	{
 		{
-			FToolMenuSection& Section = Menu->AddSection("AddActor");
+			FToolMenuSection& Section = Menu->AddSection("ActorType");
 			Section.AddSubMenu(
-				"AddActor",
+				"LevelViewportContextMenu",
 				NSLOCTEXT("LevelViewportContextMenu", "AddActorHeading", "Place Actor") , 
 				NSLOCTEXT("LevelViewportContextMenu", "AddActorMenu_ToolTip", "Templates for adding a new actor to the world"),
 				FNewToolMenuDelegate::CreateStatic(&LevelEditorCreateActorMenu::FillAddReplaceActorMenu, EActorCreateMode::Add));
 		}
 
-		if ( CanReplaceActors() )
-		{
-			FToolMenuSection& Section = Menu->AddSection("ReplaceActor");
-			Section.AddSubMenu(
-				"ReplaceActor",
-				NSLOCTEXT("LevelViewportContextMenu", "ReplaceActorHeading", "Replace Selected Actors with") , 
-				NSLOCTEXT("LevelViewportContextMenu", "ReplaceActorMenu_ToolTip", "Templates for replacing selected with new actors in the world"),
-				FNewToolMenuDelegate::CreateStatic(&LevelEditorCreateActorMenu::FillAddReplaceActorMenu, EActorCreateMode::Replace));
+			if ( CanReplaceActors() )
+			{
+				Section.AddSubMenu(
+					"LevelViewportContextMenu",
+					NSLOCTEXT("LevelViewportContextMenu", "ReplaceActorHeading", "Replace Selected Actors with") , 
+					NSLOCTEXT("LevelViewportContextMenu", "ReplaceActorMenu_ToolTip", "Templates for replacing selected with new actors in the world"),
+					FNewToolMenuDelegate::CreateStatic(&LevelEditorCreateActorMenu::FillAddReplaceActorMenu, EActorCreateMode::Replace));
+			}
 		}
 	}
 	else
diff --git a/Engine/Source/Editor/LevelEditor/Private/LevelEditorToolBar.cpp b/Engine/Source/Editor/LevelEditor/Private/LevelEditorToolBar.cpp
--- a/Engine/Source/Editor/LevelEditor/Private/LevelEditorToolBar.cpp
+++ b/Engine/Source/Editor/LevelEditor/Private/LevelEditorToolBar.cpp
@@ -303,11 +303,11 @@ namespace LevelEditorActionHelpers
 	/**
 	 * Builds a sub-menu for selecting a class
 	 *
-	 * @param InMenuBuilder		Object to append menu items/widgets to
+	 * @param InMenu		Object to append menu items/widgets to
 	 * @param InRootClass		The root class to filter the Class Viewer by to only show children of
 	 * @param InOnClassPicked	Callback delegate to fire when a class is picked
 	 */
-	void GetSelectSettingsClassSubMenu(FMenuBuilder& InMenuBuilder, UClass* InRootClass, FOnClassPicked InOnClassPicked)
+	void GetSelectSettingsClassSubMenu(UToolMenu* InMenu, UClass* InRootClass, FOnClassPicked InOnClassPicked)
 	{
 		FClassViewerInitializationOptions Options;
 		Options.Mode = EClassViewerMode::ClassPicker;
@@ -329,19 +329,18 @@ namespace LevelEditorActionHelpers
 		TSharedRef<SWidget> ClassViewer = FModuleManager::LoadModuleChecked<FClassViewerModule>("ClassViewer").CreateClassViewer(Options, InOnClassPicked);
 		FFormatNamedArguments FormatArgs;
 		FormatArgs.Add(TEXT("RootClass"), RootClassName);
-		InMenuBuilder.BeginSection(NAME_None, FText::Format(NSLOCTEXT("LevelToolBarViewMenu", "SelectGameModeLabel", "Select {RootClass} class"), FormatArgs));
-		InMenuBuilder.AddWidget(ClassViewer, FText::GetEmpty(), true);
-		InMenuBuilder.EndSection();
+		FToolMenuSection& Section = InMenu->AddSection("SelectSettingsClass", FText::Format(NSLOCTEXT("LevelToolBarViewMenu", "SelectGameModeLabel", "Select {RootClass} class"), FormatArgs));
+		Section.AddEntry(FToolMenuEntry::InitWidget("ClassViewer", ClassViewer, FText::GetEmpty(), true));
 	}
 
 	/**
 	 * Builds a sub-menu for creating a class
 	 *
-	 * @param InMenuBuilder		Object to append menu items/widgets to
+	 * @param InMenu		Object to append menu items/widgets to
 	 * @param InRootClass		The root class to filter the Class Viewer by to only show children of
 	 * @param InOnClassPicked	Callback delegate to fire when a class is picked
 	 */
-	void GetCreateSettingsClassSubMenu(FMenuBuilder& InMenuBuilder, UClass* InRootClass, FOnClassPicked InOnClassPicked)
+	void GetCreateSettingsClassSubMenu(UToolMenu* InMenu, UClass* InRootClass, FOnClassPicked InOnClassPicked)
 	{
 		FClassViewerInitializationOptions Options;
 		Options.Mode = EClassViewerMode::ClassPicker;
@@ -362,9 +361,8 @@ namespace LevelEditorActionHelpers
 		TSharedRef<SWidget> ClassViewer = FModuleManager::LoadModuleChecked<FClassViewerModule>("ClassViewer").CreateClassViewer(Options, InOnClassPicked);
 		FFormatNamedArguments FormatArgs;
 		FormatArgs.Add(TEXT("RootClass"), RootClassName);
-		InMenuBuilder.BeginSection(NAME_None, FText::Format(NSLOCTEXT("LevelToolBarViewMenu", "CreateGameModeLabel", "Select {RootClass} parent class"), FormatArgs));
-		InMenuBuilder.AddWidget(ClassViewer, FText::GetEmpty(), true);
-		InMenuBuilder.EndSection();
+		FToolMenuSection& Section = InMenu->AddSection("CreateSettingsClass", FText::Format(NSLOCTEXT("LevelToolBarViewMenu", "CreateGameModeLabel", "Select {RootClass} parent class"), FormatArgs));
+		Section.AddEntry(FToolMenuEntry::InitWidget("ClassViewer", ClassViewer, FText::GetEmpty(), true));
 	}
 
 	/** Helper struct for passing all required data to the GetBlueprintSettingsSubMenu function */
@@ -425,11 +423,11 @@ namespace LevelEditorActionHelpers
 	/**
 	 * A sub-menu for the Blueprints dropdown, facilitates all the sub-menu actions such as creating, editing, and selecting classes for the world settings game mode.
 	 *
-	 * @param InMenuBuilder		Object to append menu items/widgets to
+	 * @param InMenu		Object to append menu items/widgets to
 	 * @param InCommandList		Commandlist for menu items
 	 * @param InSettingsData	All the data needed to create the menu actions
 	 */
-	void GetBlueprintSettingsSubMenu(FMenuBuilder& InMenuBuilder, const TSharedRef<FUICommandList> InCommandList, FBlueprintMenuSettings InSettingsData);
+	void GetBlueprintSettingsSubMenu(UToolMenu* InMenu, FBlueprintMenuSettings InSettingsData);
 
 	/** Returns TRUE if the class can be edited, always TRUE for Blueprints and for native classes a compiler must be present */
 	bool CanEditClass(UClass* InClass)
@@ -523,47 +521,54 @@ namespace LevelEditorActionHelpers
 #undef LOCTEXT_NAMESPACE
 	}
 
-	void CreateGameModeSubMenu(FMenuBuilder& InMenuBuilder, TSharedRef<FUICommandList> InCommandList, TWeakPtr< SLevelEditor > InLevelEditor, bool bInProjectSettings)
+	void CreateGameModeSubMenu(FToolMenuSection& Section, const FName InName, bool bInProjectSettings)
 	{
 #define LOCTEXT_NAMESPACE "LevelToolBarViewMenu"
-		LevelEditorActionHelpers::FBlueprintMenuSettings GameModeMenuSettings;
-		GameModeMenuSettings.EditCommand = 
-			FUIAction(
-				FExecuteAction::CreateStatic< TWeakPtr< SLevelEditor > >( &OpenGameModeBlueprint, InLevelEditor, bInProjectSettings )
-			);
-		GameModeMenuSettings.OnCreateClassPicked = FOnClassPicked::CreateStatic( &LevelEditorActionHelpers::OnCreateGameModeClassPicked, InLevelEditor, bInProjectSettings );
-		GameModeMenuSettings.OnSelectClassPicked = FOnClassPicked::CreateStatic( &LevelEditorActionHelpers::OnSelectGameModeClassPicked, InLevelEditor, bInProjectSettings );
-		GameModeMenuSettings.CurrentClass = LevelEditorActionHelpers::GetGameModeClass(InLevelEditor, bInProjectSettings);
-		GameModeMenuSettings.RootClass = AGameModeBase::StaticClass();
-		GameModeMenuSettings.LevelEditor = InLevelEditor;
-		GameModeMenuSettings.bIsProjectSettings = bInProjectSettings;
-
-		auto IsGameModeActive = [](TWeakPtr< SLevelEditor > InLevelEditorPtr, bool bInProjSettings)->bool
+		Section.AddDynamicEntry(InName, FNewToolMenuSectionDelegate::CreateLambda([=](FToolMenuSection& InSection)
 		{
-			UClass* WorldSettingsGameMode = LevelEditorActionHelpers::GetGameModeClass(InLevelEditorPtr, false);
-			if((WorldSettingsGameMode == nullptr) ^ bInProjSettings ) //(WorldSettingsGameMode && !bInProjectSettings) || (!WorldSettingsGameMode && bInProjectSettings) )
+			ULevelEditorMenuContext* Context = InSection.FindContext<ULevelEditorMenuContext>();
+			if (Context && Context->LevelEditor.IsValid())
 			{
-				return false;
-			}
-			return true;
-		};
+				LevelEditorActionHelpers::FBlueprintMenuSettings GameModeMenuSettings;
+				GameModeMenuSettings.EditCommand =
+					FUIAction(
+						FExecuteAction::CreateStatic< TWeakPtr< SLevelEditor > >(&OpenGameModeBlueprint, Context->LevelEditor, bInProjectSettings)
+					);
+				GameModeMenuSettings.OnCreateClassPicked = FOnClassPicked::CreateStatic(&LevelEditorActionHelpers::OnCreateGameModeClassPicked, Context->LevelEditor, bInProjectSettings);
+				GameModeMenuSettings.OnSelectClassPicked = FOnClassPicked::CreateStatic(&LevelEditorActionHelpers::OnSelectGameModeClassPicked, Context->LevelEditor, bInProjectSettings);
+				GameModeMenuSettings.CurrentClass = LevelEditorActionHelpers::GetGameModeClass(Context->LevelEditor, bInProjectSettings);
+				GameModeMenuSettings.RootClass = AGameModeBase::StaticClass();
+				GameModeMenuSettings.LevelEditor = Context->LevelEditor;
+				GameModeMenuSettings.bIsProjectSettings = bInProjectSettings;
+
+				auto IsGameModeActive = [](TWeakPtr< SLevelEditor > InLevelEditorPtr, bool bInProjSettings)->bool
+				{
+					UClass* WorldSettingsGameMode = LevelEditorActionHelpers::GetGameModeClass(InLevelEditorPtr, false);
+					if ((WorldSettingsGameMode == nullptr) ^ bInProjSettings) //(WorldSettingsGameMode && !bInProjectSettings) || (!WorldSettingsGameMode && bInProjectSettings) )
+					{
+						return false;
+					}
+					return true;
+				};
 
-		InMenuBuilder.AddSubMenu( LevelEditorActionHelpers::GetOpenGameModeBlueprintLabel(InLevelEditor, bInProjectSettings), 
-									GetSubMenuTooltip(GameModeMenuSettings.CurrentClass, GameModeMenuSettings.RootClass, bInProjectSettings), 
-									FNewMenuDelegate::CreateStatic( &LevelEditorActionHelpers::GetBlueprintSettingsSubMenu, InCommandList, GameModeMenuSettings), 
-									FUIAction(FExecuteAction(), FCanExecuteAction(), FIsActionChecked::CreateStatic(IsGameModeActive, InLevelEditor, bInProjectSettings)),
-									NAME_None, EUserInterfaceActionType::RadioButton );
+				InSection.AddSubMenu(InName, LevelEditorActionHelpers::GetOpenGameModeBlueprintLabel(Context->LevelEditor, bInProjectSettings),
+					GetSubMenuTooltip(GameModeMenuSettings.CurrentClass, GameModeMenuSettings.RootClass, bInProjectSettings),
+					FNewToolMenuDelegate::CreateStatic(&LevelEditorActionHelpers::GetBlueprintSettingsSubMenu, GameModeMenuSettings),
+					FUIAction(FExecuteAction(), FCanExecuteAction(), FIsActionChecked::CreateStatic(IsGameModeActive, Context->LevelEditor, bInProjectSettings)),
+					EUserInterfaceActionType::RadioButton);
+			}
+		}));
 #undef LOCTEXT_NAMESPACE
 	}
 
 	/**
 	 * Builds the game mode's sub menu objects
 	 *
-	 * @param InMenuBuilder		Object to append menu items/widgets to
+	 * @param InSection			Object to append menu items/widgets to
 	 * @param InCommandList		Commandlist for menu items
 	 * @param InSettingsData	All the data needed to create the menu actions
 	 */
-	void GetGameModeSubMenu(FMenuBuilder& InMenuBuilder, const TSharedRef<FUICommandList> InCommandList, const FBlueprintMenuSettings& InSettingsData)
+	void GetGameModeSubMenu(FToolMenuSection& InSection, const FBlueprintMenuSettings& InSettingsData)
 	{
 #define LOCTEXT_NAMESPACE "LevelToolBarViewMenu"
 		// Game State
@@ -579,9 +584,9 @@ namespace LevelEditorActionHelpers
 		GameStateMenuSettings.LevelEditor = InSettingsData.LevelEditor;
 		GameStateMenuSettings.bIsProjectSettings = InSettingsData.bIsProjectSettings;
 
-		InMenuBuilder.AddSubMenu( LevelEditorActionHelpers::GetOpenGameStateBlueprintLabel(InSettingsData.LevelEditor, InSettingsData.bIsProjectSettings),
+		InSection.AddSubMenu("OpenGameStateBlueprint", LevelEditorActionHelpers::GetOpenGameStateBlueprintLabel(InSettingsData.LevelEditor, InSettingsData.bIsProjectSettings),
 			GetSubMenuTooltip(GameStateMenuSettings.CurrentClass, GameStateMenuSettings.RootClass, InSettingsData.bIsProjectSettings),
-			FNewMenuDelegate::CreateStatic( &LevelEditorActionHelpers::GetBlueprintSettingsSubMenu, InCommandList, GameStateMenuSettings )
+			FNewToolMenuDelegate::CreateStatic( &LevelEditorActionHelpers::GetBlueprintSettingsSubMenu, GameStateMenuSettings )
 		);
 
 		// Pawn
@@ -597,9 +602,9 @@ namespace LevelEditorActionHelpers
 		PawnMenuSettings.LevelEditor = InSettingsData.LevelEditor;
 		PawnMenuSettings.bIsProjectSettings = InSettingsData.bIsProjectSettings;
 
-		InMenuBuilder.AddSubMenu( LevelEditorActionHelpers::GetOpenPawnBlueprintLabel(InSettingsData.LevelEditor, InSettingsData.bIsProjectSettings),
+		InSection.AddSubMenu("OpenPawnBlueprint", LevelEditorActionHelpers::GetOpenPawnBlueprintLabel(InSettingsData.LevelEditor, InSettingsData.bIsProjectSettings),
 			GetSubMenuTooltip(PawnMenuSettings.CurrentClass, PawnMenuSettings.RootClass, InSettingsData.bIsProjectSettings),
-			FNewMenuDelegate::CreateStatic( &LevelEditorActionHelpers::GetBlueprintSettingsSubMenu, InCommandList, PawnMenuSettings )
+			FNewToolMenuDelegate::CreateStatic( &LevelEditorActionHelpers::GetBlueprintSettingsSubMenu, PawnMenuSettings )
 		);
 
 		// HUD
@@ -615,9 +620,9 @@ namespace LevelEditorActionHelpers
 		HUDMenuSettings.LevelEditor = InSettingsData.LevelEditor;
 		HUDMenuSettings.bIsProjectSettings = InSettingsData.bIsProjectSettings;
 
-		InMenuBuilder.AddSubMenu( LevelEditorActionHelpers::GetOpenHUDBlueprintLabel(InSettingsData.LevelEditor, InSettingsData.bIsProjectSettings),
+		InSection.AddSubMenu("OpenHUDBlueprint", LevelEditorActionHelpers::GetOpenHUDBlueprintLabel(InSettingsData.LevelEditor, InSettingsData.bIsProjectSettings),
 			GetSubMenuTooltip(HUDMenuSettings.CurrentClass, HUDMenuSettings.RootClass, InSettingsData.bIsProjectSettings),
-			FNewMenuDelegate::CreateStatic( &LevelEditorActionHelpers::GetBlueprintSettingsSubMenu, InCommandList, HUDMenuSettings )
+			FNewToolMenuDelegate::CreateStatic( &LevelEditorActionHelpers::GetBlueprintSettingsSubMenu, HUDMenuSettings )
 		);
 
 		// Player Controller
@@ -633,9 +638,9 @@ namespace LevelEditorActionHelpers
 		PlayerControllerMenuSettings.LevelEditor = InSettingsData.LevelEditor;
 		PlayerControllerMenuSettings.bIsProjectSettings = InSettingsData.bIsProjectSettings;
 
-		InMenuBuilder.AddSubMenu( LevelEditorActionHelpers::GetOpenPlayerControllerBlueprintLabel(InSettingsData.LevelEditor, InSettingsData.bIsProjectSettings),
+		InSection.AddSubMenu("OpenPlayerControllerBlueprint", LevelEditorActionHelpers::GetOpenPlayerControllerBlueprintLabel(InSettingsData.LevelEditor, InSettingsData.bIsProjectSettings),
 			GetSubMenuTooltip(PlayerControllerMenuSettings.CurrentClass, PlayerControllerMenuSettings.RootClass, InSettingsData.bIsProjectSettings),
-			FNewMenuDelegate::CreateStatic( &LevelEditorActionHelpers::GetBlueprintSettingsSubMenu, InCommandList, PlayerControllerMenuSettings )
+			FNewToolMenuDelegate::CreateStatic( &LevelEditorActionHelpers::GetBlueprintSettingsSubMenu, PlayerControllerMenuSettings )
 		);
 #undef LOCTEXT_NAMESPACE
 	}
@@ -656,12 +661,10 @@ namespace LevelEditorActionHelpers
 	};
 }
 
-void LevelEditorActionHelpers::GetBlueprintSettingsSubMenu(FMenuBuilder& InMenuBuilder, const TSharedRef<FUICommandList> InCommandList, FBlueprintMenuSettings InSettingsData)
+void LevelEditorActionHelpers::GetBlueprintSettingsSubMenu(UToolMenu* Menu, FBlueprintMenuSettings InSettingsData)
 {
 #define LOCTEXT_NAMESPACE "LevelToolBarViewMenu"
 
-	InMenuBuilder.PushCommandList(InCommandList);
-
 	FSlateIcon EditBPIcon(FEditorStyle::Get().GetStyleSetName(), TEXT("PropertyWindow.Button_Edit"));
 	FSlateIcon NewBPIcon(FEditorStyle::Get().GetStyleSetName(), TEXT("PropertyWindow.Button_AddToArray"));
 	FText RootClassName = FText::FromString(InSettingsData.RootClass->GetName());
@@ -673,11 +676,12 @@ void LevelEditorActionHelpers::GetBlueprintSettingsSubMenu(FMenuBuilder& InMenuB
 		Args.Add(TEXT("RootClass"), RootClassName);
 		Args.Add(TEXT("TargetLocation"), InSettingsData.bIsProjectSettings? LOCTEXT("Project", "project") : LOCTEXT("World", "world"));
 
+		FToolMenuSection& Section = Menu->AddSection("EditBlueprintOrClass");
 		if(InSettingsData.CurrentClass->ClassGeneratedBy)
 		{
 			FText BlueprintName = FText::FromString(InSettingsData.CurrentClass->ClassGeneratedBy->GetName());
 			Args.Add(TEXT("Blueprint"), BlueprintName);
-			InMenuBuilder.AddMenuEntry( FText::Format( LOCTEXT("EditBlueprint", "Edit {Blueprint}"), Args), FText::Format( LOCTEXT("EditBlueprint_Tooltip", "Open the {TargetLocation}'s assigned {RootClass} blueprint"), Args), EditBPIcon, InSettingsData.EditCommand );
+			Section.AddMenuEntry("EditBlueprint", FText::Format( LOCTEXT("EditBlueprint", "Edit {Blueprint}"), Args), FText::Format( LOCTEXT("EditBlueprint_Tooltip", "Open the {TargetLocation}'s assigned {RootClass} blueprint"), Args), EditBPIcon, InSettingsData.EditCommand );
 		}
 		else
 		{
@@ -687,7 +691,7 @@ void LevelEditorActionHelpers::GetBlueprintSettingsSubMenu(FMenuBuilder& InMenuB
 			FText MenuDescription = FText::Format( LOCTEXT("EditNativeClass", "Edit {Class}.h"), Args);
 			if(FSourceCodeNavigation::IsCompilerAvailable())
 			{
-				InMenuBuilder.AddMenuEntry( MenuDescription, FText::Format( LOCTEXT("EditNativeClass_Tooltip", "Open the {TargetLocation}'s assigned {RootClass} header"), Args), EditBPIcon, InSettingsData.EditCommand );
+				Section.AddMenuEntry("EditNativeClass", MenuDescription, FText::Format( LOCTEXT("EditNativeClass_Tooltip", "Open the {TargetLocation}'s assigned {RootClass} header"), Args), EditBPIcon, InSettingsData.EditCommand );
 			}
 			else
 			{
@@ -697,21 +701,20 @@ void LevelEditorActionHelpers::GetBlueprintSettingsSubMenu(FMenuBuilder& InMenuB
 				};
 
 				// There is no compiler present, this is always disabled with a tooltip to explain why
-				InMenuBuilder.AddMenuEntry( MenuDescription, FText::Format( LOCTEXT("CannotEditNativeClass_Tooltip", "Cannot edit the {TargetLocation}'s assigned {RootClass} header because no compiler is present!"), Args), EditBPIcon, FUIAction(FExecuteAction(), FCanExecuteAction::CreateStatic(CannotEditClass)) );
+				Section.AddMenuEntry("EditNativeClass", MenuDescription, FText::Format( LOCTEXT("CannotEditNativeClass_Tooltip", "Cannot edit the {TargetLocation}'s assigned {RootClass} header because no compiler is present!"), Args), EditBPIcon, FUIAction(FExecuteAction(), FCanExecuteAction::CreateStatic(CannotEditClass)) );
 			}
 		}
 	}
 
 	if(InSettingsData.bIsProjectSettings && InSettingsData.CurrentClass && InSettingsData.CurrentClass->IsChildOf(AGameModeBase::StaticClass()) && !FLevelEditorActionCallbacks::CanSelectGameModeBlueprint())
 	{
-		InMenuBuilder.BeginSection("CheckoutSection", LOCTEXT("CheckoutSection","Check Out Project Settings") );
+		FToolMenuSection& Section = Menu->AddSection("CheckoutSection", LOCTEXT("CheckoutSection", "Check Out Project Settings") );
 		TAttribute<FText> CheckOutLabel;
 		CheckOutLabel.BindStatic(&GetCheckOutLabel);
 
 		TAttribute<FText> CheckOutTooltip;
 		CheckOutTooltip.BindStatic(&GetCheckOutTooltip);
-		InMenuBuilder.AddMenuEntry(FLevelEditorCommands::Get().CheckOutProjectSettingsConfig, NAME_None, CheckOutLabel, CheckOutTooltip, FSlateIcon(FEditorStyle::Get().GetStyleSetName(), TEXT("Icons.Error")));
-		InMenuBuilder.EndSection();
+		Section.AddMenuEntry(FLevelEditorCommands::Get().CheckOutProjectSettingsConfig, CheckOutLabel, CheckOutTooltip, FSlateIcon(FEditorStyle::Get().GetStyleSetName(), TEXT("Icons.Error")));
 	}
 
 	auto CannotCreateSelectNativeProjectGameMode = [](bool bInIsProjectSettings) -> bool
@@ -724,49 +727,46 @@ void LevelEditorActionHelpers::GetBlueprintSettingsSubMenu(FMenuBuilder& InMenuB
 		return true;
 	};
 
+	FToolMenuSection& Section = Menu->AddSection("CreateBlueprint");
+
 	// Create a new GameMode, this is always available so the user can easily create a new one
-	InMenuBuilder.AddSubMenu( LOCTEXT("CreateBlueprint", "Create..."),
+	Section.AddSubMenu("CreateBlueprint", LOCTEXT("CreateBlueprint", "Create..."),
 		GetCreateMenuTooltip(GetGameModeClass(InSettingsData.LevelEditor, InSettingsData.bIsProjectSettings), InSettingsData.RootClass, InSettingsData.bIsProjectSettings),
-		FNewMenuDelegate::CreateStatic( &LevelEditorActionHelpers::GetCreateSettingsClassSubMenu, InSettingsData.RootClass, InSettingsData.OnCreateClassPicked ),
+		FNewToolMenuDelegate::CreateStatic( &LevelEditorActionHelpers::GetCreateSettingsClassSubMenu, InSettingsData.RootClass, InSettingsData.OnCreateClassPicked ),
 		FUIAction(
 			FExecuteAction(), 
 			InSettingsData.RootClass == AGameModeBase::StaticClass()? 
 				FCanExecuteAction::CreateStatic(CannotCreateSelectNativeProjectGameMode, InSettingsData.bIsProjectSettings) 
 				: FCanExecuteAction::CreateStatic( &CanCreateSelectSubClass, GetGameModeClass(InSettingsData.LevelEditor, InSettingsData.bIsProjectSettings), InSettingsData.bIsProjectSettings )
 		),
-		NAME_None, EUserInterfaceActionType::Button, false, NewBPIcon
+		EUserInterfaceActionType::Button, false, NewBPIcon
 	);
 
 	// Select a game mode, this is always available so the user can switch his selection
 	FFormatNamedArguments Args;
 	Args.Add(TEXT("RootClass"), RootClassName);
-	InMenuBuilder.AddSubMenu(FText::Format(LOCTEXT("SelectGameModeClass", "Select {RootClass} Class"), Args),
+	Section.AddSubMenu("SelectGameModeClass", FText::Format(LOCTEXT("SelectGameModeClass", "Select {RootClass} Class"), Args),
 		GetSelectMenuTooltip(GetGameModeClass(InSettingsData.LevelEditor, InSettingsData.bIsProjectSettings), InSettingsData.RootClass, InSettingsData.bIsProjectSettings),
-		FNewMenuDelegate::CreateStatic( &LevelEditorActionHelpers::GetSelectSettingsClassSubMenu, InSettingsData.RootClass, InSettingsData.OnSelectClassPicked ),
+		FNewToolMenuDelegate::CreateStatic( &LevelEditorActionHelpers::GetSelectSettingsClassSubMenu, InSettingsData.RootClass, InSettingsData.OnSelectClassPicked ),
 		FUIAction(
 			FExecuteAction(), 
 			InSettingsData.RootClass == AGameModeBase::StaticClass()?
 				FCanExecuteAction::CreateStatic(CannotCreateSelectNativeProjectGameMode, InSettingsData.bIsProjectSettings) 
 				: FCanExecuteAction::CreateStatic( &CanCreateSelectSubClass, GetGameModeClass(InSettingsData.LevelEditor, InSettingsData.bIsProjectSettings), InSettingsData.bIsProjectSettings )
 		),
-		NAME_None, EUserInterfaceActionType::Button
+		EUserInterfaceActionType::Button
 	);
 
 	// For GameMode classes only, there are some sub-classes we need to add to the menu
 	if(InSettingsData.RootClass == AGameModeBase::StaticClass())
 	{
-		InMenuBuilder.BeginSection(NAME_None, LOCTEXT("GameModeClasses", "Game Mode Classes"));
-
+		FToolMenuSection& GameModeClassesSection = Menu->AddSection("GameModeClasses", LOCTEXT("GameModeClasses", "Game Mode Classes"));
 		if(InSettingsData.CurrentClass)
 		{
-			GetGameModeSubMenu(InMenuBuilder, InCommandList, InSettingsData);
+			GetGameModeSubMenu(GameModeClassesSection, InSettingsData);
 		}
-
-		InMenuBuilder.EndSection();
 	}
 
-	InMenuBuilder.PopCommandList();
-
 #undef LOCTEXT_NAMESPACE
 }
 
@@ -1173,13 +1173,21 @@ void FLevelEditorToolBar::RegisterLevelEditorToolBar( const TSharedRef<FUIComman
 {
 	RegisterSourceControlMenu();
 	RegisterCinematicsMenu();
+	RegisterBuildMenu();
+
+#if WITH_LIVE_CODING
+	RegisterCompileMenu();
+#endif
+
+	RegisterQuickSettingsMenu();
+	RegisterOpenBlueprintMenu();
 
 #define LOCTEXT_NAMESPACE "LevelEditorToolBar"
 
-	UToolMenu* ToolbarBuilder = UToolMenus::Get()->RegisterMenu("LevelEditor.LevelEditorToolBar", NAME_None, EMultiBoxType::ToolBar);
+	UToolMenu* Toolbar = UToolMenus::Get()->RegisterMenu("LevelEditor.LevelEditorToolBar", NAME_None, EMultiBoxType::ToolBar);
 
 	{
-		FToolMenuSection& Section = ToolbarBuilder->AddSection("File");
+		FToolMenuSection& Section = Toolbar->AddSection("File");
 
 		// Save All Levels
 		Section.AddEntry(FToolMenuEntry::InitToolBarButton(
@@ -1266,7 +1274,7 @@ void FLevelEditorToolBar::RegisterLevelEditorToolBar( const TSharedRef<FUIComman
 			Section.AddEntry(FToolMenuEntry::InitComboButton(
 				"SourceControl",
 				FUIAction(),
-				FOnGetContent::CreateStatic(&FLevelEditorToolBar::GenerateSourceControlMenu, InCommandList),
+				FOnGetContent::CreateStatic(&FLevelEditorToolBar::GenerateSourceControlMenu, InCommandList, TWeakPtr<SLevelEditor>(InLevelEditor)),
 				LOCTEXT("SourceControl_Label", "Source Control"),
 				TAttribute<FText>::Create(&FSourceControlStatus::GetSourceControlTooltip),
 				TAttribute<FSlateIcon>::Create(&FSourceControlStatus::GetSourceControlIcon),
@@ -1276,7 +1284,7 @@ void FLevelEditorToolBar::RegisterLevelEditorToolBar( const TSharedRef<FUIComman
 	}
 
 	{
-		FToolMenuSection& Section = ToolbarBuilder->AddSection("Content");
+		FToolMenuSection& Section = Toolbar->AddSection("Content");
 		Section.AddEntry(FToolMenuEntry::InitToolBarButton(FLevelEditorCommands::Get().OpenContentBrowser, LOCTEXT( "ContentBrowser_Override", "Content" ), TAttribute<FText>(), TAttribute<FSlateIcon>(), "LevelToolbarContent"));
 		if (FLauncherPlatformModule::Get()->CanOpenLauncher(true)) 
 		{
@@ -1284,12 +1292,12 @@ void FLevelEditorToolBar::RegisterLevelEditorToolBar( const TSharedRef<FUIComman
 		}
 	}
 
-	FToolMenuSection& SettingsSection = ToolbarBuilder->AddSection("Settings");
+	FToolMenuSection& SettingsSection = Toolbar->AddSection("Settings");
 	{
 		SettingsSection.AddEntry(FToolMenuEntry::InitComboButton(
 			"LevelToolbarQuickSettings",
 			FUIAction(),
-			FOnGetContent::CreateStatic(&FLevelEditorToolBar::GenerateQuickSettingsMenu, InCommandList),
+			FOnGetContent::CreateStatic(&FLevelEditorToolBar::GenerateQuickSettingsMenu, InCommandList, TWeakPtr<SLevelEditor>(InLevelEditor)),
 			LOCTEXT("QuickSettingsCombo", "Settings"),
 			LOCTEXT("QuickSettingsCombo_ToolTip", "Project and Editor settings"),
 			FSlateIcon(FEditorStyle::GetStyleSetName(), "LevelEditor.GameSettings"),
@@ -1412,7 +1420,7 @@ void FLevelEditorToolBar::RegisterLevelEditorToolBar( const TSharedRef<FUIComman
 	}
 
 	{
-		FToolMenuSection& Section = ToolbarBuilder->AddSection("Misc");
+		FToolMenuSection& Section = Toolbar->AddSection("Misc");
 		Section.AddEntry(FToolMenuEntry::InitComboButton(
 			"OpenBlueprint",
 			FUIAction(),
@@ -1435,7 +1443,7 @@ void FLevelEditorToolBar::RegisterLevelEditorToolBar( const TSharedRef<FUIComman
 	}
 
 	{
-		FToolMenuSection& Section = ToolbarBuilder->AddSection("Compile");
+		FToolMenuSection& Section = Toolbar->AddSection("Compile");
 		// Build			
 		Section.AddEntry(FToolMenuEntry::InitToolBarButton( FLevelEditorCommands::Get().Build, LOCTEXT("BuildAll", "Build")) );
 
@@ -1443,55 +1451,55 @@ void FLevelEditorToolBar::RegisterLevelEditorToolBar( const TSharedRef<FUIComman
 		Section.AddEntry(FToolMenuEntry::InitComboButton(
 			"BuildComboButton",
 			FUIAction(),
-			FOnGetContent::CreateStatic( &FLevelEditorToolBar::GenerateBuildMenuContent, InCommandList ),
+			FOnGetContent::CreateStatic( &FLevelEditorToolBar::GenerateBuildMenuContent, InCommandList, TWeakPtr<SLevelEditor>(InLevelEditor) ),
 			LOCTEXT( "BuildCombo_Label", "Build Options" ),
 			LOCTEXT( "BuildComboToolTip", "Build options menu" ),
 			FSlateIcon(FEditorStyle::GetStyleSetName(), "LevelEditor.Build"),
 			true));
 
-		// Only show the compile options on machines with the solution (assuming they can build it)
-		if ( FSourceCodeNavigation::IsCompilerAvailable() )
-		{
-			// Since we can always add new code to the project, only hide these buttons if we haven't done so yet
-			Section.AddEntry(FToolMenuEntry::InitToolBarButton(
-				"CompileButton",
-				FUIAction(
-					FExecuteAction::CreateStatic(&FLevelEditorActionCallbacks::RecompileGameCode_Clicked),
-					FCanExecuteAction::CreateStatic(&FLevelEditorActionCallbacks::Recompile_CanExecute),
-					FIsActionChecked(),
-					FIsActionButtonVisible::CreateStatic(FLevelEditorActionCallbacks::CanShowSourceCodeActions)),
-				LOCTEXT( "CompileMenuButton", "Compile" ),
-				FLevelEditorCommands::Get().RecompileGameCode->GetDescription(),
-				FSlateIcon(FEditorStyle::GetStyleSetName(), "LevelEditor.Recompile")
+		Section.AddDynamicEntry("CompilerAvailable", FNewToolMenuSectionDelegate::CreateLambda([](FToolMenuSection& InSection)
+		{
+			// Only show the compile options on machines with the solution (assuming they can build it)
+			if (FSourceCodeNavigation::IsCompilerAvailable())
+			{
+				// Since we can always add new code to the project, only hide these buttons if we haven't done so yet
+				InSection.AddEntry(FToolMenuEntry::InitToolBarButton(
+					"CompileButton",
+					FUIAction(
+						FExecuteAction::CreateStatic(&FLevelEditorActionCallbacks::RecompileGameCode_Clicked),
+						FCanExecuteAction::CreateStatic(&FLevelEditorActionCallbacks::Recompile_CanExecute),
+						FIsActionChecked(),
+						FIsActionButtonVisible::CreateStatic(FLevelEditorActionCallbacks::CanShowSourceCodeActions)),
+					LOCTEXT("CompileMenuButton", "Compile"),
+					FLevelEditorCommands::Get().RecompileGameCode->GetDescription(),
+					FSlateIcon(FEditorStyle::GetStyleSetName(), "LevelEditor.Recompile")
 				));
 
 #if WITH_LIVE_CODING
-			Section.AddEntry(FToolMenuEntry::InitComboButton(
-				"CompileComboButton",
-				FUIAction(
-					FExecuteAction(),
-					FCanExecuteAction(),
-					FIsActionChecked(),
-					FIsActionButtonVisible::CreateStatic(FLevelEditorActionCallbacks::CanShowSourceCodeActions)), 
-				FOnGetContent::CreateStatic( &FLevelEditorToolBar::GenerateCompileMenuContent, InCommandList ),
-				LOCTEXT( "CompileCombo_Label", "Compile Options" ),
-				LOCTEXT( "CompileComboToolTip", "Compile options menu" ),
-				FSlateIcon(FEditorStyle::GetStyleSetName(), "LevelEditor.Recompile"),
-				true
+				InSection.AddEntry(FToolMenuEntry::InitComboButton(
+					"CompileComboButton",
+					FUIAction(
+						FExecuteAction(),
+						FCanExecuteAction(),
+						FIsActionChecked(),
+						FIsActionButtonVisible::CreateStatic(FLevelEditorActionCallbacks::CanShowSourceCodeActions)),
+					FNewToolMenuWidgetChoice(),
+					LOCTEXT("CompileCombo_Label", "Compile Options"),
+					LOCTEXT("CompileComboToolTip", "Compile options menu"),
+					FSlateIcon(FEditorStyle::GetStyleSetName(), "LevelEditor.Recompile"),
+					true
 				));
 #endif
-		}
+			}
+		}));
 	}
 
-	ToolbarBuilder->AddDynamicSection("Game", FNewToolBarDelegateLegacy::CreateLambda([](FToolBarBuilder& InBuilder, UToolMenu* InData)
 	{
-		InBuilder.BeginSection("Game");
+		FToolMenuSection& Section = Toolbar->AddSection("Game");
 
 		// Add the shared play-world commands that will be shown on the Kismet toolbar as well
-		FPlayWorldCommands::BuildToolbar(InBuilder, true);
-
-		InBuilder.EndSection();
-	}));
+		FPlayWorldCommands::BuildToolbar(Section, true);
+	}
 
 #undef LOCTEXT_NAMESPACE
 }
@@ -1525,38 +1533,59 @@ TSharedRef< SWidget > FLevelEditorToolBar::MakeLevelEditorToolBar( const TShared
 
 
 
-TSharedRef< SWidget > FLevelEditorToolBar::GenerateBuildMenuContent( TSharedRef<FUICommandList> InCommandList )
+TSharedRef< SWidget > FLevelEditorToolBar::GenerateBuildMenuContent( TSharedRef<FUICommandList> InCommandList, TWeakPtr<SLevelEditor> InLevelEditor)
 {
-#define LOCTEXT_NAMESPACE "LevelToolBarBuildMenu"
-
 	// Get all menu extenders for this context menu from the level editor module
 	FLevelEditorModule& LevelEditorModule = FModuleManager::GetModuleChecked<FLevelEditorModule>( TEXT("LevelEditor") );
 	TSharedPtr<FExtender> MenuExtender = LevelEditorModule.AssembleExtenders(InCommandList, LevelEditorModule.GetAllLevelEditorToolbarBuildMenuExtenders());
+	FToolMenuContext MenuContext(InCommandList, MenuExtender);
 
-	const bool bShouldCloseWindowAfterMenuSelection = true;
-	FMenuBuilder MenuBuilder( bShouldCloseWindowAfterMenuSelection, InCommandList, MenuExtender );
+	ULevelEditorMenuContext* LevelEditorMenuContext = NewObject<ULevelEditorMenuContext>();
+	LevelEditorMenuContext->LevelEditor = InLevelEditor;
+	MenuContext.AddObject(LevelEditorMenuContext);
+
+	return UToolMenus::Get()->GenerateWidget("LevelEditor.LevelEditorToolBar.BuildComboButton", MenuContext);
+}
+
+void FLevelEditorToolBar::RegisterBuildMenu()
+{
+#define LOCTEXT_NAMESPACE "LevelToolBarBuildMenu"
+
+	static const FName BaseMenuName = "LevelEditor.LevelEditorToolBar.BuildComboButton";
+	UToolMenu* Menu = UToolMenus::Get()->RegisterMenu(BaseMenuName);
 
 	struct FLightingMenus
 	{
+	public:
+
+		static void RegisterMenus(const FName InBaseMenuName)
+		{
+			FLightingMenus::RegisterLightingQualityMenu(InBaseMenuName);
+			FLightingMenus::RegisterLightingInfoMenu(InBaseMenuName);
+		}
+	private:
 
 		/** Generates a lighting quality sub-menu */
-		static void MakeLightingQualityMenu( FMenuBuilder& InMenuBuilder )
+		static void RegisterLightingQualityMenu(const FName InBaseMenuName)
 		{
-			InMenuBuilder.BeginSection("LevelEditorBuildLightingQuality", LOCTEXT( "LightingQualityHeading", "Quality Level" ) );
+			UToolMenu* SubMenu = UToolMenus::Get()->RegisterMenu(UToolMenus::JoinMenuPaths(InBaseMenuName, "LightingQuality"));
+
 			{
-				InMenuBuilder.AddMenuEntry( FLevelEditorCommands::Get().LightingQuality_Production );
-				InMenuBuilder.AddMenuEntry( FLevelEditorCommands::Get().LightingQuality_High );
-				InMenuBuilder.AddMenuEntry( FLevelEditorCommands::Get().LightingQuality_Medium );
-				InMenuBuilder.AddMenuEntry( FLevelEditorCommands::Get().LightingQuality_Preview );
+				FToolMenuSection& Section = SubMenu->AddSection("LevelEditorBuildLightingQuality", LOCTEXT( "LightingQualityHeading", "Quality Level" ) );
+				Section.AddMenuEntry( FLevelEditorCommands::Get().LightingQuality_Production );
+				Section.AddMenuEntry( FLevelEditorCommands::Get().LightingQuality_High );
+				Section.AddMenuEntry( FLevelEditorCommands::Get().LightingQuality_Medium );
+				Section.AddMenuEntry( FLevelEditorCommands::Get().LightingQuality_Preview );
 			}
-			InMenuBuilder.EndSection();
 		}
 
 		/** Generates a lighting density sub-menu */
-		static void MakeLightingDensityMenu( FMenuBuilder& InMenuBuilder )
+		static void RegisterLightingDensityMenu(const FName InBaseMenuName)
 		{
-			InMenuBuilder.BeginSection("LevelEditorBuildLightingDensity", LOCTEXT( "LightingDensityHeading", "Density Rendering" ) );
+			UToolMenu* SubMenu = UToolMenus::Get()->RegisterMenu(UToolMenus::JoinMenuPaths(InBaseMenuName, "LightingDensity"));
+
 			{
+				FToolMenuSection& Section = SubMenu->AddSection("LevelEditorBuildLightingDensity", LOCTEXT( "LightingDensityHeading", "Density Rendering" ) );
 				TSharedRef<SWidget> Ideal =		SNew(SHorizontalBox)
 												+SHorizontalBox::Slot()
 												.Padding( FMargin( 27.0f, 0.0f, 0.0f, 0.0f ) )
@@ -1568,7 +1597,7 @@ TSharedRef< SWidget > FLevelEditorToolBar::GenerateBuildMenuContent( TSharedRef<
 													.Value(FLevelEditorActionCallbacks::GetLightingDensityIdeal())
 													.OnValueChanged_Static(&FLevelEditorActionCallbacks::SetLightingDensityIdeal)
 												];
-				InMenuBuilder.AddWidget(Ideal, LOCTEXT("LightingDensity_Ideal","Ideal Density"));
+				Section.AddEntry(FToolMenuEntry::InitWidget("Ideal", Ideal, LOCTEXT("LightingDensity_Ideal","Ideal Density")));
 				
 				TSharedRef<SWidget> Maximum =	SNew(SHorizontalBox)
 												+SHorizontalBox::Slot()
@@ -1580,7 +1609,7 @@ TSharedRef< SWidget > FLevelEditorToolBar::GenerateBuildMenuContent( TSharedRef<
 													.Value(FLevelEditorActionCallbacks::GetLightingDensityMaximum())
 													.OnValueChanged_Static(&FLevelEditorActionCallbacks::SetLightingDensityMaximum)
 												];
-				InMenuBuilder.AddWidget(Maximum, LOCTEXT("LightingDensity_Maximum","Maximum Density"));
+				Section.AddEntry(FToolMenuEntry::InitWidget("Maximum", Maximum, LOCTEXT("LightingDensity_Maximum","Maximum Density")));
 
 				TSharedRef<SWidget> ClrScale =	SNew(SHorizontalBox)
 												+SHorizontalBox::Slot()
@@ -1593,7 +1622,7 @@ TSharedRef< SWidget > FLevelEditorToolBar::GenerateBuildMenuContent( TSharedRef<
 													.Value(FLevelEditorActionCallbacks::GetLightingDensityColorScale())
 													.OnValueChanged_Static(&FLevelEditorActionCallbacks::SetLightingDensityColorScale)
 												];
-				InMenuBuilder.AddWidget(ClrScale, LOCTEXT("LightingDensity_ColorScale","Color Scale"));
+				Section.AddEntry(FToolMenuEntry::InitWidget("ColorScale", ClrScale, LOCTEXT("LightingDensity_ColorScale","Color Scale")));
 
 				TSharedRef<SWidget> GrayScale =	SNew(SHorizontalBox)
 												+SHorizontalBox::Slot()
@@ -1606,18 +1635,19 @@ TSharedRef< SWidget > FLevelEditorToolBar::GenerateBuildMenuContent( TSharedRef<
 													.Value(FLevelEditorActionCallbacks::GetLightingDensityGrayscaleScale())
 													.OnValueChanged_Static(&FLevelEditorActionCallbacks::SetLightingDensityGrayscaleScale)
 												];
-				InMenuBuilder.AddWidget(GrayScale, LOCTEXT("LightingDensity_GrayscaleScale","Grayscale Scale"));
+				Section.AddEntry(FToolMenuEntry::InitWidget("GrayscaleScale", GrayScale, LOCTEXT("LightingDensity_GrayscaleScale","Grayscale Scale")));
 
-				InMenuBuilder.AddMenuEntry( FLevelEditorCommands::Get().LightingDensity_RenderGrayscale );
+				Section.AddMenuEntry( FLevelEditorCommands::Get().LightingDensity_RenderGrayscale );
 			}
-			InMenuBuilder.EndSection();
 		}
 
 		/** Generates a lighting resolution sub-menu */
-		static void MakeLightingResolutionMenu( FMenuBuilder& InMenuBuilder )
+		static void RegisterLightingResolutionMenu(const FName InBaseMenuName)
 		{
-			InMenuBuilder.BeginSection("LevelEditorBuildLightingResolution1", LOCTEXT( "LightingResolutionHeading1", "Primitive Types" ) );
+			UToolMenu* SubMenu = UToolMenus::Get()->RegisterMenu(UToolMenus::JoinMenuPaths(InBaseMenuName, "LightingResolution"));
+
 			{
+				FToolMenuSection& Section = SubMenu->AddSection("LevelEditorBuildLightingResolution1", LOCTEXT( "LightingResolutionHeading1", "Primitive Types" ) );
 				TSharedRef<SWidget> Meshes =	SNew(SHorizontalBox)
 												+SHorizontalBox::Slot()
 												.AutoWidth()
@@ -1654,7 +1684,7 @@ TSharedRef< SWidget > FLevelEditorToolBar::GenerateBuildMenuContent( TSharedRef<
 													.Value(FLevelEditorActionCallbacks::GetLightingResolutionMaxSMs())
 													.OnValueChanged_Static(&FLevelEditorActionCallbacks::SetLightingResolutionMaxSMs)
 												];
-				InMenuBuilder.AddWidget(Meshes, FText::GetEmpty(), true);
+				Section.AddEntry(FToolMenuEntry::InitWidget("Meshes", Meshes, FText::GetEmpty(), true));
 				
 				TSharedRef<SWidget> BSPs =		SNew(SHorizontalBox)
 												+SHorizontalBox::Slot()
@@ -1692,21 +1722,19 @@ TSharedRef< SWidget > FLevelEditorToolBar::GenerateBuildMenuContent( TSharedRef<
 													.Value(FLevelEditorActionCallbacks::GetLightingResolutionMaxBSPs())
 													.OnValueChanged_Static(&FLevelEditorActionCallbacks::SetLightingResolutionMaxBSPs)
 												];
-				InMenuBuilder.AddWidget(BSPs, FText::GetEmpty(), true);
+				Section.AddEntry(FToolMenuEntry::InitWidget("BSPs", BSPs, FText::GetEmpty(), true));
 			}
-			InMenuBuilder.EndSection(); //LevelEditorBuildLightingResolution1
 
-			InMenuBuilder.BeginSection("LevelEditorBuildLightingResolution2", LOCTEXT( "LightingResolutionHeading2", "Select Options" ) );
 			{
-				InMenuBuilder.AddMenuEntry( FLevelEditorCommands::Get().LightingResolution_CurrentLevel );
-				InMenuBuilder.AddMenuEntry( FLevelEditorCommands::Get().LightingResolution_SelectedLevels );
-				InMenuBuilder.AddMenuEntry( FLevelEditorCommands::Get().LightingResolution_AllLoadedLevels );
-				InMenuBuilder.AddMenuEntry( FLevelEditorCommands::Get().LightingResolution_SelectedObjectsOnly );
+				FToolMenuSection& Section = SubMenu->AddSection("LevelEditorBuildLightingResolution2", LOCTEXT( "LightingResolutionHeading2", "Select Options" ) );
+				Section.AddMenuEntry( FLevelEditorCommands::Get().LightingResolution_CurrentLevel );
+				Section.AddMenuEntry( FLevelEditorCommands::Get().LightingResolution_SelectedLevels );
+				Section.AddMenuEntry( FLevelEditorCommands::Get().LightingResolution_AllLoadedLevels );
+				Section.AddMenuEntry( FLevelEditorCommands::Get().LightingResolution_SelectedObjectsOnly );
 			}
-			InMenuBuilder.EndSection();
 
-			InMenuBuilder.BeginSection("LevelEditorBuildLightingResolution3", LOCTEXT( "LightingResolutionHeading3", "Ratio" ) );
 			{
+				FToolMenuSection& Section = SubMenu->AddSection("LevelEditorBuildLightingResolution3", LOCTEXT( "LightingResolutionHeading3", "Ratio" ) );
 				TSharedRef<SWidget> Ratio =		SNew(SSpinBox<int32>)
 												.MinValue(0)
 												.MaxValue(400)
@@ -1714,156 +1742,156 @@ TSharedRef< SWidget > FLevelEditorToolBar::GenerateBuildMenuContent( TSharedRef<
 												.Value(FLevelEditorActionCallbacks::GetLightingResolutionRatio())
 												.OnEndSliderMovement_Static(&FLevelEditorActionCallbacks::SetLightingResolutionRatio)
 												.OnValueCommitted_Static(&FLevelEditorActionCallbacks::SetLightingResolutionRatioCommit);
-				InMenuBuilder.AddWidget(Ratio, LOCTEXT( "LightingResolutionRatio", "Ratio" ));
+				Section.AddEntry(FToolMenuEntry::InitWidget("Ratio", Ratio, LOCTEXT( "LightingResolutionRatio", "Ratio" )));
 			}
-			InMenuBuilder.EndSection();
 		}
 
 		/** Generates a lighting info dialogs sub-menu */
-		static void MakeLightingInfoMenu( FMenuBuilder& InMenuBuilder )
+		static void RegisterLightingInfoMenu(const FName InBaseMenuName)
 		{
-			InMenuBuilder.BeginSection("LevelEditorBuildLightingInfo", LOCTEXT( "LightingInfoHeading", "Lighting Info Dialogs" ) );
+			RegisterLightingDensityMenu(UToolMenus::JoinMenuPaths(InBaseMenuName, "LightingInfo"));
+			RegisterLightingResolutionMenu(UToolMenus::JoinMenuPaths(InBaseMenuName, "LightingInfo"));
+
+			UToolMenu* SubMenu = UToolMenus::Get()->RegisterMenu(UToolMenus::JoinMenuPaths(InBaseMenuName, "LightingInfo"));
+
 			{
-				InMenuBuilder.AddSubMenu(
+				FToolMenuSection& Section = SubMenu->AddSection("LevelEditorBuildLightingInfo", LOCTEXT( "LightingInfoHeading", "Lighting Info Dialogs" ) );
+				Section.AddSubMenu(
+					"LightingDensity",
 					LOCTEXT( "LightingDensityRenderingSubMenu", "LightMap Density Rendering Options" ),
 					LOCTEXT( "LightingDensityRenderingSubMenu_ToolTip", "Shows the LightMap Density Rendering viewmode options." ),
-					FNewMenuDelegate::CreateStatic( &FLightingMenus::MakeLightingDensityMenu ) );
+					FNewToolMenuChoice() );
 
-				InMenuBuilder.AddSubMenu(
+				Section.AddSubMenu(
+					"LightingResolution",
 					LOCTEXT( "LightingResolutionAdjustmentSubMenu", "LightMap Resolution Adjustment" ),
 					LOCTEXT( "LightingResolutionAdjustmentSubMenu_ToolTip", "Shows the LightMap Resolution Adjustment options." ),
-					FNewMenuDelegate::CreateStatic( &FLightingMenus::MakeLightingResolutionMenu ) );
+					FNewToolMenuChoice() );
 
-				InMenuBuilder.AddMenuEntry( FLevelEditorCommands::Get().LightingStaticMeshInfo, NAME_None, LOCTEXT( "BuildLightingInfo_LightingStaticMeshInfo", "Lighting StaticMesh Info..." ) );
+				Section.AddMenuEntry( FLevelEditorCommands::Get().LightingStaticMeshInfo, LOCTEXT( "BuildLightingInfo_LightingStaticMeshInfo", "Lighting StaticMesh Info..." ) );
 			}
-			InMenuBuilder.EndSection();
 		}
 	};
 
-	MenuBuilder.BeginSection("LevelEditorLighting", LOCTEXT( "LightingHeading", "Lighting" ) );
+	FLightingMenus::RegisterMenus(BaseMenuName);
+
 	{
-		MenuBuilder.AddMenuEntry( FLevelEditorCommands::Get().BuildLightingOnly, NAME_None, LOCTEXT( "BuildLightingOnlyHeading", "Build Lighting Only" ) );
+		FToolMenuSection& Section = Menu->AddSection( "LevelEditorLighting", LOCTEXT( "LightingHeading", "Lighting" ) );
+		Section.AddMenuEntry( FLevelEditorCommands::Get().BuildLightingOnly, LOCTEXT( "BuildLightingOnlyHeading", "Build Lighting Only" ) );
 
-		MenuBuilder.AddSubMenu(
+		Section.AddSubMenu(
+			"LightingQuality",
 			LOCTEXT( "LightingQualitySubMenu", "Lighting Quality" ),
 			LOCTEXT( "LightingQualitySubMenu_ToolTip", "Allows you to select the quality level for precomputed lighting" ),
-			FNewMenuDelegate::CreateStatic( &FLightingMenus::MakeLightingQualityMenu ) );
+			FNewToolMenuChoice() );
 
-		MenuBuilder.AddSubMenu(
+		Section.AddSubMenu(
+			"LightingInfo",
 			LOCTEXT( "BuildLightingInfoSubMenu", "Lighting Info" ),
 			LOCTEXT( "BuildLightingInfoSubMenu_ToolTip", "Access the lighting info dialogs" ),
-			FNewMenuDelegate::CreateStatic( &FLightingMenus::MakeLightingInfoMenu ) );
+			FNewToolMenuChoice() );
 
-		MenuBuilder.AddMenuEntry( FLevelEditorCommands::Get().LightingBuildOptions_UseErrorColoring );
-		MenuBuilder.AddMenuEntry( FLevelEditorCommands::Get().LightingBuildOptions_ShowLightingStats );
+		Section.AddMenuEntry( FLevelEditorCommands::Get().LightingBuildOptions_UseErrorColoring );
+		Section.AddMenuEntry( FLevelEditorCommands::Get().LightingBuildOptions_ShowLightingStats );
 	}
-	MenuBuilder.EndSection();
 
-	MenuBuilder.BeginSection("LevelEditorReflections", LOCTEXT( "ReflectionHeading", "Reflections" ) );
 	{
-		MenuBuilder.AddMenuEntry( FLevelEditorCommands::Get().BuildReflectionCapturesOnly );
+		FToolMenuSection& Section = Menu->AddSection("LevelEditorReflections", LOCTEXT( "ReflectionHeading", "Reflections" ) );
+		Section.AddMenuEntry( FLevelEditorCommands::Get().BuildReflectionCapturesOnly );
 	}
-	MenuBuilder.EndSection();
 
-	MenuBuilder.BeginSection("LevelEditorVisibility", LOCTEXT( "VisibilityHeading", "Visibility" ) );
 	{
-		MenuBuilder.AddMenuEntry( FLevelEditorCommands::Get().BuildLightingOnly_VisibilityOnly );
+		FToolMenuSection& Section = Menu->AddSection("LevelEditorVisibility", LOCTEXT( "VisibilityHeading", "Visibility" ) );
+		Section.AddMenuEntry( FLevelEditorCommands::Get().BuildLightingOnly_VisibilityOnly );
 	}
-	MenuBuilder.EndSection();
 
-	MenuBuilder.BeginSection("LevelEditorGeometry", LOCTEXT( "GeometryHeading", "Geometry" ) );
 	{
-		MenuBuilder.AddMenuEntry( FLevelEditorCommands::Get().BuildGeometryOnly );
-		MenuBuilder.AddMenuEntry( FLevelEditorCommands::Get().BuildGeometryOnly_OnlyCurrentLevel );
+		FToolMenuSection& Section = Menu->AddSection("LevelEditorGeometry", LOCTEXT( "GeometryHeading", "Geometry" ) );
+		Section.AddMenuEntry( FLevelEditorCommands::Get().BuildGeometryOnly );
+		Section.AddMenuEntry( FLevelEditorCommands::Get().BuildGeometryOnly_OnlyCurrentLevel );
 	}
-	MenuBuilder.EndSection();
 
-	MenuBuilder.BeginSection("LevelEditorNavigation", LOCTEXT( "NavigationHeading", "Navigation" ) );
 	{
-		MenuBuilder.AddMenuEntry( FLevelEditorCommands::Get().BuildPathsOnly );
+		FToolMenuSection& Section = Menu->AddSection("LevelEditorNavigation", LOCTEXT( "NavigationHeading", "Navigation" ) );
+		Section.AddMenuEntry( FLevelEditorCommands::Get().BuildPathsOnly );
 	}
-	MenuBuilder.EndSection();
 
-	MenuBuilder.BeginSection("LevelEditorLOD", LOCTEXT("LODHeading", "Hierarchical LOD"));
 	{
-		MenuBuilder.AddMenuEntry(FLevelEditorCommands::Get().BuildLODsOnly);
+		FToolMenuSection& Section = Menu->AddSection("LevelEditorLOD", LOCTEXT("LODHeading", "Hierarchical LOD"));
+		Section.AddMenuEntry(FLevelEditorCommands::Get().BuildLODsOnly);
 	}
-	MenuBuilder.EndSection();
-
-	MenuBuilder.BeginSection("LevelEditorTextureStreaming", LOCTEXT("TextureStreamingHeading", "Texture Streaming"));
+	
 	{
-		if (CVarStreamingUseNewMetrics.GetValueOnAnyThread() != 0) // There is not point of in building texture streaming data with the old system.
+		FToolMenuSection& Section = Menu->AddSection("LevelEditorTextureStreaming", LOCTEXT("TextureStreamingHeading", "Texture Streaming"));
+		Section.AddDynamicEntry("BuildTextureStreamingOnly", FNewToolMenuSectionDelegate::CreateLambda([](FToolMenuSection& InSection)
 		{
-			MenuBuilder.AddMenuEntry(FLevelEditorCommands::Get().BuildTextureStreamingOnly);
-		}
-		MenuBuilder.AddMenuEntry(FLevelEditorCommands::Get().BuildVirtualTextureOnly);
+			if (CVarStreamingUseNewMetrics.GetValueOnAnyThread() != 0) // There is no point of in building texture streaming data with the old system.
+			{
+				InSection.AddMenuEntry(FLevelEditorCommands::Get().BuildTextureStreamingOnly);
+			}
+		}));
+		Section.AddMenuEntry(FLevelEditorCommands::Get().BuildVirtualTextureOnly);
 	}
-	MenuBuilder.EndSection();
 
-	MenuBuilder.BeginSection("LevelEditorAutomation", LOCTEXT( "AutomationHeading", "Automation" ) );
+	
 	{
-		MenuBuilder.AddMenuEntry( 
+		FToolMenuSection& Section = Menu->AddSection("LevelEditorAutomation", LOCTEXT( "AutomationHeading", "Automation" ) );
+		Section.AddMenuEntry( 
 			FLevelEditorCommands::Get().BuildAndSubmitToSourceControl,
-			NAME_None,
 			TAttribute<FText>(),
 			TAttribute<FText>(),
 			FSlateIcon(FEditorStyle::GetStyleSetName(), "LevelEditor.Tabs.BuildAndSubmit")
 			);
 	}
-	MenuBuilder.EndSection();
 
 	// Map Check
-	MenuBuilder.BeginSection("LevelEditorVerification", LOCTEXT( "VerificationHeading", "Verification" ) );
 	{
-		MenuBuilder.AddMenuEntry( FLevelEditorCommands::Get().MapCheck, NAME_None, LOCTEXT("OpenMapCheck", "Map Check") );
+		FToolMenuSection& Section = Menu->AddSection("LevelEditorVerification", LOCTEXT( "VerificationHeading", "Verification" ) );
+		Section.AddMenuEntry( FLevelEditorCommands::Get().MapCheck, LOCTEXT("OpenMapCheck", "Map Check") );
 	}
-	MenuBuilder.EndSection();
 
 #undef LOCTEXT_NAMESPACE
-
-	return MenuBuilder.MakeWidget();
 }
 
-static void MakeMaterialQualityLevelMenu( FMenuBuilder& MenuBuilder )
+static void MakeMaterialQualityLevelMenu( UToolMenu* InMenu )
 {
-	MenuBuilder.BeginSection("LevelEditorMaterialQualityLevel", NSLOCTEXT( "LevelToolBarViewMenu", "MaterialQualityLevelHeading", "Material Quality Level" ) );
 	{
-		MenuBuilder.AddMenuEntry(FLevelEditorCommands::Get().MaterialQualityLevel_Low);
-		MenuBuilder.AddMenuEntry(FLevelEditorCommands::Get().MaterialQualityLevel_Medium);
-		MenuBuilder.AddMenuEntry(FLevelEditorCommands::Get().MaterialQualityLevel_High);
+		FToolMenuSection& Section = InMenu->AddSection("LevelEditorMaterialQualityLevel", NSLOCTEXT( "LevelToolBarViewMenu", "MaterialQualityLevelHeading", "Material Quality Level" ) );
+		Section.AddMenuEntry(FLevelEditorCommands::Get().MaterialQualityLevel_Low);
+		Section.AddMenuEntry(FLevelEditorCommands::Get().MaterialQualityLevel_Medium);
+		Section.AddMenuEntry(FLevelEditorCommands::Get().MaterialQualityLevel_High);
 	}
-	MenuBuilder.EndSection();
 }
 
-static void MakeShaderModelPreviewMenu( FMenuBuilder& MenuBuilder )
+static void MakeShaderModelPreviewMenu( UToolMenu* InMenu )
 {
 #define LOCTEXT_NAMESPACE "LevelToolBarViewMenu"
 
-	MenuBuilder.BeginSection("EditorPreviewMode", LOCTEXT("EditorPreviewModeDevices", "Preview Devices"));
+	FToolMenuSection& Section = InMenu->AddSection("EditorPreviewMode", LOCTEXT("EditorPreviewModeDevices", "Preview Devices"));
 
 	// SM5
-	MenuBuilder.AddMenuEntry(FLevelEditorCommands::Get().PreviewPlatformOverride_SM5);
+	Section.AddMenuEntry(FLevelEditorCommands::Get().PreviewPlatformOverride_SM5);
 
 	// Android
 	bool bAndroidBuildForES31 = false;
 	GConfig->GetBool(TEXT("/Script/AndroidRuntimeSettings.AndroidRuntimeSettings"), TEXT("bBuildForES31"), bAndroidBuildForES31, GEngineIni);
 	if (bAndroidBuildForES31)
 	{
-		MenuBuilder.AddMenuEntry(FLevelEditorCommands::Get().PreviewPlatformOverride_AndroidGLES31);
+		Section.AddMenuEntry(FLevelEditorCommands::Get().PreviewPlatformOverride_AndroidGLES31);
 	}
 
 	bool bAndroidSupportsVulkan = false;
 	GConfig->GetBool(TEXT("/Script/AndroidRuntimeSettings.AndroidRuntimeSettings"), TEXT("bSupportsVulkan"), bAndroidSupportsVulkan, GEngineIni);
 	if (bAndroidSupportsVulkan)
 	{
-		MenuBuilder.AddMenuEntry(FLevelEditorCommands::Get().PreviewPlatformOverride_AndroidVulkanES31);
+		Section.AddMenuEntry(FLevelEditorCommands::Get().PreviewPlatformOverride_AndroidVulkanES31);
 	}
 
 	bool bAndroidBuildForES2 = false;
 	GConfig->GetBool(TEXT("/Script/AndroidRuntimeSettings.AndroidRuntimeSettings"), TEXT("bBuildForES2"), bAndroidBuildForES2, GEngineIni);
 	if (bAndroidBuildForES2)
 	{
-    	MenuBuilder.AddMenuEntry(FLevelEditorCommands::Get().PreviewPlatformOverride_AndroidGLES2);
+    	Section.AddMenuEntry(FLevelEditorCommands::Get().PreviewPlatformOverride_AndroidGLES2);
 	}
 
 	// iOS
@@ -1871,78 +1899,81 @@ static void MakeShaderModelPreviewMenu( FMenuBuilder& MenuBuilder )
 	GConfig->GetBool(TEXT("/Script/IOSRuntimeSettings.IOSRuntimeSettings"), TEXT("bSupportsMetal"), bIOSSupportsMetal, GEngineIni);
 	if (bIOSSupportsMetal)
 	{
-		MenuBuilder.AddMenuEntry(FLevelEditorCommands::Get().PreviewPlatformOverride_IOSMetalES31);
+		Section.AddMenuEntry(FLevelEditorCommands::Get().PreviewPlatformOverride_IOSMetalES31);
 	}
 
 	// HTML5
-	MenuBuilder.AddMenuEntry(FLevelEditorCommands::Get().PreviewPlatformOverride_HTML5);
-
-    MenuBuilder.EndSection();
+	Section.AddMenuEntry(FLevelEditorCommands::Get().PreviewPlatformOverride_HTML5);
 
 #undef LOCTEXT_NAMESPACE
 }
 
-static void MakeScalabilityMenu( FMenuBuilder& MenuBuilder )
+static void MakeScalabilityMenu( UToolMenu* InMenu )
 {
-	MenuBuilder.AddWidget(SNew(SScalabilitySettings), FText(), true);
+	{
+		FToolMenuSection& Section = InMenu->AddSection("Section");
+		Section.AddEntry(FToolMenuEntry::InitWidget("ScalabilitySettings", SNew(SScalabilitySettings), FText(), true));
+	}
 }
 
-static void MakePreviewSettingsMenu( FMenuBuilder& MenuBuilder )
+static void MakePreviewSettingsMenu( UToolMenu* InMenu )
 {
 #define LOCTEXT_NAMESPACE "LevelToolBarViewMenu"
-
-	MenuBuilder.BeginSection("LevelEditorPreview", LOCTEXT("PreviewHeading", "Previewing"));
 	{
-		MenuBuilder.AddMenuEntry(FLevelEditorCommands::Get().DrawBrushMarkerPolys);
-		MenuBuilder.AddMenuEntry(FLevelEditorCommands::Get().OnlyLoadVisibleInPIE);
-		MenuBuilder.AddMenuEntry(FLevelEditorCommands::Get().ToggleParticleSystemLOD);
-		MenuBuilder.AddMenuEntry(FLevelEditorCommands::Get().ToggleParticleSystemHelpers);
-		MenuBuilder.AddMenuEntry(FLevelEditorCommands::Get().ToggleFreezeParticleSimulation);
-		MenuBuilder.AddMenuEntry(FLevelEditorCommands::Get().ToggleLODViewLocking);
-		MenuBuilder.AddMenuEntry(FLevelEditorCommands::Get().LevelStreamingVolumePrevis);
+		FToolMenuSection& Section = InMenu->AddSection("LevelEditorPreview", LOCTEXT("PreviewHeading", "Previewing"));
+		Section.AddMenuEntry(FLevelEditorCommands::Get().DrawBrushMarkerPolys);
+		Section.AddMenuEntry(FLevelEditorCommands::Get().OnlyLoadVisibleInPIE);
+		Section.AddMenuEntry(FLevelEditorCommands::Get().ToggleParticleSystemLOD);
+		Section.AddMenuEntry(FLevelEditorCommands::Get().ToggleParticleSystemHelpers);
+		Section.AddMenuEntry(FLevelEditorCommands::Get().ToggleFreezeParticleSimulation);
+		Section.AddMenuEntry(FLevelEditorCommands::Get().ToggleLODViewLocking);
+		Section.AddMenuEntry(FLevelEditorCommands::Get().LevelStreamingVolumePrevis);
 	}
-	MenuBuilder.EndSection();
 #undef LOCTEXT_NAMESPACE
 }
 
 #if WITH_LIVE_CODING
-TSharedRef< SWidget > FLevelEditorToolBar::GenerateCompileMenuContent( TSharedRef<FUICommandList> InCommandList )
+void FLevelEditorToolBar::RegisterCompileMenu()
 {
 #define LOCTEXT_NAMESPACE "LevelToolBarCompileMenu"
 
-	const bool bShouldCloseWindowAfterMenuSelection = true;
-	FMenuBuilder MenuBuilder( bShouldCloseWindowAfterMenuSelection, InCommandList );
+	UToolMenu* Menu = UToolMenus::Get()->RegisterMenu("LevelEditor.LevelEditorToolBar.CompileComboButton");
 
-	MenuBuilder.BeginSection("LiveCodingMode", LOCTEXT( "LiveCodingMode", "General" ) );
 	{
-		MenuBuilder.AddMenuEntry( FLevelEditorCommands::Get().LiveCoding_Enable );
+		FToolMenuSection& Section = Menu->AddSection("LiveCodingMode", LOCTEXT( "LiveCodingMode", "General" ) );
+		Section.AddMenuEntry( FLevelEditorCommands::Get().LiveCoding_Enable );
 	}
-	MenuBuilder.EndSection();
 
-	MenuBuilder.BeginSection("LiveCodingActions", LOCTEXT( "LiveCodingActions", "Actions" ) );
 	{
-		MenuBuilder.AddMenuEntry( FLevelEditorCommands::Get().LiveCoding_StartSession );
-		MenuBuilder.AddMenuEntry( FLevelEditorCommands::Get().LiveCoding_ShowConsole );
-		MenuBuilder.AddMenuEntry( FLevelEditorCommands::Get().LiveCoding_Settings );
+		FToolMenuSection& Section = Menu->AddSection("LiveCodingActions", LOCTEXT( "LiveCodingActions", "Actions" ) );
+		Section.AddMenuEntry( FLevelEditorCommands::Get().LiveCoding_StartSession );
+		Section.AddMenuEntry( FLevelEditorCommands::Get().LiveCoding_ShowConsole );
+		Section.AddMenuEntry( FLevelEditorCommands::Get().LiveCoding_Settings );
 	}
-	MenuBuilder.EndSection();
-
-	return MenuBuilder.MakeWidget();
 
 #undef LOCTEXT_NAMESPACE
 }
 #endif
 
-TSharedRef< SWidget > FLevelEditorToolBar::GenerateQuickSettingsMenu( TSharedRef<FUICommandList> InCommandList )
+TSharedRef< SWidget > FLevelEditorToolBar::GenerateQuickSettingsMenu(TSharedRef<FUICommandList> InCommandList, TWeakPtr<SLevelEditor> InLevelEditor)
 {
-#define LOCTEXT_NAMESPACE "LevelToolBarViewMenu"
-
 	// Get all menu extenders for this context menu from the level editor module
 	FLevelEditorModule& LevelEditorModule = FModuleManager::GetModuleChecked<FLevelEditorModule>( TEXT("LevelEditor") );
 	TSharedPtr<FExtender> MenuExtender = LevelEditorModule.AssembleExtenders(InCommandList, LevelEditorModule.GetAllLevelEditorToolbarViewMenuExtenders());
 
-	const bool bShouldCloseWindowAfterMenuSelection = true;
-	FMenuBuilder MenuBuilder( bShouldCloseWindowAfterMenuSelection, InCommandList, MenuExtender );
+	FToolMenuContext MenuContext(InCommandList, MenuExtender);
+	ULevelEditorMenuContext* LevelEditorMenuContext = NewObject<ULevelEditorMenuContext>();
+	LevelEditorMenuContext->LevelEditor = InLevelEditor;
+	MenuContext.AddObject(LevelEditorMenuContext);
+
+	return UToolMenus::Get()->GenerateWidget("LevelEditor.LevelEditorToolBar.LevelToolbarQuickSettings", MenuContext);
+}
+
+void FLevelEditorToolBar::RegisterQuickSettingsMenu()
+{
+#define LOCTEXT_NAMESPACE "LevelToolBarViewMenu"
+
+	UToolMenu* Menu = UToolMenus::Get()->RegisterMenu("LevelEditor.LevelEditorToolBar.LevelToolbarQuickSettings");
 
 	struct Local
 	{
@@ -1952,56 +1983,60 @@ TSharedRef< SWidget > FLevelEditorToolBar::GenerateQuickSettingsMenu( TSharedRef
 		}
 	};
 
-	MenuBuilder.BeginSection("ProjectSettingsSection", LOCTEXT("ProjectSettings","Game Specific Settings") );
 	{
+		FToolMenuSection& Section = Menu->AddSection("ProjectSettingsSection", LOCTEXT("ProjectSettings", "Game Specific Settings"));
 
-		MenuBuilder.AddMenuEntry(FLevelEditorCommands::Get().WorldProperties);
+		Section.AddMenuEntry(FLevelEditorCommands::Get().WorldProperties);
 
-		MenuBuilder.AddMenuEntry(
+		Section.AddMenuEntry(
+			"ProjectSettings",
 			LOCTEXT("ProjectSettingsMenuLabel", "Project Settings..."),
 			LOCTEXT("ProjectSettingsMenuToolTip", "Change the settings of the currently loaded project"),
 			FSlateIcon(FEditorStyle::GetStyleSetName(), "ProjectSettings.TabIcon"),
 			FUIAction(FExecuteAction::CreateStatic(&Local::OpenSettings, FName("Project"), FName("Project"), FName("General")))
 			);
 
-		if (IModularFeatures::Get().IsModularFeatureAvailable(EditorFeatures::PluginsEditor))
+		Section.AddDynamicEntry("PluginsEditor", FNewToolMenuDelegateLegacy::CreateLambda([](FMenuBuilder& InMenuBuilder, UToolMenu* InMenu)
 		{
-			FGlobalTabmanager::Get()->PopulateTabSpawnerMenu(MenuBuilder, "PluginsEditor");
-		}
+			if (IModularFeatures::Get().IsModularFeatureAvailable(EditorFeatures::PluginsEditor))
+			{
+				FGlobalTabmanager::Get()->PopulateTabSpawnerMenu(InMenuBuilder, "PluginsEditor");
+			}
+		}));
 	}
-	MenuBuilder.EndSection();
 
-	MenuBuilder.BeginSection("LevelEditorSelection", LOCTEXT("SelectionHeading","Selection") );
 	{
-		MenuBuilder.AddMenuEntry( FLevelEditorCommands::Get().AllowTranslucentSelection );
-		MenuBuilder.AddMenuEntry( FLevelEditorCommands::Get().AllowGroupSelection );
-		MenuBuilder.AddMenuEntry( FLevelEditorCommands::Get().StrictBoxSelect );
-		MenuBuilder.AddMenuEntry( FLevelEditorCommands::Get().TransparentBoxSelect );
-		MenuBuilder.AddMenuEntry( FLevelEditorCommands::Get().ShowTransformWidget );
+		FToolMenuSection& Section = Menu->AddSection("LevelEditorSelection", LOCTEXT("SelectionHeading","Selection") );
+		Section.AddMenuEntry( FLevelEditorCommands::Get().AllowTranslucentSelection );
+		Section.AddMenuEntry( FLevelEditorCommands::Get().AllowGroupSelection );
+		Section.AddMenuEntry( FLevelEditorCommands::Get().StrictBoxSelect );
+		Section.AddMenuEntry( FLevelEditorCommands::Get().TransparentBoxSelect );
+		Section.AddMenuEntry( FLevelEditorCommands::Get().ShowTransformWidget );
 	}
-	MenuBuilder.EndSection();
 
-	MenuBuilder.BeginSection("LevelEditorScalability", LOCTEXT("ScalabilityHeading", "Scalability") );
 	{
-		MenuBuilder.AddSubMenu(
+		FToolMenuSection& Section = Menu->AddSection("LevelEditorScalability", LOCTEXT("ScalabilityHeading", "Scalability") );
+		Section.AddSubMenu(
+			"Scalability",
 			LOCTEXT( "ScalabilitySubMenu", "Engine Scalability Settings" ),
 			LOCTEXT( "ScalabilitySubMenu_ToolTip", "Open the engine scalability settings" ),
-			FNewMenuDelegate::CreateStatic( &MakeScalabilityMenu ) );
+			FNewToolMenuDelegate::CreateStatic( &MakeScalabilityMenu ) );
 
-		MenuBuilder.AddSubMenu(
+		Section.AddSubMenu(
+			"MaterialQualityLevel",
 			LOCTEXT( "MaterialQualityLevelSubMenu", "Material Quality Level" ),
 			LOCTEXT( "MaterialQualityLevelSubMenu_ToolTip", "Sets the value of the CVar \"r.MaterialQualityLevel\" (low=0, high=1, medium=2). This affects materials via the QualitySwitch material expression." ),
-			FNewMenuDelegate::CreateStatic( &MakeMaterialQualityLevelMenu ) );
+			FNewToolMenuDelegate::CreateStatic( &MakeMaterialQualityLevelMenu ) );
 
-		MenuBuilder.AddSubMenu(
+		Section.AddSubMenu(
+			"FeatureLevelPreview",
 			LOCTEXT("FeatureLevelPreviewSubMenu", "Preview Rendering Level"),
 			LOCTEXT("FeatureLevelPreviewSubMenu_ToolTip", "Sets the rendering level used by the main editor"),
-			FNewMenuDelegate::CreateStatic(&MakeShaderModelPreviewMenu));
+			FNewToolMenuDelegate::CreateStatic(&MakeShaderModelPreviewMenu));
 	}
-	MenuBuilder.EndSection();
 
-	MenuBuilder.BeginSection("LevelEditorAudio", LOCTEXT("AudioHeading", "Real Time Audio") );
 	{
+		FToolMenuSection& Section = Menu->AddSection("LevelEditorAudio", LOCTEXT("AudioHeading", "Real Time Audio") );
 		TSharedRef<SWidget> VolumeItem = SNew(SHorizontalBox)
 											+SHorizontalBox::Slot()
 											.FillWidth(0.9f)
@@ -2017,13 +2052,12 @@ TSharedRef< SWidget > FLevelEditorToolBar::GenerateQuickSettingsMenu( TSharedRef
 											+SHorizontalBox::Slot()
 											.FillWidth(0.1f);
 
-		MenuBuilder.AddWidget(VolumeItem, LOCTEXT("VolumeControlLabel","Volume"));
+		Section.AddEntry(FToolMenuEntry::InitWidget("Volume", VolumeItem, LOCTEXT("VolumeControlLabel","Volume")));
 	}
-	MenuBuilder.EndSection();
 
-	MenuBuilder.BeginSection( "Snapping", LOCTEXT("SnappingHeading","Snapping") );
 	{
-		MenuBuilder.AddMenuEntry( FLevelEditorCommands::Get().EnableActorSnap );
+		FToolMenuSection& Section = Menu->AddSection( "Snapping", LOCTEXT("SnappingHeading","Snapping") );
+		Section.AddMenuEntry( FLevelEditorCommands::Get().EnableActorSnap );
 		TSharedRef<SWidget> SnapItem = 
 		SNew(SHorizontalBox)
 	          +SHorizontalBox::Slot()
@@ -2036,35 +2070,35 @@ TSharedRef< SWidget > FLevelEditorToolBar::GenerateQuickSettingsMenu( TSharedRef
 	          ]
 	          +SHorizontalBox::Slot()
 	          .FillWidth(0.1f);
-		MenuBuilder.AddWidget(SnapItem, LOCTEXT("ActorSnapLabel","Distance"));
+		Section.AddEntry(FToolMenuEntry::InitWidget("Snap", SnapItem, LOCTEXT("ActorSnapLabel", "Distance")));
 
-		MenuBuilder.AddMenuEntry( FLevelEditorCommands::Get().ToggleSocketSnapping );
-		MenuBuilder.AddMenuEntry( FLevelEditorCommands::Get().EnableVertexSnap );
+		Section.AddMenuEntry( FLevelEditorCommands::Get().ToggleSocketSnapping );
+		Section.AddMenuEntry( FLevelEditorCommands::Get().EnableVertexSnap );
 	}
-	MenuBuilder.EndSection();
 
-	MenuBuilder.BeginSection("LevelEditorViewport", LOCTEXT("ViewportHeading", "Viewport") );
 	{
-		MenuBuilder.AddMenuEntry( FLevelEditorCommands::Get().ToggleHideViewportUI );
+		FToolMenuSection& Section = Menu->AddSection("LevelEditorViewport", LOCTEXT("ViewportHeading", "Viewport") );
+		Section.AddMenuEntry( FLevelEditorCommands::Get().ToggleHideViewportUI );
 
-		MenuBuilder.AddSubMenu( LOCTEXT("PreviewMenu", "Previewing"), LOCTEXT("PreviewMenuTooltip","Game Preview Settings"), FNewMenuDelegate::CreateStatic( &MakePreviewSettingsMenu ) );
+		Section.AddSubMenu( "Preview", LOCTEXT("PreviewMenu", "Previewing"), LOCTEXT("PreviewMenuTooltip","Game Preview Settings"), FNewToolMenuDelegate::CreateStatic( &MakePreviewSettingsMenu ) );
 	}
-	MenuBuilder.EndSection();
 
 #undef LOCTEXT_NAMESPACE
-
-	return MenuBuilder.MakeWidget();
 }
 
 
-TSharedRef< SWidget > FLevelEditorToolBar::GenerateSourceControlMenu(TSharedRef<FUICommandList> InCommandList)
+TSharedRef< SWidget > FLevelEditorToolBar::GenerateSourceControlMenu(TSharedRef<FUICommandList> InCommandList, TWeakPtr< SLevelEditor > InLevelEditor)
 {
 	// Get all menu extenders for this context menu from the level editor module
 	FLevelEditorModule& LevelEditorModule = FModuleManager::GetModuleChecked<FLevelEditorModule>(TEXT("LevelEditor"));
 	TSharedPtr<FExtender> MenuExtender = LevelEditorModule.AssembleExtenders(InCommandList, LevelEditorModule.GetAllLevelEditorToolbarSourceControlMenuExtenders());
 
-	FToolMenuContext Context(InCommandList, MenuExtender);
-	return UToolMenus::Get()->GenerateWidget("LevelEditor.LevelEditorToolBar.SourceControl", Context);
+	FToolMenuContext MenuContext(InCommandList, MenuExtender);
+	ULevelEditorMenuContext* LevelEditorMenuContext = NewObject<ULevelEditorMenuContext>();
+	LevelEditorMenuContext->LevelEditor = InLevelEditor;
+	MenuContext.AddObject(LevelEditorMenuContext);
+
+	return UToolMenus::Get()->GenerateWidget("LevelEditor.LevelEditorToolBar.SourceControl", MenuContext);
 }
 
 void FLevelEditorToolBar::RegisterSourceControlMenu()
@@ -2122,35 +2156,53 @@ TSharedRef< SWidget > FLevelEditorToolBar::GenerateOpenBlueprintMenuContent( TSh
 {
 #define LOCTEXT_NAMESPACE "LevelToolBarViewMenu"
 
+	FLevelEditorModule& LevelEditorModule = FModuleManager::LoadModuleChecked<FLevelEditorModule>("LevelEditor");
+	TSharedPtr<FExtender> MenuExtender = FExtender::Combine(LevelEditorModule.GetAllLevelEditorToolbarBlueprintsMenuExtenders());
+	FToolMenuContext MenuContext(InCommandList, MenuExtender);
+	ULevelEditorMenuContext* LevelEditorMenuContext = NewObject<ULevelEditorMenuContext>();
+	LevelEditorMenuContext->LevelEditor = InLevelEditor;
+	MenuContext.AddObject(LevelEditorMenuContext);
+
+	return UToolMenus::Get()->GenerateWidget("LevelEditor.LevelEditorToolBar.OpenBlueprint", MenuContext);
+}
+
+void FLevelEditorToolBar::RegisterOpenBlueprintMenu()
+{
+#define LOCTEXT_NAMESPACE "LevelToolBarViewMenu"
+	UToolMenu* Menu = UToolMenus::Get()->RegisterMenu("LevelEditor.LevelEditorToolBar.OpenBlueprint");
+
 	struct FBlueprintMenus
 	{
 		/** Generates a sub-level Blueprints sub-menu */
-		static void MakeSubLevelsMenu(FMenuBuilder& InMenuBuilder, TWeakPtr< SLevelEditor > InLvlEditor)
+		static void MakeSubLevelsMenu(UToolMenu* InMenu)
 		{
-			FSlateIcon EditBP(FEditorStyle::Get().GetStyleSetName(), TEXT("LevelEditor.OpenLevelBlueprint"));
-
-			InMenuBuilder.BeginSection(NAME_None, LOCTEXT("SubLevelsHeading", "Sub-Level Blueprints"));
+			ULevelEditorMenuContext* Context = InMenu->FindContext<ULevelEditorMenuContext>();
+			if (Context && Context->LevelEditor.IsValid())
 			{
-				UWorld* World = InLvlEditor.Pin()->GetWorld();
-				// Sort the levels alphabetically 
-				TArray<ULevel*> SortedLevels = World->GetLevels();
-				Algo::Sort(SortedLevels, LevelEditorActionHelpers::FLevelSortByName());
+				FSlateIcon EditBP(FEditorStyle::Get().GetStyleSetName(), TEXT("LevelEditor.OpenLevelBlueprint"));
 
-				for (ULevel* const Level : SortedLevels)
 				{
-					if (Level != NULL && Level->GetOutermost() != NULL && !Level->IsPersistentLevel())
+					FToolMenuSection& Section = InMenu->AddSection("SubLevels", LOCTEXT("SubLevelsHeading", "Sub-Level Blueprints"));
+					UWorld* World = Context->LevelEditor.Pin()->GetWorld();
+					// Sort the levels alphabetically 
+					TArray<ULevel*> SortedLevels = World->GetLevels();
+					Algo::Sort(SortedLevels, LevelEditorActionHelpers::FLevelSortByName());
+
+					for (ULevel* const Level : SortedLevels)
 					{
-						FUIAction UIAction
-						(
-							FExecuteAction::CreateStatic(&FLevelEditorToolBar::OnOpenSubLevelBlueprint, Level)
-						);
+						if (Level != NULL && Level->GetOutermost() != NULL && !Level->IsPersistentLevel())
+						{
+							FUIAction UIAction
+							(
+								FExecuteAction::CreateStatic(&FLevelEditorToolBar::OnOpenSubLevelBlueprint, Level)
+							);
 
-						FText DisplayName = FText::Format(LOCTEXT("SubLevelBlueprintItem", "Edit {0}"), FText::FromString(FPaths::GetCleanFilename(Level->GetOutermost()->GetName())));
-						InMenuBuilder.AddMenuEntry(DisplayName, FText::GetEmpty(), EditBP, UIAction);
+							FText DisplayName = FText::Format(LOCTEXT("SubLevelBlueprintItem", "Edit {0}"), FText::FromString(FPaths::GetCleanFilename(Level->GetOutermost()->GetName())));
+							Section.AddMenuEntry(NAME_None, DisplayName, FText::GetEmpty(), EditBP, UIAction);
+						}
 					}
 				}
 			}
-			InMenuBuilder.EndSection();
 		}
 
 		/** Handle BP being selected from popup picker */
@@ -2165,7 +2217,7 @@ TSharedRef< SWidget > FLevelEditorToolBar::GenerateOpenBlueprintMenuContent( TSh
 
 
 		/** Generates 'open blueprint' sub-menu */
-		static void MakeOpenBPClassMenu(FMenuBuilder& InMenuBuilder)
+		static void MakeOpenBPClassMenu(UToolMenu* InMenu)
 		{
 			FContentBrowserModule& ContentBrowserModule = FModuleManager::Get().LoadModuleChecked<FContentBrowserModule>(TEXT("ContentBrowser"));
 
@@ -2188,80 +2240,77 @@ TSharedRef< SWidget > FLevelEditorToolBar::GenerateOpenBlueprintMenuContent( TSh
 				];
 		
 
-			InMenuBuilder.BeginSection(NAME_None, LOCTEXT("BrowseHeader", "Browse"));
 			{
-				InMenuBuilder.AddWidget(Widget, FText::GetEmpty());
+				FToolMenuSection& Section = InMenu->AddSection("Browse", LOCTEXT("BrowseHeader", "Browse"));
+				Section.AddEntry(FToolMenuEntry::InitWidget("PickClassWidget", Widget, FText::GetEmpty()));
 			}
-			InMenuBuilder.EndSection();
 		}
 	};
 
-	FLevelEditorModule& LevelEditorModule = FModuleManager::LoadModuleChecked<FLevelEditorModule>("LevelEditor");
-	TSharedPtr<FExtender> Extender = FExtender::Combine(LevelEditorModule.GetAllLevelEditorToolbarBlueprintsMenuExtenders());
-
-	const bool bShouldCloseWindowAfterMenuSelection = true;
-	FMenuBuilder MenuBuilder( bShouldCloseWindowAfterMenuSelection, InCommandList, Extender);
-
-	MenuBuilder.BeginSection("BlueprintClass", LOCTEXT("BlueprintClass", "Blueprint Class"));
 	{
+		FToolMenuSection& Section = Menu->AddSection("BlueprintClass", LOCTEXT("BlueprintClass", "Blueprint Class"));
+
 		// Create a blank BP
-		MenuBuilder.AddMenuEntry(FLevelEditorCommands::Get().CreateBlankBlueprintClass);
+		Section.AddMenuEntry(FLevelEditorCommands::Get().CreateBlankBlueprintClass);
 
 		// Convert selection to BP
-		MenuBuilder.AddMenuEntry(FLevelEditorCommands::Get().ConvertSelectionToBlueprintViaHarvest);
-		MenuBuilder.AddMenuEntry(FLevelEditorCommands::Get().ConvertSelectionToBlueprintViaSubclass);
+		Section.AddMenuEntry(FLevelEditorCommands::Get().ConvertSelectionToBlueprintViaHarvest);
+		Section.AddMenuEntry(FLevelEditorCommands::Get().ConvertSelectionToBlueprintViaSubclass);
 
 		// Open an existing Blueprint Class...
 		FSlateIcon OpenBPIcon(FEditorStyle::GetStyleSetName(), "LevelEditor.OpenClassBlueprint");
-		MenuBuilder.AddSubMenu(
+		Section.AddSubMenu(
+			"OpenBlueprintClass",
 			LOCTEXT("OpenBlueprintClassSubMenu", "Open Blueprint Class..."),
 			LOCTEXT("OpenBlueprintClassSubMenu_ToolTip", "Open an existing Blueprint Class in this project"),
-			FNewMenuDelegate::CreateStatic(&FBlueprintMenus::MakeOpenBPClassMenu),
+			FNewToolMenuDelegate::CreateStatic(&FBlueprintMenus::MakeOpenBPClassMenu),
 			false,
 			OpenBPIcon);
 	}
-	MenuBuilder.EndSection();
 
-	MenuBuilder.BeginSection("LevelScriptBlueprints", LOCTEXT("LevelScriptBlueprints", "Level Blueprints"));
 	{
-		MenuBuilder.AddMenuEntry( FLevelEditorCommands::Get().OpenLevelBlueprint );
+		FToolMenuSection& Section = Menu->AddSection("LevelScriptBlueprints", LOCTEXT("LevelScriptBlueprints", "Level Blueprints"));
+		Section.AddMenuEntry( FLevelEditorCommands::Get().OpenLevelBlueprint );
 
-		// If there are any sub-levels, display the sub-menu. A single level means there is only the persistent level
-		UWorld* World = InLevelEditor.Pin()->GetWorld();
-		if(World->GetNumLevels() > 1)
+		Section.AddDynamicEntry("SubLevels", FNewToolMenuSectionDelegate::CreateLambda([](FToolMenuSection& InSection)
 		{
-			MenuBuilder.AddSubMenu(
-				LOCTEXT( "SubLevelsSubMenu", "Sub-Levels" ),
-				LOCTEXT( "SubLevelsSubMenu_ToolTip", "Shows available sub-level Blueprints that can be edited." ),
-				FNewMenuDelegate::CreateStatic( &FBlueprintMenus::MakeSubLevelsMenu, InLevelEditor ), 
-				FUIAction(), NAME_None, EUserInterfaceActionType::Button, false, FSlateIcon(FEditorStyle::Get().GetStyleSetName(), TEXT("LevelEditor.OpenLevelBlueprint")) );
-		}
+			ULevelEditorMenuContext* Context = InSection.FindContext<ULevelEditorMenuContext>();
+			if (Context && Context->LevelEditor.IsValid())
+			{
+				// If there are any sub-levels, display the sub-menu. A single level means there is only the persistent level
+				UWorld* World = Context->LevelEditor.Pin()->GetWorld();
+				if (World->GetNumLevels() > 1)
+				{
+					InSection.AddSubMenu(
+						"SubLevels",
+						LOCTEXT("SubLevelsSubMenu", "Sub-Levels"),
+						LOCTEXT("SubLevelsSubMenu_ToolTip", "Shows available sub-level Blueprints that can be edited."),
+						FNewToolMenuDelegate::CreateStatic(&FBlueprintMenus::MakeSubLevelsMenu),
+						FUIAction(), EUserInterfaceActionType::Button, false, FSlateIcon(FEditorStyle::Get().GetStyleSetName(), TEXT("LevelEditor.OpenLevelBlueprint")));
+				}
+			}
+		}));
 	}
-	MenuBuilder.EndSection();
 
-	MenuBuilder.BeginSection("ProjectSettingsClasses", LOCTEXT("ProjectSettingsClasses", "Project Settings"));
 	{
-		// If source control is enabled, queue up a query to the status of the config file so it is (hopefully) ready before we get to the sub-menu
-		if(ISourceControlModule::Get().IsEnabled())
-		{
-			FString ConfigFilePath = FPaths::ConvertRelativePathToFull(FString::Printf(TEXT("%sDefaultEngine.ini"), *FPaths::SourceConfigDir()));
-
-			// note: calling QueueStatusUpdate often does not spam status updates as an internal timer prevents this
-			ISourceControlModule::Get().QueueStatusUpdate(ConfigFilePath);
-		}
-		LevelEditorActionHelpers::CreateGameModeSubMenu(MenuBuilder, InCommandList, InLevelEditor, true);
+		FToolMenuSection& Section = Menu->AddSection("ProjectSettingsClasses", LOCTEXT("ProjectSettingsClasses", "Project Settings"));
+		LevelEditorActionHelpers::CreateGameModeSubMenu(Section, "ProjectSettingsClasses", true);
 	}
-	MenuBuilder.EndSection();
 
-	MenuBuilder.BeginSection("WorldSettingsClasses", LOCTEXT("WorldSettingsClasses", "World Override"));
 	{
-		LevelEditorActionHelpers::CreateGameModeSubMenu(MenuBuilder, InCommandList, InLevelEditor, false);
+		FToolMenuSection& Section = Menu->AddSection("WorldSettingsClasses", LOCTEXT("WorldSettingsClasses", "World Override"));
+		LevelEditorActionHelpers::CreateGameModeSubMenu(Section, "WorldSettingsClasses", false);
 	}
-	MenuBuilder.EndSection();
 
-#undef LOCTEXT_NAMESPACE
+	// If source control is enabled, queue up a query to the status of the config file so it is (hopefully) ready before we get to the sub-menu
+	if(ISourceControlModule::Get().IsEnabled())
+	{
+		FString ConfigFilePath = FPaths::ConvertRelativePathToFull(FString::Printf(TEXT("%sDefaultEngine.ini"), *FPaths::SourceConfigDir()));
 
-	return MenuBuilder.MakeWidget();
+		// note: calling QueueStatusUpdate often does not spam status updates as an internal timer prevents this
+		ISourceControlModule::Get().QueueStatusUpdate(ConfigFilePath);
+	}
+#undef LOCTEXT_NAMESPACE
 }
 
 void FLevelEditorToolBar::OnOpenSubLevelBlueprint( ULevel* InLevel )
@@ -2278,20 +2327,24 @@ void FLevelEditorToolBar::OnOpenSubLevelBlueprint( ULevel* InLevel )
 	}
 }
 
-TSharedRef< SWidget > FLevelEditorToolBar::GenerateCinematicsMenuContent( TSharedRef<FUICommandList> InCommandList, TWeakPtr<SLevelEditor> LevelEditorWeakPtr )
+TSharedRef< SWidget > FLevelEditorToolBar::GenerateCinematicsMenuContent(TSharedRef<FUICommandList> InCommandList, TWeakPtr<SLevelEditor> InLevelEditor)
 {
 	FLevelEditorModule& LevelEditorModule = FModuleManager::LoadModuleChecked<FLevelEditorModule>("LevelEditor");
-	FToolMenuContext Context(InCommandList, FExtender::Combine(LevelEditorModule.GetAllLevelEditorToolbarCinematicsMenuExtenders()));
+
+	FToolMenuContext MenuContext(InCommandList, FExtender::Combine(LevelEditorModule.GetAllLevelEditorToolbarCinematicsMenuExtenders()));
 	ULevelEditorMenuContext* LevelEditorMenuContext = NewObject<ULevelEditorMenuContext>();
-	LevelEditorMenuContext->LevelEditor = LevelEditorWeakPtr;
-	Context.AddObject(LevelEditorMenuContext);
-	return UToolMenus::Get()->GenerateWidget("LevelEditor.LevelEditorToolBar.Cinematics", Context);
+	LevelEditorMenuContext->LevelEditor = InLevelEditor;
+	MenuContext.AddObject(LevelEditorMenuContext);
+
+	return UToolMenus::Get()->GenerateWidget("LevelEditor.LevelEditorToolBar.Cinematics", MenuContext);
 }
 
 void FLevelEditorToolBar::RegisterCinematicsMenu()
 {
 #define LOCTEXT_NAMESPACE "LevelToolBarCinematicsMenu"
 
+#pragma message("TODO: CLEAN UP MERGED CODE")
+#if 0
 	UToolMenu* Menu = UToolMenus::Get()->RegisterMenu("LevelEditor.LevelEditorToolBar.Cinematics");
 	Menu->bShouldCloseWindowAfterMenuSelection = true;
 
@@ -2301,7 +2354,7 @@ void FLevelEditorToolBar::RegisterCinematicsMenu()
 	FToolMenuSection& ExistingCinematicSection = Menu->AddSection("LevelEditorExistingCinematic", LOCTEXT("CinematicMenuCombo_ExistingHeading", "Edit Existing Cinematic"));
 	ExistingCinematicSection.AddDynamicEntry("LevelEditorExistingCinematic", FNewToolMenuSectionDelegate::CreateLambda([](FToolMenuSection& InSection)
 	{
-		ULevelEditorMenuContext* FoundContext = InSection.Context.Find<ULevelEditorMenuContext>();
+		ULevelEditorMenuContext* FoundContext = InSection.Context.FindContext<ULevelEditorMenuContext>();
 		if (!FoundContext)
 		{
 			return;
@@ -2354,6 +2407,7 @@ void FLevelEditorToolBar::RegisterCinematicsMenu()
 
 		InSection.AddEntry(FToolMenuEntry::InitWidget("LevelEditorExistingCinematic", MiniSceneOutliner, FText::GetEmpty(), true));
 	}));
+#endif
 
 #undef LOCTEXT_NAMESPACE
 }
diff --git a/Engine/Source/Editor/LevelEditor/Private/LevelEditorToolBar.h b/Engine/Source/Editor/LevelEditor/Private/LevelEditorToolBar.h
--- a/Engine/Source/Editor/LevelEditor/Private/LevelEditorToolBar.h
+++ b/Engine/Source/Editor/LevelEditor/Private/LevelEditorToolBar.h
@@ -34,51 +34,35 @@ protected:
 	 *
 	 * @return	Menu content widget
 	 */
-	static TSharedRef< SWidget > GenerateBuildMenuContent( TSharedRef<FUICommandList> InCommandList );
-
-#if WITH_LIVE_CODING
-	/**
-	 * Generates menu content for the compile combo button drop down menu
-	 *
-	 * @return	Menu content widget
-	 */
-	static TSharedRef< SWidget > GenerateCompileMenuContent( TSharedRef<FUICommandList> InCommandList );
-#endif
-
-	/**
-	 * Generates menu content for the create actor combo button drop down menu
-	 *
-	 * @return	Menu content widget
-	 */
-	static TSharedRef< SWidget > GenerateCreateContent( TSharedRef<FUICommandList> InCommandList );
+	static TSharedRef< SWidget > GenerateBuildMenuContent(TSharedRef<FUICommandList> InCommandList, TWeakPtr<SLevelEditor> InLevelEditor);
 
 	/**
 	 * Generates menu content for the quick settings combo button drop down menu
 	 *
 	 * @return	Menu content widget
 	 */
-	static TSharedRef< SWidget > GenerateQuickSettingsMenu( TSharedRef<FUICommandList> InCommandList );
+	static TSharedRef< SWidget > GenerateQuickSettingsMenu(TSharedRef<FUICommandList> InCommandList, TWeakPtr<SLevelEditor> InLevelEditor);
 
 	/**
 	 * Generates menu content for the source control combo button drop down menu
 	 *
 	 * @return	Menu content widget
 	 */
-	static TSharedRef< SWidget > GenerateSourceControlMenu(TSharedRef<FUICommandList> InCommandList);
+	static TSharedRef< SWidget > GenerateSourceControlMenu(TSharedRef<FUICommandList> InCommandList, TWeakPtr<SLevelEditor> InLevelEditor);
 
 	/**
 	 * Generates menu content for the compile combo button drop down menu
 	 *
 	 * @return	Menu content widget
 	 */
-	static TSharedRef< SWidget > GenerateOpenBlueprintMenuContent( TSharedRef<FUICommandList> InCommandList, TWeakPtr< SLevelEditor > InLevelEditor );
+	static TSharedRef< SWidget > GenerateOpenBlueprintMenuContent(TSharedRef<FUICommandList> InCommandList, TWeakPtr<SLevelEditor> InLevelEditor);
 
 	/**
 	 * Generates menu content for the Cinematics combo button drop down menu
 	 *
 	 * @return	Menu content widget
 	 */
-	static TSharedRef< SWidget > GenerateCinematicsMenuContent( TSharedRef<FUICommandList> InCommandList, TWeakPtr<SLevelEditor> LevelEditorWeakPtr );
+	static TSharedRef< SWidget > GenerateCinematicsMenuContent(TSharedRef<FUICommandList> InCommandList, TWeakPtr<SLevelEditor> InLevelEditor);
 
 	/**
 	 * Delegate for actor selection within the Cinematics popup menu's SceneOutliner.
@@ -97,4 +81,15 @@ private:
 
 	static void RegisterSourceControlMenu();
 	static void RegisterCinematicsMenu();
+	static void RegisterBuildMenu();
+
+#if WITH_LIVE_CODING
+	/**
+	 * Generates menu content for the compile combo button drop down menu
+	 */
+	static void RegisterCompileMenu();
+#endif
+
+	static void RegisterQuickSettingsMenu();
+	static void RegisterOpenBlueprintMenu();
 };
diff --git a/Engine/Source/Editor/MaterialEditor/Private/MaterialInstanceEditor.cpp b/Engine/Source/Editor/MaterialEditor/Private/MaterialInstanceEditor.cpp
--- a/Engine/Source/Editor/MaterialEditor/Private/MaterialInstanceEditor.cpp
+++ b/Engine/Source/Editor/MaterialEditor/Private/MaterialInstanceEditor.cpp
@@ -935,7 +935,6 @@ void FMaterialInstanceEditor::GenerateInheritanceMenu(UToolMenu* Menu)
 	RebuildInheritanceList();
 	Menu->bShouldCloseWindowAfterMenuSelection = true;
 	Menu->bSearchable = true;
-	Menu->SetMaxHeight(500);
 	const FName ParentName = TEXT("ParentChain");
 	FToolMenuSection& Section = Menu->AddSection(ParentName, LOCTEXT("ParentChain", "Parent Chain"));
 	if (bIsFunctionPreviewMaterial)
diff --git a/Engine/Source/Editor/Sequencer/Sequencer.Build.cs b/Engine/Source/Editor/Sequencer/Sequencer.Build.cs
--- a/Engine/Source/Editor/Sequencer/Sequencer.Build.cs
+++ b/Engine/Source/Editor/Sequencer/Sequencer.Build.cs
@@ -48,7 +48,7 @@ public class Sequencer : ModuleRules
 				"GraphEditor",
                 "ViewportInteraction",
                 "SerializedRecorderInterface",
-				"ToolMenus"
+				"ToolMenus",
 			}
 		);
 
diff --git a/Engine/Source/Editor/SkeletalMeshEditor/Private/SkeletalMeshEditor.cpp b/Engine/Source/Editor/SkeletalMeshEditor/Private/SkeletalMeshEditor.cpp
--- a/Engine/Source/Editor/SkeletalMeshEditor/Private/SkeletalMeshEditor.cpp
+++ b/Engine/Source/Editor/SkeletalMeshEditor/Private/SkeletalMeshEditor.cpp
@@ -31,6 +31,8 @@
 #include "Widgets/Layout/SBox.h"
 #include "Widgets/Layout/SBorder.h"
 #include "Framework/Application/SlateApplication.h"
+#include "ToolMenus.h"
+#include "SkeletalMeshToolMenuContext.h"
 #include "EditorViewportClient.h"
 #include "Settings/EditorExperimentalSettings.h"
 #include "Algo/Transform.h"
@@ -220,13 +222,14 @@ void FSkeletalMeshEditor::InitSkeletalMeshEditor(const EToolkitMode::Type Mode,
 	ISkeletonEditorModule& SkeletonEditorModule = FModuleManager::GetModuleChecked<ISkeletonEditorModule>("SkeletonEditor");
 	SkeletonTree = SkeletonEditorModule.CreateSkeletonTree(PersonaToolkit->GetSkeleton(), SkeletonTreeArgs);
 
+	BindCommands();
+	RegisterToolbar();
+
 	const bool bCreateDefaultStandaloneMenu = true;
 	const bool bCreateDefaultToolbar = true;
 	const TSharedRef<FTabManager::FLayout> DummyLayout = FTabManager::NewLayout("NullLayout")->AddArea(FTabManager::NewPrimaryArea());
 	FAssetEditorToolkit::InitAssetEditor(Mode, InitToolkitHost, SkeletalMeshEditorAppIdentifier, DummyLayout, bCreateDefaultStandaloneMenu, bCreateDefaultToolbar, InSkeletalMesh);
 
-	BindCommands();
-
 	AddApplicationMode(
 		SkeletalMeshEditorModes::SkeletalMeshEditorMode,
 		MakeShareable(new FSkeletalMeshEditorMode(SharedThis(this), SkeletonTree.ToSharedRef())));
@@ -286,133 +289,118 @@ void FSkeletalMeshEditor::BindCommands()
 		FExecuteAction::CreateRaw(&GetPersonaToolkit()->GetPreviewScene().Get(), &IPersonaPreviewScene::TogglePlayback));
 }
 
-void FSkeletalMeshEditor::ExtendToolbar()
+TSharedPtr<FSkeletalMeshEditor> FSkeletalMeshEditor::GetSkeletalMeshEditor(const FToolMenuContext& InMenuContext)
 {
-	// If the ToolbarExtender is valid, remove it before rebuilding it
-	if (ToolbarExtender.IsValid())
+	if (USkeletalMeshToolMenuContext* Context = InMenuContext.FindContext<USkeletalMeshToolMenuContext>())
 	{
-		RemoveToolbarExtender(ToolbarExtender);
-		ToolbarExtender.Reset();
+		if (Context->SkeletalMeshEditor.IsValid())
+		{
+			return StaticCastSharedPtr<FSkeletalMeshEditor>(Context->SkeletalMeshEditor.Pin());
+		}
 	}
 
-	ToolbarExtender = MakeShareable(new FExtender);
+	return TSharedPtr<FSkeletalMeshEditor>();
+}
 
-	auto ConstructReimportContextMenu = [this]()
+void FSkeletalMeshEditor::RegisterReimportContextMenu(const FName InBaseMenuName)
+{
+	static auto ReimportMeshWithNewFileAction = [](const FToolMenuContext& InMenuContext, int32 SourceFileIndex)
 	{
-		bool bShowSubMenu = SkeletalMesh != nullptr && SkeletalMesh->AssetImportData != nullptr && SkeletalMesh->AssetImportData->GetSourceFileCount() > 1;
-		FMenuBuilder MenuBuilder(true, nullptr);
-		
-		if (!bShowSubMenu)
+		TSharedPtr<FSkeletalMeshEditor> SkeletalMeshEditor = GetSkeletalMeshEditor(InMenuContext);
+		if (SkeletalMeshEditor.IsValid())
 		{
-			//Reimport
-			MenuBuilder.AddMenuEntry(FSkeletalMeshEditorCommands::Get().ReimportMesh->GetLabel(),
-				FSkeletalMeshEditorCommands::Get().ReimportMesh->GetDescription(),
-				FSlateIcon(),
-				FUIAction(FExecuteAction::CreateSP(this, &FSkeletalMeshEditor::HandleReimportMesh, 0)));
-			
-			MenuBuilder.AddMenuEntry(FSkeletalMeshEditorCommands::Get().ReimportMeshWithNewFile->GetLabel(),
-				FSkeletalMeshEditorCommands::Get().ReimportMeshWithNewFile->GetDescription(),
-				FSlateIcon(),
-				FUIAction(FExecuteAction::CreateSP(this, &FSkeletalMeshEditor::HandleReimportMeshWithNewFile, 0)));
-
-			//Reimport ALL
-			MenuBuilder.AddMenuEntry(FSkeletalMeshEditorCommands::Get().ReimportAllMesh->GetLabel(),
-				FSkeletalMeshEditorCommands::Get().ReimportAllMesh->GetDescription(),
-				FSlateIcon(),
-				FUIAction(FExecuteAction::CreateSP(this, &FSkeletalMeshEditor::HandleReimportAllMesh, 0)));
-
-			MenuBuilder.AddMenuEntry(FSkeletalMeshEditorCommands::Get().ReimportAllMeshWithNewFile->GetLabel(),
-				FSkeletalMeshEditorCommands::Get().ReimportAllMeshWithNewFile->GetDescription(),
-				FSlateIcon(),
-				FUIAction(FExecuteAction::CreateSP(this, &FSkeletalMeshEditor::HandleReimportAllMeshWithNewFile, 0)));
-
-			FText ReimportMultiSources = LOCTEXT("ReimportMultiSources", "Reimport Content");
-			FText ReimportMultiSourcesTooltip = LOCTEXT("ReimportMultiSourcesTooltip", "Reimport Geometry or Skinning Weights content, this will create multi import source file.");
-
-			auto CreateMultiContentSubMenu = [this](FMenuBuilder& SubMenuBuilder)
+			if (USkeletalMesh* FoundSkeletalMesh = SkeletalMeshEditor->SkeletalMesh)
 			{
-				auto UpdateImportDataContentType = [this](int32 SourceFileIndex)
+				UFbxSkeletalMeshImportData* SkeletalMeshImportData = Cast<UFbxSkeletalMeshImportData>(FoundSkeletalMesh->AssetImportData);
+				if (SkeletalMeshImportData)
 				{
-					UFbxSkeletalMeshImportData* SkeletalMeshImportData = Cast<UFbxSkeletalMeshImportData>(SkeletalMesh->AssetImportData);
-					if (SkeletalMeshImportData)
+					SkeletalMeshImportData->ImportContentType = SourceFileIndex == 0 ? EFBXImportContentType::FBXICT_All : SourceFileIndex == 1 ? EFBXImportContentType::FBXICT_Geometry : EFBXImportContentType::FBXICT_SkinningWeights;
+					SkeletalMeshEditor->HandleReimportMeshWithNewFile(SourceFileIndex);
+				}
+			}
+		}
+	};
+
+	static auto CreateMultiContentSubMenu = [](UToolMenu* InMenu)
+	{
+		FToolMenuSection& Section = InMenu->AddSection("Reimport");
+
+		TSharedPtr<FSkeletalMeshEditor> SkeletalMeshEditor = GetSkeletalMeshEditor(InMenu->Context);
+		if (SkeletalMeshEditor.IsValid())
+		{
+			Section.AddMenuEntry(
+				"ReimportGeometryContentLabel",
+				LOCTEXT("ReimportGeometryContentLabel", "Geometry"),
+				LOCTEXT("ReimportGeometryContentLabelTooltipTooltip", "Reimport Geometry Only"),
+				FSlateIcon(FEditorStyle::GetStyleSetName(), "ContentBrowser.AssetActions.ReimportAsset"),
+				FToolMenuExecuteAction::CreateLambda(ReimportMeshWithNewFileAction, 1)
+			);
+
+			Section.AddMenuEntry(
+				"ReimportSkinningAndWeightsContentLabel",
+				LOCTEXT("ReimportSkinningAndWeightsContentLabel", "Skinning And Weights"),
+				LOCTEXT("ReimportSkinningAndWeightsContentLabelTooltipTooltip", "Reimport Skinning And Weights Only"),
+				FSlateIcon(FEditorStyle::GetStyleSetName(), "ContentBrowser.AssetActions.ReimportAsset"),
+				FToolMenuExecuteAction::CreateLambda(ReimportMeshWithNewFileAction, 2)
+			);
+		}
+	};
+
+	static auto ReimportAction = [](const FToolMenuContext& InMenuContext, const int32 SourceFileIndex, const bool bReimportAll, const bool bWithNewFile)
+	{
+		TSharedPtr<FSkeletalMeshEditor> SkeletalMeshEditor = GetSkeletalMeshEditor(InMenuContext);
+		if (SkeletalMeshEditor.IsValid())
+		{
+			if (USkeletalMesh* FoundSkeletalMesh = SkeletalMeshEditor->SkeletalMesh)
+			{
+				UFbxSkeletalMeshImportData* SkeletalMeshImportData = FoundSkeletalMesh ? Cast<UFbxSkeletalMeshImportData>(FoundSkeletalMesh->AssetImportData) : nullptr;
+				if (SkeletalMeshImportData)
+				{
+					SkeletalMeshImportData->ImportContentType = SourceFileIndex == 0 ? EFBXImportContentType::FBXICT_All : SourceFileIndex == 1 ? EFBXImportContentType::FBXICT_Geometry : EFBXImportContentType::FBXICT_SkinningWeights;
+					if (bReimportAll)
 					{
-						SkeletalMeshImportData->ImportContentType = SourceFileIndex == 0 ? EFBXImportContentType::FBXICT_All : SourceFileIndex == 1 ? EFBXImportContentType::FBXICT_Geometry : EFBXImportContentType::FBXICT_SkinningWeights;
-						HandleReimportMeshWithNewFile(SourceFileIndex);
+						if (bWithNewFile)
+						{
+							SkeletalMeshEditor->HandleReimportAllMeshWithNewFile(SourceFileIndex);
+						}
+						else
+						{
+							SkeletalMeshEditor->HandleReimportAllMesh(SourceFileIndex);
+						}
 					}
-				};
-
-				FText ReimportGeometryContentLabel = LOCTEXT("ReimportGeometryContentLabel", "Geometry");
-				FText ReimportGeometryContentLabelTooltip = LOCTEXT("ReimportGeometryContentLabelTooltipTooltip", "Reimport Geometry Only");
-				SubMenuBuilder.AddMenuEntry(
-					ReimportGeometryContentLabel,
-					ReimportGeometryContentLabelTooltip,
-					FSlateIcon(FEditorStyle::GetStyleSetName(), "ContentBrowser.AssetActions.ReimportAsset"),
-					FUIAction(
-						FExecuteAction::CreateLambda(UpdateImportDataContentType, 1),
-						FCanExecuteAction()
-					)
-				);
-				FText ReimportSkinningAndWeightsContentLabel = LOCTEXT("ReimportSkinningAndWeightsContentLabel", "Skinning And Weights");
-				FText ReimportSkinningAndWeightsContentLabelTooltip = LOCTEXT("ReimportSkinningAndWeightsContentLabelTooltipTooltip", "Reimport Skinning And Weights Only");
-				SubMenuBuilder.AddMenuEntry(
-					ReimportSkinningAndWeightsContentLabel,
-					ReimportSkinningAndWeightsContentLabelTooltip,
-					FSlateIcon(FEditorStyle::GetStyleSetName(), "ContentBrowser.AssetActions.ReimportAsset"),
-					FUIAction(
-						FExecuteAction::CreateLambda(UpdateImportDataContentType, 2),
-						FCanExecuteAction()
-					)
-				);
-			};
-
-			MenuBuilder.AddSubMenu(
-				ReimportMultiSources,
-				ReimportMultiSourcesTooltip,
-				FNewMenuDelegate::CreateLambda(CreateMultiContentSubMenu));
+					else
+					{
+						if (bWithNewFile)
+						{
+							SkeletalMeshEditor->HandleReimportMeshWithNewFile(SourceFileIndex);
+						}
+						else
+						{
+							SkeletalMeshEditor->HandleReimportMesh(SourceFileIndex);
+						}
+					}
+				}
+			}
 		}
-		else
+	};
+
+	static auto CreateReimportSubMenu = [](UToolMenu* InMenu, bool bReimportAll, bool bWithNewFile)
+	{
+		TSharedPtr<FSkeletalMeshEditor> SkeletalMeshEditor = GetSkeletalMeshEditor(InMenu->Context);
+		if (SkeletalMeshEditor.IsValid())
 		{
-			auto CreateSubMenu = [this](FMenuBuilder& SubMenuBuilder, bool bReimportAll, bool bWithNewFile)
+			USkeletalMesh* InSkeletalMesh = SkeletalMeshEditor->SkeletalMesh;
+			if (InSkeletalMesh && InSkeletalMesh->AssetImportData)
 			{
-				//Get the data, we cannot use the closure since the lambda will be call when the function scope will be gone
+				//Get the data
 				TArray<FString> SourceFilePaths;
-				SkeletalMesh->AssetImportData->ExtractFilenames(SourceFilePaths);
+				InSkeletalMesh->AssetImportData->ExtractFilenames(SourceFilePaths);
 				TArray<FString> SourceFileLabels;
-				SkeletalMesh->AssetImportData->ExtractDisplayLabels(SourceFileLabels);
+				InSkeletalMesh->AssetImportData->ExtractDisplayLabels(SourceFileLabels);
 
 				if (SourceFileLabels.Num() > 0 && SourceFileLabels.Num() == SourceFilePaths.Num())
 				{
-					auto UpdateImportDataContentType = [this](int32 SourceFileIndex, bool bReimportAll, bool bWithNewFile)
-					{
-						UFbxSkeletalMeshImportData* SkeletalMeshImportData = Cast<UFbxSkeletalMeshImportData>(SkeletalMesh->AssetImportData);
-						if (SkeletalMeshImportData)
-						{
-							SkeletalMeshImportData->ImportContentType = SourceFileIndex == 0 ? EFBXImportContentType::FBXICT_All : SourceFileIndex == 1 ? EFBXImportContentType::FBXICT_Geometry : EFBXImportContentType::FBXICT_SkinningWeights;
-							if (bReimportAll)
-							{
-								if (bWithNewFile)
-								{
-									HandleReimportAllMeshWithNewFile(SourceFileIndex);
-								}
-								else
-								{
-									HandleReimportAllMesh(SourceFileIndex);
-								}
-							}
-							else
-							{
-								if (bWithNewFile)
-								{
-									HandleReimportMeshWithNewFile(SourceFileIndex);
-								}
-								else
-								{
-									HandleReimportMesh(SourceFileIndex);
-								}
-							}
-						}
-					};
-
+					FToolMenuSection& Section = InMenu->AddSection("Reimport");
 					for (int32 SourceFileIndex = 0; SourceFileIndex < SourceFileLabels.Num(); ++SourceFileIndex)
 					{
 						FText ReimportLabel = FText::Format(LOCTEXT("ReimportNoLabel", "SourceFile {0}"), SourceFileIndex);
@@ -422,77 +410,112 @@ void FSkeletalMeshEditor::ExtendToolbar()
 							ReimportLabel = FText::Format(LOCTEXT("ReimportLabel", "{0}"), FText::FromString(SourceFileLabels[SourceFileIndex]));
 							ReimportLabelTooltip = FText::Format(LOCTEXT("ReimportLabelTooltip", "Reimport {0} File: {1}"), FText::FromString(SourceFileLabels[SourceFileIndex]), FText::FromString(SourceFilePaths[SourceFileIndex]));
 						}
-						SubMenuBuilder.AddMenuEntry(
+
+						Section.AddMenuEntry(
+							NAME_None,
 							ReimportLabel,
 							ReimportLabelTooltip,
 							FSlateIcon(FEditorStyle::GetStyleSetName(), "ContentBrowser.AssetActions.ReimportAsset"),
-							FUIAction(
-								FExecuteAction::CreateLambda(UpdateImportDataContentType, SourceFileIndex, bReimportAll, bWithNewFile),
-								FCanExecuteAction()
-							)
+							FToolMenuExecuteAction::CreateLambda(ReimportAction, SourceFileIndex, bReimportAll, bWithNewFile)
 						);
 					}
 				}
-			};
-
-			//Create 4 submenu: Reimport, ReimportWithNewFile, ReimportAll and ReimportAllWithNewFile
-			MenuBuilder.AddSubMenu(
-				FSkeletalMeshEditorCommands::Get().ReimportMesh->GetLabel(),
-				FSkeletalMeshEditorCommands::Get().ReimportMesh->GetDescription(),
-				FNewMenuDelegate::CreateLambda(CreateSubMenu, false, false));
-			
-			MenuBuilder.AddSubMenu(
-				FSkeletalMeshEditorCommands::Get().ReimportMeshWithNewFile->GetLabel(),
-				FSkeletalMeshEditorCommands::Get().ReimportMeshWithNewFile->GetDescription(),
-				FNewMenuDelegate::CreateLambda(CreateSubMenu, false, true));
-
-			MenuBuilder.AddSubMenu(
-				FSkeletalMeshEditorCommands::Get().ReimportAllMesh->GetLabel(),
-				FSkeletalMeshEditorCommands::Get().ReimportAllMesh->GetDescription(),
-				FNewMenuDelegate::CreateLambda(CreateSubMenu, true, false));
-
-			MenuBuilder.AddSubMenu(
-				FSkeletalMeshEditorCommands::Get().ReimportAllMeshWithNewFile->GetLabel(),
-				FSkeletalMeshEditorCommands::Get().ReimportAllMeshWithNewFile->GetDescription(),
-				FNewMenuDelegate::CreateLambda(CreateSubMenu, true, true));
+			}
 		}
-
-		return MenuBuilder.MakeWidget();
 	};
 
-	// extend extra menu/toolbars
-	auto FillToolbar = [this, ConstructReimportContextMenu](FToolBarBuilder& ToolbarBuilder)
+	if (!UToolMenus::Get()->IsMenuRegistered(UToolMenus::JoinMenuPaths(InBaseMenuName, "ReimportContextMenu")))
 	{
-		FPersonaModule& PersonaModule = FModuleManager::LoadModuleChecked<FPersonaModule>("Persona");
-		FPersonaModule::FCommonToolbarExtensionArgs Args;
-		Args.bPreviewMesh = false;
-		PersonaModule.AddCommonToolbarExtensions(ToolbarBuilder, PersonaToolkit.ToSharedRef(), Args);
-
-		ToolbarBuilder.BeginSection("Mesh");
+		UToolMenu* ToolMenu = UToolMenus::Get()->RegisterMenu(UToolMenus::JoinMenuPaths(InBaseMenuName, "ReimportContextMenu"));
+		ToolMenu->AddDynamicSection("Section", FNewToolMenuDelegate::CreateLambda([](UToolMenu* InMenu)
 		{
-			ToolbarBuilder.AddToolBarButton(FSkeletalMeshEditorCommands::Get().ReimportMesh);
-			ToolbarBuilder.AddComboButton(
-				FUIAction(),
-				FOnGetContent::CreateLambda(ConstructReimportContextMenu),
-				TAttribute<FText>(),
-				TAttribute<FText>()
-			);
-		}
-		ToolbarBuilder.EndSection();
+			TSharedPtr<FSkeletalMeshEditor> SkeletalMeshEditor = GetSkeletalMeshEditor(InMenu->Context);
+			if (SkeletalMeshEditor.IsValid())
+			{
+				USkeletalMesh* InSkeletalMesh = SkeletalMeshEditor->SkeletalMesh;
+				bool bShowSubMenu = InSkeletalMesh != nullptr && InSkeletalMesh->AssetImportData != nullptr && InSkeletalMesh->AssetImportData->GetSourceFileCount() > 1;
 
-		ToolbarBuilder.BeginSection("Selection");
-		{
-			ToolbarBuilder.AddToolBarButton(FSkeletalMeshEditorCommands::Get().MeshSectionSelection);
-		}
-		ToolbarBuilder.EndSection();
-	};
+				FToolMenuSection& Section = InMenu->AddSection("Section");
+				if (!bShowSubMenu)
+				{
+					//Reimport
+					Section.AddMenuEntry(
+						FSkeletalMeshEditorCommands::Get().ReimportMesh->GetCommandName(),
+						FSkeletalMeshEditorCommands::Get().ReimportMesh->GetLabel(),
+						FSkeletalMeshEditorCommands::Get().ReimportMesh->GetDescription(),
+						FSlateIcon(),
+						FUIAction(FExecuteAction::CreateSP(SkeletalMeshEditor.ToSharedRef(), &FSkeletalMeshEditor::HandleReimportMesh, 0)));
+
+					Section.AddMenuEntry(
+						FSkeletalMeshEditorCommands::Get().ReimportMeshWithNewFile->GetCommandName(),
+						FSkeletalMeshEditorCommands::Get().ReimportMeshWithNewFile->GetLabel(),
+						FSkeletalMeshEditorCommands::Get().ReimportMeshWithNewFile->GetDescription(),
+						FSlateIcon(),
+						FUIAction(FExecuteAction::CreateSP(SkeletalMeshEditor.ToSharedRef(), &FSkeletalMeshEditor::HandleReimportMeshWithNewFile, 0)));
+
+					//Reimport ALL
+					Section.AddMenuEntry(
+						FSkeletalMeshEditorCommands::Get().ReimportAllMesh->GetCommandName(),
+						FSkeletalMeshEditorCommands::Get().ReimportAllMesh->GetLabel(),
+						FSkeletalMeshEditorCommands::Get().ReimportAllMesh->GetDescription(),
+						FSlateIcon(),
+						FUIAction(FExecuteAction::CreateSP(SkeletalMeshEditor.ToSharedRef(), &FSkeletalMeshEditor::HandleReimportAllMesh, 0)));
+
+					Section.AddMenuEntry(
+						FSkeletalMeshEditorCommands::Get().ReimportAllMeshWithNewFile->GetCommandName(),
+						FSkeletalMeshEditorCommands::Get().ReimportAllMeshWithNewFile->GetLabel(),
+						FSkeletalMeshEditorCommands::Get().ReimportAllMeshWithNewFile->GetDescription(),
+						FSlateIcon(),
+						FUIAction(FExecuteAction::CreateSP(SkeletalMeshEditor.ToSharedRef(), &FSkeletalMeshEditor::HandleReimportAllMeshWithNewFile, 0)));
+
+					Section.AddSubMenu(
+						"ReimportMultiSources",
+						LOCTEXT("ReimportMultiSources", "Reimport Content"),
+						LOCTEXT("ReimportMultiSourcesTooltip", "Reimport Geometry or Skinning Weights content, this will create multi import source file."),
+						FNewToolMenuDelegate::CreateLambda(CreateMultiContentSubMenu));
+				}
+				else
+				{
+					//Create 4 submenu: Reimport, ReimportWithNewFile, ReimportAll and ReimportAllWithNewFile
+					Section.AddSubMenu(
+						FSkeletalMeshEditorCommands::Get().ReimportMesh->GetCommandName(),
+						FSkeletalMeshEditorCommands::Get().ReimportMesh->GetLabel(),
+						FSkeletalMeshEditorCommands::Get().ReimportMesh->GetDescription(),
+						FNewToolMenuDelegate::CreateLambda(CreateReimportSubMenu, false, false));
+			
+					Section.AddSubMenu(
+						FSkeletalMeshEditorCommands::Get().ReimportMeshWithNewFile->GetCommandName(),
+						FSkeletalMeshEditorCommands::Get().ReimportMeshWithNewFile->GetLabel(),
+						FSkeletalMeshEditorCommands::Get().ReimportMeshWithNewFile->GetDescription(),
+						FNewToolMenuDelegate::CreateLambda(CreateReimportSubMenu, false, true));
+
+					Section.AddSubMenu(
+						FSkeletalMeshEditorCommands::Get().ReimportAllMesh->GetCommandName(),
+						FSkeletalMeshEditorCommands::Get().ReimportAllMesh->GetLabel(),
+						FSkeletalMeshEditorCommands::Get().ReimportAllMesh->GetDescription(),
+						FNewToolMenuDelegate::CreateLambda(CreateReimportSubMenu, true, false));
+
+					Section.AddSubMenu(
+						FSkeletalMeshEditorCommands::Get().ReimportAllMeshWithNewFile->GetCommandName(),
+						FSkeletalMeshEditorCommands::Get().ReimportAllMeshWithNewFile->GetLabel(),
+						FSkeletalMeshEditorCommands::Get().ReimportAllMeshWithNewFile->GetDescription(),
+						FNewToolMenuDelegate::CreateLambda(CreateReimportSubMenu, true, true));
+				}
+			}
+		}));
+	}
+}
 
-	ToolbarExtender->AddToolBarExtension(
-		"Asset",
-		EExtensionHook::After,
-		GetToolkitCommands(),
-		FToolBarExtensionDelegate::CreateLambda(FillToolbar)
-		);
+void FSkeletalMeshEditor::ExtendToolbar()
+{
+	// If the ToolbarExtender is valid, remove it before rebuilding it
+	if (ToolbarExtender.IsValid())
+	{
+		RemoveToolbarExtender(ToolbarExtender);
+		ToolbarExtender.Reset();
+	}
+
+	ToolbarExtender = MakeShareable(new FExtender);
 
 	AddToolbarExtender(ToolbarExtender);
 
@@ -515,6 +538,7 @@ void FSkeletalMeshEditor::ExtendToolbar()
 		GetToolkitCommands(),
 		FToolBarExtensionDelegate::CreateLambda([this](FToolBarBuilder& ParentToolbarBuilder)
 	{
+		// Second toolbar on right side
 		FPersonaModule& PersonaModule = FModuleManager::LoadModuleChecked<FPersonaModule>("Persona");
 		TSharedRef<class IAssetFamily> AssetFamily = PersonaModule.CreatePersonaAssetFamily(SkeletalMesh);
 		AddToolbarWidget(PersonaModule.CreateAssetFamilyShortcutWidget(SharedThis(this), AssetFamily));
@@ -522,6 +546,53 @@ void FSkeletalMeshEditor::ExtendToolbar()
 	));
 }
 
+void FSkeletalMeshEditor::RegisterToolbar()
+{
+	static const FName MenuName = "AssetEditor.SkeletalMeshEditor.ToolBar";
+	RegisterReimportContextMenu(MenuName);
+
+	if (!UToolMenus::Get()->IsMenuRegistered(MenuName))
+	{
+		UToolMenu* ToolMenu = UToolMenus::Get()->RegisterMenu(MenuName, "AssetEditor.DefaultToolBar", EMultiBoxType::ToolBar);
+
+		const FToolMenuInsert SectionInsertLocation("Asset", EToolMenuInsertType::After);
+
+		{
+			ToolMenu->AddDynamicSection("Persona", FNewToolBarDelegateLegacy::CreateLambda([](FToolBarBuilder& ToolbarBuilder, UToolMenu* InMenu)
+			{
+				TSharedPtr<FSkeletalMeshEditor> SkeletalMeshEditor = GetSkeletalMeshEditor(InMenu->Context);
+				if (SkeletalMeshEditor.IsValid() && SkeletalMeshEditor->PersonaToolkit.IsValid())
+				{
+					FPersonaModule& PersonaModule = FModuleManager::LoadModuleChecked<FPersonaModule>("Persona");
+					FPersonaModule::FCommonToolbarExtensionArgs Args;
+					Args.bPreviewMesh = false;
+					PersonaModule.AddCommonToolbarExtensions(ToolbarBuilder, SkeletalMeshEditor->PersonaToolkit.ToSharedRef(), Args);
+				}
+			}), SectionInsertLocation);
+		}
+
+		{
+			FToolMenuSection& Section = ToolMenu->AddSection("Mesh", FText(), SectionInsertLocation);
+			Section.AddEntry(FToolMenuEntry::InitToolBarButton(FSkeletalMeshEditorCommands::Get().ReimportMesh));
+			Section.AddEntry(FToolMenuEntry::InitComboButton("ReimportContextMenu", FUIAction(), FNewToolMenuDelegate()));
+		}
+
+		{
+			FToolMenuSection& Section = ToolMenu->AddSection("SkeletalMesh", FText(), FToolMenuInsert("Mesh", EToolMenuInsertType::After));
+			Section.AddEntry(FToolMenuEntry::InitToolBarButton(FSkeletalMeshEditorCommands::Get().MeshSectionSelection));
+		}
+	}
+}
+
+void FSkeletalMeshEditor::InitToolMenuContext(FToolMenuContext& MenuContext)
+{
+	FAssetEditorToolkit::InitToolMenuContext(MenuContext);
+
+	USkeletalMeshToolMenuContext* Context = NewObject<USkeletalMeshToolMenuContext>();
+	Context->SkeletalMeshEditor = SharedThis(this);
+	MenuContext.AddObject(Context);
+}
+
 void FSkeletalMeshEditor::FillMeshClickMenu(FMenuBuilder& MenuBuilder, HActor* HitProxy, const FViewportClick& Click)
 {
 	UDebugSkelMeshComponent* MeshComp = GetPersonaToolkit()->GetPreviewMeshComponent();
diff --git a/Engine/Source/Editor/SkeletalMeshEditor/Private/SkeletalMeshEditor.h b/Engine/Source/Editor/SkeletalMeshEditor/Private/SkeletalMeshEditor.h
--- a/Engine/Source/Editor/SkeletalMeshEditor/Private/SkeletalMeshEditor.h
+++ b/Engine/Source/Editor/SkeletalMeshEditor/Private/SkeletalMeshEditor.h
@@ -21,6 +21,8 @@ class ISkeletonTreeItem;
 struct HActor;
 struct FViewportClick;
 struct FSkeletalMeshClothBuildParams;
+struct FToolMenuContext;
+class UToolMenu;
 
 namespace SkeletalMeshEditorModes
 {
@@ -61,6 +63,7 @@ public:
 	virtual FText GetBaseToolkitName() const override;
 	virtual FString GetWorldCentricTabPrefix() const override;
 	virtual FLinearColor GetWorldCentricTabColorScale() const override;
+	virtual void InitToolMenuContext(FToolMenuContext& MenuContext) override;
 
 	//~ Begin FAssetEditorToolkit Interface.
 	virtual bool OnRequestClose() override;
@@ -137,9 +140,14 @@ private:
 	void OnRemoveSectionFromLodAndBelowMenuItemClicked(int32 LodIndex, int32 SectionIndex);
 	//////////////////////////////////////////////////////////////////////////
 
+	void RegisterReimportContextMenu(const FName InBaseMenuName);
+
+	static TSharedPtr<FSkeletalMeshEditor> GetSkeletalMeshEditor(const FToolMenuContext& InMenuContext);
+
 private:
 	void ExtendMenu();
 
+	void RegisterToolbar();
 	void ExtendToolbar();
 
 	void BindCommands();
diff --git a/Engine/Source/Editor/SkeletalMeshEditor/Private/SkeletalMeshEditorMode.cpp b/Engine/Source/Editor/SkeletalMeshEditor/Private/SkeletalMeshEditorMode.cpp
--- a/Engine/Source/Editor/SkeletalMeshEditor/Private/SkeletalMeshEditorMode.cpp
+++ b/Engine/Source/Editor/SkeletalMeshEditor/Private/SkeletalMeshEditorMode.cpp
@@ -85,6 +85,8 @@ FSkeletalMeshEditorMode::FSkeletalMeshEditorMode(TSharedRef<FWorkflowCentricAppl
 	LayoutExtender = MakeShared<FLayoutExtender>();
 	PersonaModule.OnRegisterLayoutExtensions().Broadcast(*LayoutExtender.Get());
 	TabLayout->ProcessExtensions(*LayoutExtender.Get());
+
+	SkeletalMeshEditor->RegisterModeToolbarIfUnregistered(GetModeName());
 }
 
 void FSkeletalMeshEditorMode::RegisterTabFactories(TSharedPtr<FTabManager> InTabManager)
diff --git a/Engine/Source/Editor/SkeletalMeshEditor/SkeletalMeshEditor.Build.cs b/Engine/Source/Editor/SkeletalMeshEditor/SkeletalMeshEditor.Build.cs
--- a/Engine/Source/Editor/SkeletalMeshEditor/SkeletalMeshEditor.Build.cs
+++ b/Engine/Source/Editor/SkeletalMeshEditor/SkeletalMeshEditor.Build.cs
@@ -34,7 +34,7 @@ public class SkeletalMeshEditor : ModuleRules
                 "ClothingSystemRuntimeCommon",
                 "ClothingSystemEditorInterface",
 				"ClothingSystemRuntimeInterface",
-				"SkeletalMeshUtilitiesCommon",
+				"ToolMenus",
             }
 		);
 
diff --git a/Engine/Source/Editor/UMGEditor/Private/BlueprintModes/WidgetDesignerApplicationMode.cpp b/Engine/Source/Editor/UMGEditor/Private/BlueprintModes/WidgetDesignerApplicationMode.cpp
--- a/Engine/Source/Editor/UMGEditor/Private/BlueprintModes/WidgetDesignerApplicationMode.cpp
+++ b/Engine/Source/Editor/UMGEditor/Private/BlueprintModes/WidgetDesignerApplicationMode.cpp
@@ -8,7 +8,7 @@
 
 #include "WidgetBlueprintEditorToolbar.h"
 #include "UMGEditorModule.h"
-
+#include "ToolMenus.h"
 
 #include "TabFactory/PaletteTabSummoner.h"
 #include "TabFactory/HierarchyTabSummoner.h"
@@ -123,9 +123,13 @@ FWidgetDesignerApplicationMode::FWidgetDesignerApplicationMode(TSharedPtr<FWidge
 	ToolbarExtender = UMGEditorModule.GetToolBarExtensibilityManager()->GetAllExtenders();
 	
 	InWidgetEditor->GetWidgetToolbarBuilder()->AddWidgetBlueprintEditorModesToolbar(ToolbarExtender);
-	InWidgetEditor->GetWidgetToolbarBuilder()->AddWidgetReflector(ToolbarExtender);
-	InWidgetEditor->GetToolbarBuilder()->AddCompileToolbar(ToolbarExtender);
-	InWidgetEditor->GetToolbarBuilder()->AddDebuggingToolbar(ToolbarExtender);
+
+	if (UToolMenu* Toolbar = InWidgetEditor->RegisterModeToolbarIfUnregistered(GetModeName()))
+	{
+		InWidgetEditor->GetWidgetToolbarBuilder()->AddWidgetReflector(Toolbar);
+		InWidgetEditor->GetToolbarBuilder()->AddCompileToolbar(Toolbar);
+		InWidgetEditor->GetToolbarBuilder()->AddDebuggingToolbar(Toolbar);
+	}
 }
 
 void FWidgetDesignerApplicationMode::RegisterTabFactories(TSharedPtr<FTabManager> InTabManager)
diff --git a/Engine/Source/Editor/UMGEditor/Private/BlueprintModes/WidgetGraphApplicationMode.cpp b/Engine/Source/Editor/UMGEditor/Private/BlueprintModes/WidgetGraphApplicationMode.cpp
--- a/Engine/Source/Editor/UMGEditor/Private/BlueprintModes/WidgetGraphApplicationMode.cpp
+++ b/Engine/Source/Editor/UMGEditor/Private/BlueprintModes/WidgetGraphApplicationMode.cpp
@@ -7,6 +7,7 @@
 
 #include "WidgetBlueprintEditorToolbar.h"
 #include "BlueprintModes/WidgetBlueprintApplicationModes.h"
+#include "ToolMenus.h"
 
 /////////////////////////////////////////////////////
 // FWidgetGraphApplicationMode
@@ -78,10 +79,14 @@ FWidgetGraphApplicationMode::FWidgetGraphApplicationMode(TSharedPtr<FWidgetBluep
 	//@TODO: Keep this in sync with AnimBlueprintMode.cpp
 	ToolbarExtender = MakeShareable(new FExtender);
 	InWidgetEditor->GetWidgetToolbarBuilder()->AddWidgetBlueprintEditorModesToolbar(ToolbarExtender);
-	InWidgetEditor->GetToolbarBuilder()->AddCompileToolbar(ToolbarExtender);
-	InWidgetEditor->GetToolbarBuilder()->AddScriptingToolbar(ToolbarExtender);
-	InWidgetEditor->GetToolbarBuilder()->AddBlueprintGlobalOptionsToolbar(ToolbarExtender);
-	InWidgetEditor->GetToolbarBuilder()->AddDebuggingToolbar(ToolbarExtender);
+
+	if (UToolMenu* Toolbar = InWidgetEditor->RegisterModeToolbarIfUnregistered(GetModeName()))
+	{
+		InWidgetEditor->GetToolbarBuilder()->AddCompileToolbar(Toolbar);
+		InWidgetEditor->GetToolbarBuilder()->AddScriptingToolbar(Toolbar);
+		InWidgetEditor->GetToolbarBuilder()->AddBlueprintGlobalOptionsToolbar(Toolbar);
+		InWidgetEditor->GetToolbarBuilder()->AddDebuggingToolbar(Toolbar);
+	}
 }
 
 void FWidgetGraphApplicationMode::RegisterTabFactories(TSharedPtr<FTabManager> InTabManager)
diff --git a/Engine/Source/Editor/UMGEditor/Private/WidgetBlueprintEditorToolbar.cpp b/Engine/Source/Editor/UMGEditor/Private/WidgetBlueprintEditorToolbar.cpp
--- a/Engine/Source/Editor/UMGEditor/Private/WidgetBlueprintEditorToolbar.cpp
+++ b/Engine/Source/Editor/UMGEditor/Private/WidgetBlueprintEditorToolbar.cpp
@@ -14,6 +14,8 @@
 #include "IDocumentation.h"
 #include "BlueprintEditor.h"
 #include "Framework/MultiBox/MultiBoxBuilder.h"
+#include "ToolMenus.h"
+#include "BlueprintEditorContext.h"
 
 #include "WidgetBlueprintEditor.h"
 #include "WorkflowOrientedApp/SModeWidget.h"
@@ -123,33 +125,22 @@ void FWidgetBlueprintEditorToolbar::FillWidgetBlueprintEditorModesToolbar(FToolB
 	}
 }
 
-void FWidgetBlueprintEditorToolbar::AddWidgetReflector(TSharedPtr<FExtender> Extender)
-{
-	TSharedPtr<FWidgetBlueprintEditor> BlueprintEditorPtr = WidgetEditor.Pin();
 
-	Extender->AddToolBarExtension(
-		"Asset",
-		EExtensionHook::After,
-		BlueprintEditorPtr->GetToolkitCommands(),
-		FToolBarExtensionDelegate::CreateSP(this, &FWidgetBlueprintEditorToolbar::FillWidgetReflectorToolbar));
-}
-
-void FWidgetBlueprintEditorToolbar::FillWidgetReflectorToolbar(FToolBarBuilder& ToolbarBuilder)
+void FWidgetBlueprintEditorToolbar::AddWidgetReflector(UToolMenu* InMenu)
 {
-	ToolbarBuilder.BeginSection("WidgetTools");
+	FToolMenuSection& Section = InMenu->AddSection("WidgetTools");
+	Section.InsertPosition = FToolMenuInsert("Asset", EToolMenuInsertType::After);
 
-	ToolbarBuilder.AddToolBarButton(
+	Section.AddEntry(FToolMenuEntry::InitToolBarButton(
+		"OpenWidgetReflector",
 		FUIAction(
 			FExecuteAction::CreateLambda([=] { FGlobalTabmanager::Get()->InvokeTab(FTabId("WidgetReflector")); }),
 			FCanExecuteAction()
 		)
-		, NAME_None
 		, LOCTEXT("OpenWidgetReflector", "Widget Reflector")
 		, LOCTEXT("OpenWidgetReflectorToolTip", "Opens the Widget Reflector, a handy tool for diagnosing problems with live widgets.")
 		, FSlateIcon(FCoreStyle::Get().GetStyleSetName(), "WidgetReflector.Icon")
-	);
-
-	ToolbarBuilder.EndSection();
+	));
 }
 
 #undef LOCTEXT_NAMESPACE
diff --git a/Engine/Source/Editor/UMGEditor/Private/WidgetBlueprintEditorToolbar.h b/Engine/Source/Editor/UMGEditor/Private/WidgetBlueprintEditorToolbar.h
--- a/Engine/Source/Editor/UMGEditor/Private/WidgetBlueprintEditorToolbar.h
+++ b/Engine/Source/Editor/UMGEditor/Private/WidgetBlueprintEditorToolbar.h
@@ -7,6 +7,7 @@
 class FExtender;
 class FToolBarBuilder;
 class FWidgetBlueprintEditor;
+class UToolMenu;
 
 /**
  * Handles all of the toolbar related construction for the widget blueprint editor.
@@ -23,14 +24,11 @@ public:
 	 */
 	void AddWidgetBlueprintEditorModesToolbar(TSharedPtr<FExtender> Extender);
 
-	void AddWidgetReflector(TSharedPtr<FExtender> Extender);
+	void AddWidgetReflector(UToolMenu* InMenu);
 
 public:
 	/**  */
 	void FillWidgetBlueprintEditorModesToolbar(FToolBarBuilder& ToolbarBuilder);
 
-	/**  */
-	void FillWidgetReflectorToolbar(FToolBarBuilder& ToolbarBuilder);
-
 	TWeakPtr<FWidgetBlueprintEditor> WidgetEditor;
 };
diff --git a/Engine/Source/Editor/UMGEditor/UMGEditor.Build.cs b/Engine/Source/Editor/UMGEditor/UMGEditor.Build.cs
--- a/Engine/Source/Editor/UMGEditor/UMGEditor.Build.cs
+++ b/Engine/Source/Editor/UMGEditor/UMGEditor.Build.cs
@@ -64,7 +64,8 @@ public class UMGEditor : ModuleRules
                 "TargetPlatform",
 				"TimeManagement",
 				"GameProjectGeneration",
-				"PropertyPath"
+				"PropertyPath",
+				"ToolMenus",
 			}
 			);
 	}
diff --git a/Engine/Source/Editor/UnrealEd/Private/Kismet2/DebuggerCommands.cpp b/Engine/Source/Editor/UnrealEd/Private/Kismet2/DebuggerCommands.cpp
--- a/Engine/Source/Editor/UnrealEd/Private/Kismet2/DebuggerCommands.cpp
+++ b/Engine/Source/Editor/UnrealEd/Private/Kismet2/DebuggerCommands.cpp
@@ -69,6 +69,8 @@
 #include "Framework/Notifications/NotificationManager.h"
 #include "Widgets/Notifications/SNotificationList.h"
 
+#include "ToolMenus.h"
+#include "SBlueprintEditorToolbar.h"
 
 #define LOCTEXT_NAMESPACE "DebuggerCommands"
 
@@ -612,86 +614,92 @@ void FPlayWorldCommands::AddPIEPreviewDeviceActions(const FPlayWorldCommands &Co
 	}
 }
 
-void FPlayWorldCommands::BuildToolbar( FToolBarBuilder& ToolbarBuilder, bool bIncludeLaunchButtonAndOptions )
+void FPlayWorldCommands::BuildToolbar(FToolMenuSection& InSection, bool bIncludeLaunchButtonAndOptions)
 {
 	// Play
-	ToolbarBuilder.AddToolBarButton( 
-		FPlayWorldCommands::Get().RepeatLastPlay, 
-		NAME_None, 
+	InSection.AddEntry(FToolMenuEntry::InitToolBarButton(
+		FPlayWorldCommands::Get().RepeatLastPlay,
 		LOCTEXT("RepeatLastPlay", "Play"),
 		TAttribute< FText >::Create( TAttribute< FText >::FGetter::CreateStatic( &FInternalPlayWorldCommandCallbacks::GetRepeatLastPlayToolTip ) ),
-		TAttribute< FSlateIcon >::Create( TAttribute< FSlateIcon >::FGetter::CreateStatic( &FInternalPlayWorldCommandCallbacks::GetRepeatLastPlayIcon ) ),
+		TAttribute< FSlateIcon >::Create(TAttribute< FSlateIcon >::FGetter::CreateStatic(&FInternalPlayWorldCommandCallbacks::GetRepeatLastPlayIcon)),
 		FName(TEXT("LevelToolbarPlay"))
-	);
+	));
 
 	// Play combo box
 	FUIAction SpecialPIEOptionsMenuAction;
-	SpecialPIEOptionsMenuAction.IsActionVisibleDelegate = FIsActionButtonVisible::CreateStatic( &FInternalPlayWorldCommandCallbacks::CanShowNonPlayWorldOnlyActions );
+	SpecialPIEOptionsMenuAction.IsActionVisibleDelegate = FIsActionButtonVisible::CreateStatic(&FInternalPlayWorldCommandCallbacks::CanShowNonPlayWorldOnlyActions);
 
-	ToolbarBuilder.AddComboButton(
+	InSection.AddEntry(FToolMenuEntry::InitComboButton(
+		"PlayCombo",
 		SpecialPIEOptionsMenuAction,
 		FOnGetContent::CreateStatic( &GeneratePlayMenuContent, GlobalPlayWorldActions.ToSharedRef() ),
 		LOCTEXT( "PlayCombo_Label", "Active Play Mode" ),
 		LOCTEXT( "PIEComboToolTip", "Change Play Mode and Play Settings" ),
 		FSlateIcon(FEditorStyle::GetStyleSetName(), "LevelEditor.RepeatLastPlay"),
 		true
-	);
+	));
 
-	if (bIncludeLaunchButtonAndOptions && GetDefault<UEditorStyleSettings>()->bShowLaunchMenus)
+	if (bIncludeLaunchButtonAndOptions)
 	{
-		// Launch
-		ToolbarBuilder.AddToolBarButton( 
-			FPlayWorldCommands::Get().RepeatLastLaunch, 
-			NAME_None, 
-			LOCTEXT("RepeatLastLaunch", "Launch"),
-			TAttribute< FText >::Create( TAttribute< FText >::FGetter::CreateStatic( &FInternalPlayWorldCommandCallbacks::GetRepeatLastLaunchToolTip ) ),
-			TAttribute< FSlateIcon >::Create( TAttribute< FSlateIcon >::FGetter::CreateStatic( &FInternalPlayWorldCommandCallbacks::GetRepeatLastLaunchIcon ) ),
-			FName(TEXT("RepeatLastLaunch"))
-		);
-
-		// Launch combo box
-		FUIAction LaunchMenuAction;
-		LaunchMenuAction.IsActionVisibleDelegate = FIsActionButtonVisible::CreateStatic( &FInternalPlayWorldCommandCallbacks::CanShowNonPlayWorldOnlyActions );
-
-		ToolbarBuilder.AddComboButton(
-			LaunchMenuAction,
-			FOnGetContent::CreateStatic( &GenerateLaunchMenuContent, GlobalPlayWorldActions.ToSharedRef() ),
-			LOCTEXT( "LaunchCombo_Label", "Launch Options" ),
-			LOCTEXT( "PODComboToolTip", "Options for launching on a device" ),
-			FSlateIcon(FEditorStyle::GetStyleSetName(), "LevelEditor.RepeatLastLaunch"),
-			true
-		);
+		InSection.AddDynamicEntry("LaunchButtons", FNewToolMenuSectionDelegate::CreateLambda([](FToolMenuSection& InDynamicSection)
+		{
+			if (GetDefault<UEditorStyleSettings>()->bShowLaunchMenus)
+			{
+				// Launch
+				InDynamicSection.AddEntry(FToolMenuEntry::InitToolBarButton(
+					FPlayWorldCommands::Get().RepeatLastLaunch,
+					LOCTEXT("RepeatLastLaunch", "Launch"),
+					TAttribute< FText >::Create(TAttribute< FText >::FGetter::CreateStatic(&FInternalPlayWorldCommandCallbacks::GetRepeatLastLaunchToolTip)),
+					TAttribute< FSlateIcon >::Create(TAttribute< FSlateIcon >::FGetter::CreateStatic(&FInternalPlayWorldCommandCallbacks::GetRepeatLastLaunchIcon)),
+					FName(TEXT("RepeatLastLaunch"))
+				));
+
+				// Launch combo box
+				FUIAction LaunchMenuAction;
+				LaunchMenuAction.IsActionVisibleDelegate = FIsActionButtonVisible::CreateStatic(&FInternalPlayWorldCommandCallbacks::CanShowNonPlayWorldOnlyActions);
+
+				InDynamicSection.AddEntry(FToolMenuEntry::InitComboButton(
+					"LaunchCombo",
+					LaunchMenuAction,
+					FOnGetContent::CreateStatic(&GenerateLaunchMenuContent, GlobalPlayWorldActions.ToSharedRef()),
+					LOCTEXT("LaunchCombo_Label", "Launch Options"),
+					LOCTEXT("PODComboToolTip", "Options for launching on a device"),
+					FSlateIcon(FEditorStyle::GetStyleSetName(), "LevelEditor.RepeatLastLaunch"),
+					true
+				));
+			}
+		}));
 	}
 
 	// Resume/pause toggle (only one will be visible, and only in PIE/SIE)
-	ToolbarBuilder.AddToolBarButton(FPlayWorldCommands::Get().ResumePlaySession, NAME_None, TAttribute<FText>(),
+	InSection.AddEntry(FToolMenuEntry::InitToolBarButton(FPlayWorldCommands::Get().ResumePlaySession, TAttribute<FText>(),
 		TAttribute<FText>::Create(TAttribute<FText>::FGetter::CreateStatic(&FInternalPlayWorldCommandCallbacks::GetResumePlaySessionToolTip)),
 		TAttribute<FSlateIcon>::Create(TAttribute<FSlateIcon>::FGetter::CreateStatic(&FInternalPlayWorldCommandCallbacks::GetResumePlaySessionImage)),
 		FName(TEXT("ResumePlaySession"))
-	);
+	));
 
-	ToolbarBuilder.AddToolBarButton(FPlayWorldCommands::Get().PausePlaySession, NAME_None, TAttribute<FText>(), TAttribute<FText>(), TAttribute<FSlateIcon>(), FName(TEXT("PausePlaySession")));
-	ToolbarBuilder.AddToolBarButton(FPlayWorldCommands::Get().SingleFrameAdvance, NAME_None, TAttribute<FText>(), TAttribute<FText>(), TAttribute<FSlateIcon>(), FName(TEXT("SingleFrameAdvance")));
+	InSection.AddEntry(FToolMenuEntry::InitToolBarButton(FPlayWorldCommands::Get().PausePlaySession, TAttribute<FText>(), TAttribute<FText>(), TAttribute<FSlateIcon>(), FName(TEXT("PausePlaySession"))));
+	InSection.AddEntry(FToolMenuEntry::InitToolBarButton(FPlayWorldCommands::Get().SingleFrameAdvance, TAttribute<FText>(), TAttribute<FText>(), TAttribute<FSlateIcon>(), FName(TEXT("SingleFrameAdvance"))));
 
 	// Stop
-	ToolbarBuilder.AddToolBarButton(FPlayWorldCommands::Get().StopPlaySession, NAME_None, TAttribute<FText>(), TAttribute<FText>(), TAttribute<FSlateIcon>(), FName(TEXT("StopPlaySession")));
+	InSection.AddEntry(FToolMenuEntry::InitToolBarButton(FPlayWorldCommands::Get().StopPlaySession, TAttribute<FText>(), TAttribute<FText>(), TAttribute<FSlateIcon>(), FName(TEXT("StopPlaySession"))));
 
 	// Late Join
-	ToolbarBuilder.AddToolBarButton(FPlayWorldCommands::Get().LateJoinSession, NAME_None, TAttribute<FText>(), TAttribute<FText>(), TAttribute<FSlateIcon>(), FName(TEXT("LateJoinSession")));
+	InSection.AddEntry(FToolMenuEntry::InitToolBarButton(FPlayWorldCommands::Get().LateJoinSession, TAttribute<FText>(), TAttribute<FText>(), TAttribute<FSlateIcon>(), FName(TEXT("LateJoinSession"))));
 
 	// Eject/possess toggle
-	ToolbarBuilder.AddToolBarButton(FPlayWorldCommands::Get().PossessEjectPlayer, NAME_None, 
+	InSection.AddEntry(FToolMenuEntry::InitToolBarButton(FPlayWorldCommands::Get().PossessEjectPlayer,
 		TAttribute<FText>::Create(TAttribute<FText>::FGetter::CreateStatic(&FInternalPlayWorldCommandCallbacks::GetPossessEjectLabel)),
 		TAttribute<FText>::Create(TAttribute<FText>::FGetter::CreateStatic(&FInternalPlayWorldCommandCallbacks::GetPossessEjectTooltip)),
 		TAttribute<FSlateIcon>::Create(TAttribute<FSlateIcon>::FGetter::CreateStatic(&FInternalPlayWorldCommandCallbacks::GetPossessEjectImage)),
 		FName(TEXT("PossessEjectPlayer"))
-	);
+	));
 
 	// Single-stepping only buttons
-	ToolbarBuilder.AddToolBarButton(FPlayWorldCommands::Get().ShowCurrentStatement, NAME_None, TAttribute<FText>(), TAttribute<FText>(), TAttribute<FSlateIcon>(), FName(TEXT("ShowCurrentStatement")));
-	ToolbarBuilder.AddToolBarButton(FPlayWorldCommands::Get().StepInto, NAME_None, TAttribute<FText>(), TAttribute<FText>(), TAttribute<FSlateIcon>(), FName(TEXT("StepInto")));
-	ToolbarBuilder.AddToolBarButton(FPlayWorldCommands::Get().StepOver, NAME_None, TAttribute<FText>(), TAttribute<FText>(), TAttribute<FSlateIcon>(), FName(TEXT("StepOver")));
-	ToolbarBuilder.AddToolBarButton(FPlayWorldCommands::Get().StepOut, NAME_None, TAttribute<FText>(), TAttribute<FText>(), TAttribute<FSlateIcon>(), FName(TEXT("StepOut")));
+	InSection.AddEntry(FToolMenuEntry::InitToolBarButton(FPlayWorldCommands::Get().ShowCurrentStatement, TAttribute<FText>(), TAttribute<FText>(), TAttribute<FSlateIcon>(), FName(TEXT("ShowCurrentStatement"))));
+	InSection.AddEntry(FToolMenuEntry::InitToolBarButton(FPlayWorldCommands::Get().StepInto, TAttribute<FText>(), TAttribute<FText>(), TAttribute<FSlateIcon>(), FName(TEXT("StepInto"))));
+	InSection.AddEntry(FToolMenuEntry::InitToolBarButton(FPlayWorldCommands::Get().StepOver, TAttribute<FText>(), TAttribute<FText>(), TAttribute<FSlateIcon>(), FName(TEXT("StepOver"))));
+	InSection.AddEntry(FToolMenuEntry::InitToolBarButton(FPlayWorldCommands::Get().StepOut, TAttribute<FText>(), TAttribute<FText>(), TAttribute<FSlateIcon>(), FName(TEXT("StepOut"))));
 }
 
 // function will enumerate available Android devices that can export their profile to a json file
@@ -827,11 +835,11 @@ TSharedRef< SWidget > FPlayWorldCommands::GeneratePlayMenuContent( TSharedRef<FU
 
 	struct FLocal
 	{
-		static void AddPlayModeMenuEntry( FMenuBuilder& MenuBuilder, EPlayModeType PlayMode )
+		static void AddPlayModeMenuEntry(FMenuBuilder& MenuBuilder, EPlayModeType PlayMode)
 		{
 			TSharedPtr<FUICommandInfo> PlayModeCommand;
 
-			switch(PlayMode)
+			switch (PlayMode)
 			{
 			case PlayMode_InEditorFloating:
 				PlayModeCommand = FPlayWorldCommands::Get().PlayInEditorFloating;
diff --git a/Engine/Source/Editor/UnrealEd/Private/Toolkits/AssetEditorToolkit.cpp b/Engine/Source/Editor/UnrealEd/Private/Toolkits/AssetEditorToolkit.cpp
--- a/Engine/Source/Editor/UnrealEd/Private/Toolkits/AssetEditorToolkit.cpp
+++ b/Engine/Source/Editor/UnrealEd/Private/Toolkits/AssetEditorToolkit.cpp
@@ -975,10 +975,18 @@ FName FAssetEditorToolkit::GetToolMenuAppName() const
 
 FName FAssetEditorToolkit::GetToolMenuToolbarName() const
 {
+	FName ParentName;
+	return GetToolMenuToolbarName(ParentName);
+}
+
+FName FAssetEditorToolkit::GetToolMenuToolbarName(FName& OutParentName) const
+{
+	static const FName DefaultToolbarName = "AssetEditor.DefaultToolBar";
+	OutParentName = DefaultToolbarName;
 	return *(TEXT("AssetEditor.") + GetToolMenuAppName().ToString() + TEXT(".ToolBar"));
 }
 
-void FAssetEditorToolkit::RegisterMenus()
+void FAssetEditorToolkit::RegisterDefaultToolBar()
 {
 	static const FName DefaultToolBarName("AssetEditor.DefaultToolBar");
 	UToolMenus* ToolMenus = UToolMenus::Get();
@@ -993,18 +1001,24 @@ void FAssetEditorToolkit::RegisterMenus()
 	}
 }
 
+void FAssetEditorToolkit::InitToolMenuContext(FToolMenuContext& MenuContext)
+{
+
+}
+
 void FAssetEditorToolkit::GenerateToolbar()
 {
 	TSharedPtr<FExtender> Extender = FExtender::Combine(ToolbarExtenders);
 
-	RegisterMenus();
+	RegisterDefaultToolBar();
 
-	const FName ToolBarName = GetToolMenuToolbarName();
+	FName ParentToolbarName;
+	const FName ToolBarName = GetToolMenuToolbarName(ParentToolbarName);
 	UToolMenus* ToolMenus = UToolMenus::Get();
 	UToolMenu* FoundMenu = ToolMenus->FindMenu(ToolBarName);
 	if (!FoundMenu || !FoundMenu->IsRegistered())
 	{
-		FoundMenu = ToolMenus->RegisterMenu(ToolBarName, "AssetEditor.DefaultToolBar", EMultiBoxType::ToolBar);
+		FoundMenu = ToolMenus->RegisterMenu(ToolBarName, ParentToolbarName, EMultiBoxType::ToolBar);
 	}
 
 	FToolMenuContext MenuContext(GetToolkitCommands(), Extender);
@@ -1013,6 +1027,8 @@ void FAssetEditorToolkit::GenerateToolbar()
 	ToolkitMenuContext->Toolkit = AsShared();
 	MenuContext.AddObject(ToolkitMenuContext);
 
+	InitToolMenuContext(MenuContext);
+
 	UToolMenu* GeneratedToolbar = ToolMenus->GenerateMenu(ToolBarName, MenuContext);
 	GeneratedToolbar->bToolBarIsFocusable = bIsToolbarFocusable;
 	GeneratedToolbar->bToolBarForceSmallIcons = bIsToolbarUsingSmallIcons;
diff --git a/Engine/Source/Editor/UnrealEd/Public/Kismet2/DebuggerCommands.h b/Engine/Source/Editor/UnrealEd/Public/Kismet2/DebuggerCommands.h
--- a/Engine/Source/Editor/UnrealEd/Public/Kismet2/DebuggerCommands.h
+++ b/Engine/Source/Editor/UnrealEd/Public/Kismet2/DebuggerCommands.h
@@ -11,6 +11,7 @@
 #include "Framework/Commands/UICommandList.h"
 
 class FToolBarBuilder;
+struct FToolMenuSection;
 
 /** This class acts as a generic widget that listens to and process global play world actions */
 class UNREALED_API SGlobalPlayWorldActions : public SCompoundWidget
@@ -52,7 +53,7 @@ public:
 	static void BindGlobalPlayWorldCommands();
 
 	/** Populates a toolbar with the menu commands for play-world control (pause/resume/stop/possess/eject/step/show current loc) */
-	UNREALED_API static void BuildToolbar( FToolBarBuilder& ToolBarBuilder, bool bIncludeLaunchButtonAndOptions = false );
+	UNREALED_API static void BuildToolbar( FToolMenuSection& InSection, bool bIncludeLaunchButtonAndOptions = false );
 
 	/**
 	* Return the active widget that processes play world actions for PIE
diff --git a/Engine/Source/Editor/UnrealEd/Public/SkeletalMeshToolMenuContext.h b/Engine/Source/Editor/UnrealEd/Public/SkeletalMeshToolMenuContext.h
new file mode 100644
--- /dev/null
+++ b/Engine/Source/Editor/UnrealEd/Public/SkeletalMeshToolMenuContext.h
@@ -0,0 +1,18 @@
+// Copyright 1998-2019 Epic Games, Inc. All Rights Reserved.
+
+#pragma once
+
+#include "CoreMinimal.h"
+
+#include "SkeletalMeshToolMenuContext.generated.h"
+
+class ISkeletalMeshEditor;
+
+UCLASS()
+class UNREALED_API USkeletalMeshToolMenuContext : public UObject
+{
+	GENERATED_BODY()
+public:
+
+	TWeakPtr<ISkeletalMeshEditor> SkeletalMeshEditor;
+};
diff --git a/Engine/Source/Editor/UnrealEd/Public/Toolkits/AssetEditorToolkit.h b/Engine/Source/Editor/UnrealEd/Public/Toolkits/AssetEditorToolkit.h
--- a/Engine/Source/Editor/UnrealEd/Public/Toolkits/AssetEditorToolkit.h
+++ b/Engine/Source/Editor/UnrealEd/Public/Toolkits/AssetEditorToolkit.h
@@ -18,6 +18,7 @@
 
 class FAssetEditorModeManager;
 class FMenuBuilder;
+struct FToolMenuContext;
 class SBorder;
 class SStandaloneAssetEditorToolkitHost;
 
@@ -178,7 +179,7 @@ public:
 	}
 
 	/** Registers default tool bar */
-	static void RegisterMenus();
+	static void RegisterDefaultToolBar();
 	
 	/** Makes a default asset editing toolbar */
 	void GenerateToolbar();
@@ -187,6 +188,8 @@ public:
 	void RegenerateMenusAndToolbars();
 
 	/** Get name used by tool menu */
+	virtual FName GetToolMenuToolbarName(FName& OutParentName) const;
+	virtual void InitToolMenuContext(FToolMenuContext& MenuContext);
 	FName GetToolMenuToolbarName() const;
 	FName GetToolMenuAppName() const;
 
