Fix build issues after tool menus merge
#rnx

diff --git a/Engine/Plugins/Media/MediaPlayerEditor/Source/MediaPlayerEditor/Private/Widgets/SMediaPlayerEditorMedia.cpp b/Engine/Plugins/Media/MediaPlayerEditor/Source/MediaPlayerEditor/Private/Widgets/SMediaPlayerEditorMedia.cpp
--- a/Engine/Plugins/Media/MediaPlayerEditor/Source/MediaPlayerEditor/Private/Widgets/SMediaPlayerEditorMedia.cpp
+++ b/Engine/Plugins/Media/MediaPlayerEditor/Source/MediaPlayerEditor/Private/Widgets/SMediaPlayerEditorMedia.cpp
@@ -143,7 +143,7 @@ void SMediaPlayerEditorMedia::RegisterMenus()
 				{
 					if (Context->SelectedAsset)
 					{
-						FAssetEditorManager::Get().OpenEditorForAsset(Context->SelectedAsset);
+						GEditor->GetEditorSubsystem<UAssetEditorSubsystem>()->OpenEditorForAsset(Context->SelectedAsset);
 					}
 				}
 			})
diff --git a/Engine/Source/Editor/LevelEditor/Private/LevelEditorCreateActorMenu.cpp b/Engine/Source/Editor/LevelEditor/Private/LevelEditorCreateActorMenu.cpp
--- a/Engine/Source/Editor/LevelEditor/Private/LevelEditorCreateActorMenu.cpp
+++ b/Engine/Source/Editor/LevelEditor/Private/LevelEditorCreateActorMenu.cpp
@@ -367,7 +367,6 @@ void LevelEditorCreateActorMenu::FillAddReplaceViewportContextMenuSections(UTool
 				NSLOCTEXT("LevelViewportContextMenu", "AddActorHeading", "Place Actor") , 
 				NSLOCTEXT("LevelViewportContextMenu", "AddActorMenu_ToolTip", "Templates for adding a new actor to the world"),
 				FNewToolMenuDelegate::CreateStatic(&LevelEditorCreateActorMenu::FillAddReplaceActorMenu, EActorCreateMode::Add));
-		}
 
 			if ( CanReplaceActors() )
 			{
diff --git a/Engine/Source/Editor/LevelEditor/Private/LevelEditorToolBar.cpp b/Engine/Source/Editor/LevelEditor/Private/LevelEditorToolBar.cpp
--- a/Engine/Source/Editor/LevelEditor/Private/LevelEditorToolBar.cpp
+++ b/Engine/Source/Editor/LevelEditor/Private/LevelEditorToolBar.cpp
@@ -1563,6 +1563,7 @@ void FLevelEditorToolBar::RegisterBuildMenu()
 			FLightingMenus::RegisterLightingQualityMenu(InBaseMenuName);
 			FLightingMenus::RegisterLightingInfoMenu(InBaseMenuName);
 		}
+
 	private:
 
 		/** Generates a lighting quality sub-menu */
@@ -2158,6 +2159,7 @@ TSharedRef< SWidget > FLevelEditorToolBar::GenerateOpenBlueprintMenuContent( TSh
 
 	FLevelEditorModule& LevelEditorModule = FModuleManager::LoadModuleChecked<FLevelEditorModule>("LevelEditor");
 	TSharedPtr<FExtender> MenuExtender = FExtender::Combine(LevelEditorModule.GetAllLevelEditorToolbarBlueprintsMenuExtenders());
+
 	FToolMenuContext MenuContext(InCommandList, MenuExtender);
 	ULevelEditorMenuContext* LevelEditorMenuContext = NewObject<ULevelEditorMenuContext>();
 	LevelEditorMenuContext->LevelEditor = InLevelEditor;
@@ -2343,8 +2345,6 @@ void FLevelEditorToolBar::RegisterCinematicsMenu()
 {
 #define LOCTEXT_NAMESPACE "LevelToolBarCinematicsMenu"
 
-#pragma message("TODO: CLEAN UP MERGED CODE")
-#if 0
 	UToolMenu* Menu = UToolMenus::Get()->RegisterMenu("LevelEditor.LevelEditorToolBar.Cinematics");
 	Menu->bShouldCloseWindowAfterMenuSelection = true;
 
@@ -2407,7 +2407,6 @@ void FLevelEditorToolBar::RegisterCinematicsMenu()
 
 		InSection.AddEntry(FToolMenuEntry::InitWidget("LevelEditorExistingCinematic", MiniSceneOutliner, FText::GetEmpty(), true));
 	}));
-#endif
 
 #undef LOCTEXT_NAMESPACE
 }
diff --git a/Engine/Source/Editor/SkeletalMeshEditor/Private/SkeletalMeshEditor.cpp b/Engine/Source/Editor/SkeletalMeshEditor/Private/SkeletalMeshEditor.cpp
--- a/Engine/Source/Editor/SkeletalMeshEditor/Private/SkeletalMeshEditor.cpp
+++ b/Engine/Source/Editor/SkeletalMeshEditor/Private/SkeletalMeshEditor.cpp
@@ -40,7 +40,6 @@
 #include "FbxMeshUtils.h"
 #include "LODUtilities.h"
 
-#include "LODUtilities.h"
 #include "ScopedTransaction.h"
 #include "ComponentReregisterContext.h"
 #include "EditorFramework/AssetImportData.h"
diff --git a/Engine/Source/Editor/SkeletalMeshEditor/SkeletalMeshEditor.Build.cs b/Engine/Source/Editor/SkeletalMeshEditor/SkeletalMeshEditor.Build.cs
--- a/Engine/Source/Editor/SkeletalMeshEditor/SkeletalMeshEditor.Build.cs
+++ b/Engine/Source/Editor/SkeletalMeshEditor/SkeletalMeshEditor.Build.cs
@@ -34,6 +34,7 @@ public class SkeletalMeshEditor : ModuleRules
                 "ClothingSystemRuntimeCommon",
                 "ClothingSystemEditorInterface",
 				"ClothingSystemRuntimeInterface",
+				"SkeletalMeshUtilitiesCommon",
 				"ToolMenus",
             }
 		);
