Fix ActorDetails not being updated when moving Actor across LevelGrid

#rnx

diff --git a/Engine/Source/Editor/LevelEditor/Private/SLevelEditor.cpp b/Engine/Source/Editor/LevelEditor/Private/SLevelEditor.cpp
--- a/Engine/Source/Editor/LevelEditor/Private/SLevelEditor.cpp
+++ b/Engine/Source/Editor/LevelEditor/Private/SLevelEditor.cpp
@@ -67,7 +67,7 @@ namespace LevelEditorConstants
 }
 
 SLevelEditor::SLevelEditor()
-	: World(NULL)
+	: World(NULL), bNeedsRefresh(false)
 {
 	const bool bAreRealTimeThumbnailsAllowed = false;
 	ThumbnailPool = MakeShareable(new FAssetThumbnailPool(LevelEditorConstants::ThumbnailPoolSize, bAreRealTimeThumbnailsAllowed));
@@ -178,6 +178,8 @@ void SLevelEditor::Construct( const SLevelEditor::FArguments& InArgs)
 	UEditorEngine* Editor = (UEditorEngine*)GEngine;
 	World->ChangeFeatureLevel(Editor->GetActiveFeatureLevelPreviewType());
 
+	LevelActorOuterChangedHandle = GEditor->OnLevelActorOuterChanged().AddSP(this, &SLevelEditor::OnLevelActorOuterChanged);
+
 	// Patch into the OnPreviewFeatureLevelChanged() delegate to swap out the current feature level with a user selection.
 	PreviewFeatureLevelChangedHandle = Editor->OnPreviewFeatureLevelChanged().AddLambda([this](ERHIFeatureLevel::Type NewFeatureLevel)
 		{
@@ -317,6 +319,7 @@ SLevelEditor::~SLevelEditor()
 
 	if (GEditor)
 	{
+		GEditor->OnLevelActorOuterChanged().Remove(LevelActorOuterChangedHandle);
 		GEditor->GetEditorWorldContext(true).RemoveRef(World);
 	}
 }
@@ -1581,15 +1584,20 @@ void SLevelEditor::OnActorSelectionChanged(const TArray<UObject*>& NewSelection,
 		TSharedPtr<SActorDetails> ActorDetails = It->Pin();
 		if( ActorDetails.IsValid() )
 		{
-			ActorDetails->SetObjects(NewSelection, bForceRefresh);
+			ActorDetails->SetObjects(NewSelection, bForceRefresh || bNeedsRefresh);
 		}
 		else
 		{
 			// remove stray entries here
 		}
 	}
+	bNeedsRefresh = false;
 }
 
+void SLevelEditor::OnLevelActorOuterChanged(AActor* InActor, UObject* InOldOuter)
+{
+	bNeedsRefresh = true;
+}
 
 void SLevelEditor::AddStandaloneLevelViewport( const TSharedRef<SLevelViewport>& LevelViewport )
 {
diff --git a/Engine/Source/Editor/LevelEditor/Private/SLevelEditor.h b/Engine/Source/Editor/LevelEditor/Private/SLevelEditor.h
--- a/Engine/Source/Editor/LevelEditor/Private/SLevelEditor.h
+++ b/Engine/Source/Editor/LevelEditor/Private/SLevelEditor.h
@@ -212,6 +212,9 @@ private:
 
 	/** Called when actors are selected or unselected */
 	void OnActorSelectionChanged(const TArray<UObject*>& NewSelection, bool bForceRefresh = false);
+
+	/** Called when an actor changes outer */
+	void OnLevelActorOuterChanged(AActor* InActor = nullptr, UObject* InOldOuter = nullptr);
 private:
 
 	// Tracking the active viewports in this level editor.
@@ -263,4 +266,10 @@ private:
 
 	/** Handle to the registered OnPreviewFeatureLevelChanged delegate. */
 	FDelegateHandle PreviewFeatureLevelChangedHandle;
+
+	/** Handle to the registered OnLevelActorOuterChanged delegate */
+	FDelegateHandle LevelActorOuterChangedHandle;
+		
+	/** If this flag is raised we will force refresh on next selection update. */
+	bool bNeedsRefresh : 1;
 };
diff --git a/Engine/Source/Runtime/Engine/Classes/GameFramework/Actor.h b/Engine/Source/Runtime/Engine/Classes/GameFramework/Actor.h
--- a/Engine/Source/Runtime/Engine/Classes/GameFramework/Actor.h
+++ b/Engine/Source/Runtime/Engine/Classes/GameFramework/Actor.h
@@ -1624,6 +1624,7 @@ public:
 	virtual void PreEditUndo() override;
 	virtual void PostEditUndo() override;
 	virtual void PostEditImport() override;
+	virtual void PostTransacted(const FTransactionObjectEvent& TransactionEvent) override;
 	virtual bool IsSelectedInEditor() const override;
 
 	/** When selected can this actor be deleted? */
diff --git a/Engine/Source/Runtime/Engine/Private/ActorEditor.cpp b/Engine/Source/Runtime/Engine/Private/ActorEditor.cpp
--- a/Engine/Source/Runtime/Engine/Private/ActorEditor.cpp
+++ b/Engine/Source/Runtime/Engine/Private/ActorEditor.cpp
@@ -564,6 +564,15 @@ bool AActor::InternalPostEditUndo()
 	return true;
 }
 
+void AActor::PostTransacted(const FTransactionObjectEvent& TransactionEvent)
+{
+	Super::PostTransacted(TransactionEvent);
+	if (TransactionEvent.HasOuterChange())
+	{
+		GEngine->BroadcastLevelActorOuterChanged(this, StaticFindObject(ULevel::StaticClass(), nullptr, *TransactionEvent.GetOriginalObjectOuterPathName().ToString()));
+	}
+}
+
 void AActor::PostEditUndo()
 {
 	if (InternalPostEditUndo())
