Misc. fixes for level streaming UI and workflow. Removing invalid levels from the Levels UI will no longer ensure. Added minimum size for warning about placing actors outside level bounds, and that will also not warn if you are working in the only unlocked level. Level swapping button is now fully clickable, and is hidden when there's only one level. Selected Actor Level text layout has been improved. We now force unloaded levels to load if you select an unloaded level from the button in the viewport.

diff --git a/Engine/Source/Editor/LevelEditor/Private/SLevelViewport.cpp b/Engine/Source/Editor/LevelEditor/Private/SLevelViewport.cpp
--- a/Engine/Source/Editor/LevelEditor/Private/SLevelViewport.cpp
+++ b/Engine/Source/Editor/LevelEditor/Private/SLevelViewport.cpp
@@ -295,12 +295,12 @@ void SLevelViewport::ConstructViewportOverlayContent()
 		.AutoHeight()
 		.Padding(2.0f, 1.0f, 2.0f, 1.0f)
 		[
-			SNew(SHorizontalBox)
+			SNew(SVerticalBox)
 			.Visibility(this, &SLevelViewport::GetSelectedActorsCurrentLevelTextVisibility)
 			// Current level label
-			+ SHorizontalBox::Slot()
-			.AutoWidth()
-			.Padding(2.0f, 1.0f, 2.0f, 1.0f)
+			+ SVerticalBox::Slot()
+			.AutoHeight()
+			.Padding(6.0f, 1.0f, 2.0f, 1.0f)
 			[
 				SNew(STextBlock)
 				.Text(this, &SLevelViewport::GetSelectedActorsCurrentLevelText, true)
@@ -308,9 +308,9 @@ void SLevelViewport::ConstructViewportOverlayContent()
 				.ShadowOffset(FVector2D(1, 1))
 			]
 			// Current level
-			+ SHorizontalBox::Slot()
-			.AutoWidth()
-			.Padding(4.0f, 1.0f, 2.0f, 1.0f)
+			+ SVerticalBox::Slot()
+			.AutoHeight()
+			.Padding(6.0f, 1.0f, 2.0f, 1.0f)
 			[
 				SNew(STextBlock)
 				.Text(this, &SLevelViewport::GetSelectedActorsCurrentLevelText, false)
@@ -333,7 +333,7 @@ void SLevelViewport::ConstructViewportOverlayContent()
 				.VAlign(VAlign_Center)
 				.ButtonStyle(FEditorStyle::Get(), "EditorViewportToolBar.MenuButton")
 				.OnClicked(this, &SLevelViewport::OnMenuClicked)
-				.Visibility(this, &SLevelViewport::GetCurrentLevelTextVisibility)
+				.Visibility(this, &SLevelViewport::GetCurrentLevelButtonVisibility)
 				[
 					SNew(SHorizontalBox)
 					.Visibility(this, &SLevelViewport::GetCurrentLevelTextVisibility)
@@ -3566,7 +3566,20 @@ EVisibility SLevelViewport::GetCurrentLevelTextVisibility() const
 	{
 		ContentVisibility = EVisibility::SelfHitTestInvisible;
 	}
-	return (&GetLevelViewportClient() == GCurrentLevelEditingViewportClient) && !IsPlayInEditorViewportActive() ? ContentVisibility : EVisibility::Collapsed;
+	return (&GetLevelViewportClient() == GCurrentLevelEditingViewportClient) 
+		&& !IsPlayInEditorViewportActive() 
+		&& GetWorld() && GetWorld()->GetCurrentLevel()->OwningWorld->GetLevels().Num() > 1
+		?  ContentVisibility : EVisibility::Collapsed;
+}
+
+EVisibility SLevelViewport::GetCurrentLevelButtonVisibility() const
+{
+	EVisibility TextVisibility = GetCurrentLevelTextVisibility();
+	if (TextVisibility == EVisibility::SelfHitTestInvisible)
+	{
+		TextVisibility = EVisibility::Visible;
+	}
+	return TextVisibility;
 }
 
 EVisibility SLevelViewport::GetSelectedActorsCurrentLevelTextVisibility() const
@@ -3576,7 +3589,11 @@ EVisibility SLevelViewport::GetSelectedActorsCurrentLevelTextVisibility() const
 	{
 		ContentVisibility = EVisibility::SelfHitTestInvisible;
 	}
-	return (&GetLevelViewportClient() == GCurrentLevelEditingViewportClient) && (GEditor->GetSelectedActorCount() > 0) && !IsPlayInEditorViewportActive() ? ContentVisibility : EVisibility::Collapsed;
+	return (&GetLevelViewportClient() == GCurrentLevelEditingViewportClient) 
+		&& (GEditor->GetSelectedActorCount() > 0) 
+		&& !IsPlayInEditorViewportActive() 
+		&& GetWorld() && GetWorld()->GetCurrentLevel()->OwningWorld->GetLevels().Num() > 1
+		? ContentVisibility : EVisibility::Collapsed;
 }
 
 FText SLevelViewport::GetSelectedActorsCurrentLevelText(bool bDrawOnlyLabel) const
diff --git a/Engine/Source/Editor/LevelEditor/Public/SLevelViewport.h b/Engine/Source/Editor/LevelEditor/Public/SLevelViewport.h
--- a/Engine/Source/Editor/LevelEditor/Public/SLevelViewport.h
+++ b/Engine/Source/Editor/LevelEditor/Public/SLevelViewport.h
@@ -295,6 +295,9 @@ public:
 	/** @return The visibility of the current level text display */
 	virtual EVisibility GetCurrentLevelTextVisibility() const;
 
+	/** @return The visibility of the current level button display */
+	virtual EVisibility GetCurrentLevelButtonVisibility() const;
+
 	/** @return The visibility of the current level text display */
 	virtual EVisibility GetSelectedActorsCurrentLevelTextVisibility() const;
 
diff --git a/Engine/Source/Editor/UnrealEd/Classes/Settings/LevelEditorMiscSettings.h b/Engine/Source/Editor/UnrealEd/Classes/Settings/LevelEditorMiscSettings.h
--- a/Engine/Source/Editor/UnrealEd/Classes/Settings/LevelEditorMiscSettings.h
+++ b/Engine/Source/Editor/UnrealEd/Classes/Settings/LevelEditorMiscSettings.h
@@ -79,6 +79,9 @@ public:
 	UPROPERTY(EditAnywhere, config, Category = Levels, Meta = (ClampMin = "0.0", ClampMax = "100.0"))
 	float PercentageThresholdForPrompt;
 
+	UPROPERTY(EditAnywhere, config, Category = Levels)
+	FVector MinimumBoundsForCheckingSize;
+
 public:
 
 	/** The save directory for newly created screenshots */
diff --git a/Engine/Source/Editor/UnrealEd/Private/AssetSelection.cpp b/Engine/Source/Editor/UnrealEd/Private/AssetSelection.cpp
--- a/Engine/Source/Editor/UnrealEd/Private/AssetSelection.cpp
+++ b/Engine/Source/Editor/UnrealEd/Private/AssetSelection.cpp
@@ -448,9 +448,31 @@ namespace ActorPlacementUtils
 				}
 			}
 		}
+		if (InLevel && InLevel->OwningWorld)
+		{
+			int32 LevelCount = InLevel->OwningWorld->GetStreamingLevels().Num();
+			int32 NumLockedLevels = 0;
+			// Check for streaming level count b/c we know there is > 1 streaming level
+			for (ULevelStreaming* StreamingLevel : InLevel->OwningWorld->GetStreamingLevels())
+			{
+				StreamingLevel->bLocked ? NumLockedLevels++ : 0;
+			}
+			// If there is only one unlocked level, a) ours is the unlocked level b/c of the previous IsLevelLocked test and b) we shouldn't try to check for level bounds on the next test
+			if (LevelCount - NumLockedLevels == 1)
+			{
+				return true;
+			}
+		}
 		if (InLevel && GetDefault<ULevelEditorMiscSettings>()->bPromptWhenAddingToLevelOutsideBounds)
 		{
 			FBox CurrentLevelBounds = ALevelBounds::CalculateLevelBounds(InLevel);
+			FVector BoundsExtent = CurrentLevelBounds.GetExtent();
+			if (BoundsExtent.X < GetDefault<ULevelEditorMiscSettings>()->MinimumBoundsForCheckingSize.X
+				&& BoundsExtent.Y < GetDefault<ULevelEditorMiscSettings>()->MinimumBoundsForCheckingSize.Y
+				&& BoundsExtent.Z < GetDefault<ULevelEditorMiscSettings>()->MinimumBoundsForCheckingSize.Z)
+			{
+				return true;
+			}
 			FVector ExpandedScale = FVector(1.0f + (GetDefault<ULevelEditorMiscSettings>()->PercentageThresholdForPrompt / 100.0f));
 			FTransform ExpandedScaleTransform = FTransform::Identity;
 			ExpandedScaleTransform.SetScale3D(ExpandedScale);
diff --git a/Engine/Source/Editor/UnrealEd/Private/Settings/SettingsClasses.cpp b/Engine/Source/Editor/UnrealEd/Private/Settings/SettingsClasses.cpp
--- a/Engine/Source/Editor/UnrealEd/Private/Settings/SettingsClasses.cpp
+++ b/Engine/Source/Editor/UnrealEd/Private/Settings/SettingsClasses.cpp
@@ -348,6 +348,7 @@ ULevelEditorMiscSettings::ULevelEditorMiscSettings( const FObjectInitializer& Ob
 	bPromptWhenAddingToLevelBeforeCheckout = true;
 	bPromptWhenAddingToLevelOutsideBounds = true;
 	PercentageThresholdForPrompt = 20.0f;
+	MinimumBoundsForCheckingSize = FVector(500.0f, 500.0f, 50.0f);
 	bCreateNewAudioDeviceForPlayInEditor = true;
 }
 
diff --git a/Engine/Source/Editor/WorldBrowser/Private/StreamingLevels/StreamingLevelCollectionModel.cpp b/Engine/Source/Editor/WorldBrowser/Private/StreamingLevels/StreamingLevelCollectionModel.cpp
--- a/Engine/Source/Editor/WorldBrowser/Private/StreamingLevels/StreamingLevelCollectionModel.cpp
+++ b/Engine/Source/Editor/WorldBrowser/Private/StreamingLevels/StreamingLevelCollectionModel.cpp
@@ -496,9 +496,10 @@ void FStreamingLevelCollectionModel::FixupInvalidReference_Executed()
 
 void FStreamingLevelCollectionModel::RemoveInvalidSelectedLevels_Executed()
 {
-	for (TSharedPtr<FLevelModel> LevelModel : InvalidSelectedLevels)
+	// needs to be an index-based iterator b/c we are removing elements based on it
+	for (int32 LevelIdx = InvalidSelectedLevels.Num() - 1; LevelIdx == 0; LevelIdx--)
 	{
-		TSharedPtr<FStreamingLevelModel> TargetModel = StaticCastSharedPtr<FStreamingLevelModel>(LevelModel);
+		TSharedPtr<FStreamingLevelModel> TargetModel = StaticCastSharedPtr<FStreamingLevelModel>(InvalidSelectedLevels[LevelIdx]);
 		ULevelStreaming* LevelStreaming = TargetModel->GetLevelStreaming().Get();
 
 		if (LevelStreaming)
diff --git a/Engine/Source/Editor/WorldBrowser/Private/WorldBrowserModule.cpp b/Engine/Source/Editor/WorldBrowser/Private/WorldBrowserModule.cpp
--- a/Engine/Source/Editor/WorldBrowser/Private/WorldBrowserModule.cpp
+++ b/Engine/Source/Editor/WorldBrowser/Private/WorldBrowserModule.cpp
@@ -44,25 +44,22 @@ void FWorldBrowserModule::BuildLevelMenu(FMenuBuilder& MenuBuilder)
 		FLevelModelList ModelList = Model->GetFilteredLevels();
 		for (TSharedPtr<FLevelModel> LevelModel : ModelList)
 		{
-			if (LevelModel->GetLevelObject())
-			{
-				FUIAction Action(FExecuteAction::CreateRaw(this, &FWorldBrowserModule::SetCurrentSublevel, LevelModel->GetLevelObject()),
-					FCanExecuteAction(),
-					FIsActionChecked::CreateRaw(this, &FWorldBrowserModule::IsCurrentSublevel, LevelModel->GetLevelObject()));
-				MenuBuilder.AddMenuEntry(FText::FromString(LevelModel->GetDisplayName()), FText::GetEmpty(), FSlateIcon(), Action, NAME_None, EUserInterfaceActionType::Button);
-			}
+			FUIAction Action(FExecuteAction::CreateRaw(this, &FWorldBrowserModule::SetCurrentSublevel, LevelModel),
+				FCanExecuteAction(),
+				FIsActionChecked::CreateRaw(this, &FWorldBrowserModule::IsCurrentSublevel, LevelModel));
+			MenuBuilder.AddMenuEntry(FText::FromString(LevelModel->GetDisplayName()), FText::GetEmpty(), FSlateIcon(), Action, NAME_None, EUserInterfaceActionType::Button);
 		}
 	}
 }
 
-bool FWorldBrowserModule::IsCurrentSublevel(ULevel* InLevel)
+bool FWorldBrowserModule::IsCurrentSublevel(TSharedPtr<FLevelModel> InLevelModel)
 {
-	return InLevel ? InLevel->IsCurrentLevel() : false;
+	return InLevelModel->IsCurrent();
 }
 
-void FWorldBrowserModule::SetCurrentSublevel(ULevel* InLevel)
+void FWorldBrowserModule::SetCurrentSublevel(TSharedPtr<FLevelModel> InLevelModel)
 {
-	EditorLevelUtils::MakeLevelCurrent(InLevel);
+	InLevelModel->MakeLevelCurrent();
 }
 
 void FWorldBrowserModule::StartupModule()
diff --git a/Engine/Source/Editor/WorldBrowser/Public/WorldBrowserModule.h b/Engine/Source/Editor/WorldBrowser/Public/WorldBrowserModule.h
--- a/Engine/Source/Editor/WorldBrowser/Public/WorldBrowserModule.h
+++ b/Engine/Source/Editor/WorldBrowser/Public/WorldBrowserModule.h
@@ -63,8 +63,8 @@ private:
 	/** Fill out the level menu with entries for level operations */
 	void BuildLevelMenu(FMenuBuilder& MenuBuilder);
 
-	bool IsCurrentSublevel(ULevel* InLevel);
-	void SetCurrentSublevel(ULevel* InLevel);
+	bool IsCurrentSublevel(TSharedPtr<class FLevelModel> InLevelModel);
+	void SetCurrentSublevel(TSharedPtr<class FLevelModel> InLevelModel);
 				
 private:
 	TWeakPtr<class FLevelCollectionModel>	WorldModel;
