Level Editor context menus converted to editor menu subsystem
#rnx

diff --git a/Engine/Source/Editor/DetailCustomizations/DetailCustomizations.Build.cs b/Engine/Source/Editor/DetailCustomizations/DetailCustomizations.Build.cs
--- a/Engine/Source/Editor/DetailCustomizations/DetailCustomizations.Build.cs
+++ b/Engine/Source/Editor/DetailCustomizations/DetailCustomizations.Build.cs
@@ -59,6 +59,7 @@ public class DetailCustomizations : ModuleRules
 				"AudioSettingsEditor",
 				"HeadMountedDisplay",
                 "DataTableEditor",
+				"EditorMenus",
             }
 		);
 
diff --git a/Engine/Source/Editor/DetailCustomizations/Private/ActorDetails.cpp b/Engine/Source/Editor/DetailCustomizations/Private/ActorDetails.cpp
--- a/Engine/Source/Editor/DetailCustomizations/Private/ActorDetails.cpp
+++ b/Engine/Source/Editor/DetailCustomizations/Private/ActorDetails.cpp
@@ -18,7 +18,7 @@
 #include "Widgets/Images/SImage.h"
 #include "Widgets/Text/STextBlock.h"
 #include "Widgets/Layout/SBox.h"
-#include "Framework/MultiBox/MultiBoxBuilder.h"
+#include "EditorMenuSubsystem.h"
 #include "Widgets/Input/SButton.h"
 #include "Widgets/Input/SComboButton.h"
 #include "EditorStyleSet.h"
@@ -399,11 +399,18 @@ bool FActorDetails::AreAnySelectedActorsInLevelScript() const
 /** Util to create a menu for events we can add for the selected actor */
 TSharedRef<SWidget> FActorDetails::MakeEventOptionsWidgetFromSelection()
 {
-	FMenuBuilder EventMenuBuilder( true, NULL );
+	UEditorMenuSubsystem* EditorMenus = UEditorMenuSubsystem::Get();
+	static const FName MenuName("DetailCustomizations.EventOptions");
+	if (!EditorMenus->IsMenuRegistered(MenuName))
+	{
+		EditorMenus->RegisterMenu(MenuName);
+	}
 
+	FEditorMenuContext Context;
+	UEditorMenu* Menu = EditorMenus->GenerateMenu(MenuName, Context);
 	AActor* Actor = SelectedActors[0].Get();
-	FKismetEditorUtilities::AddLevelScriptEventOptionsForActor(EventMenuBuilder, SelectedActors[0], true, true, false);
-	return EventMenuBuilder.MakeWidget();
+	FKismetEditorUtilities::AddLevelScriptEventOptionsForActor(Menu, SelectedActors[0], true, true, false);
+	return EditorMenus->GenerateWidget(Menu);
 }
 
 
diff --git a/Engine/Source/Editor/EditorMenus/Private/EditorMenuEntry.cpp b/Engine/Source/Editor/EditorMenus/Private/EditorMenuEntry.cpp
--- a/Engine/Source/Editor/EditorMenus/Private/EditorMenuEntry.cpp
+++ b/Engine/Source/Editor/EditorMenus/Private/EditorMenuEntry.cpp
@@ -81,6 +81,18 @@ FEditorMenuEntry FEditorMenuEntry::InitSubMenu(const FName InParentMenu, const F
 	return Entry;
 }
 
+FEditorMenuEntry FEditorMenuEntry::InitSubMenu(const FName InParentMenu, const FName InName, const FEditorUIActionChoice& InAction, const TSharedRef<SWidget>& InWidget, const FNewEditorMenuChoice& InMakeMenu, bool bInShouldCloseWindowAfterMenuSelection)
+{
+	FEditorMenuEntry Entry(UEditorMenuSubsystem::Get()->CurrentOwner(), InName, EMultiBlockType::MenuEntry);
+	Entry.Action = InAction;
+	Entry.MakeWidget.BindLambda([=](const FEditorMenuContext&) { return InWidget; });
+	Entry.bShouldCloseWindowAfterMenuSelection = bInShouldCloseWindowAfterMenuSelection;
+	Entry.SubMenuData.bIsSubMenu = true;
+	Entry.SubMenuData.ConstructMenu = InMakeMenu;
+	Entry.SubMenuData.bOpenSubMenuOnClick = false;
+	return Entry;
+}
+
 FEditorMenuEntry FEditorMenuEntry::InitToolBarButton(const FName InName, const FEditorUIActionChoice& InAction, const TAttribute<FText>& InLabel, const TAttribute<FText>& InToolTip, const TAttribute<FSlateIcon>& InIcon, const EUserInterfaceActionType InUserInterfaceActionType, FName InTutorialHighlightName)
 {
 	FEditorMenuEntry Entry(UEditorMenuSubsystem::Get()->CurrentOwner(), InName, EMultiBlockType::ToolBarButton);
@@ -124,11 +136,11 @@ FEditorMenuEntry FEditorMenuEntry::InitToolBarSeparator(const FName InName)
 	return FEditorMenuEntry(UEditorMenuSubsystem::Get()->CurrentOwner(), InName, EMultiBlockType::ToolBarSeparator);
 }
 
-FEditorMenuEntry FEditorMenuEntry::InitWidget(const FName InName, const TSharedRef<SWidget>& Widget, const FText& Label, bool bNoIndent, bool bSearchable)
+FEditorMenuEntry FEditorMenuEntry::InitWidget(const FName InName, const TSharedRef<SWidget>& InWidget, const FText& Label, bool bNoIndent, bool bSearchable)
 {
 	FEditorMenuEntry Entry(UEditorMenuSubsystem::Get()->CurrentOwner(), InName, EMultiBlockType::Widget);
 	Entry.Label = Label;
-	Entry.MakeWidget.BindLambda([Widget](const FEditorMenuContext&) { return Widget; });
+	Entry.MakeWidget.BindLambda([=](const FEditorMenuContext&) { return InWidget; });
 	Entry.WidgetData.bNoIndent = bNoIndent;
 	Entry.WidgetData.bSearchable = bSearchable;
 	return Entry;
@@ -144,6 +156,11 @@ void FEditorMenuEntry::ResetActions()
 	//ScriptObject = nullptr;
 }
 
+bool FEditorMenuEntry::IsNonLegacyDynamicConstruct() const
+{
+	return Construct.IsBound() || IsScriptObjectDynamicConstruct();
+}
+
 bool FEditorMenuEntry::IsScriptObjectDynamicConstruct() const
 {
 	static const FName ConstructMenuEntryName = GET_FUNCTION_NAME_CHECKED(UEditorMenuEntryScript, ConstructMenuEntry);
diff --git a/Engine/Source/Editor/EditorMenus/Private/EditorMenuSection.cpp b/Engine/Source/Editor/EditorMenus/Private/EditorMenuSection.cpp
--- a/Engine/Source/Editor/EditorMenus/Private/EditorMenuSection.cpp
+++ b/Engine/Source/Editor/EditorMenus/Private/EditorMenuSection.cpp
@@ -113,6 +113,11 @@ int32 FEditorMenuSection::IndexOfBlock(const FName InName) const
 	return INDEX_NONE;
 }
 
+bool FEditorMenuSection::IsNonLegacyDynamic() const
+{
+	return EditorMenuSectionDynamic || Construct.NewEditorMenuDelegate.IsBound();
+}
+
 void FEditorMenuSection::AssembleBlock(const FEditorMenuEntry& Block)
 {
 	const EEditorMenuInsertType Position = Block.InsertPosition.Position;
@@ -145,7 +150,6 @@ void FEditorMenuSection::AssembleBlock(const FEditorMenuEntry& Block)
 	}
 }
 
-
 int32 FEditorMenuSection::RemoveEntry(const FName InName)
 {
 	return Blocks.RemoveAll([InName](const FEditorMenuEntry& Block) { return Block.Name == InName; });
@@ -160,6 +164,7 @@ int32 FEditorMenuSection::RemoveEntriesByOwner(const FEditorMenuOwner InOwner)
 
 	return 0;
 }
+
 int32 FEditorMenuSection::FindBlockInsertIndex(const FEditorMenuEntry& InBlock) const
 {
 	const FEditorMenuInsert InPosition = InBlock.InsertPosition;
diff --git a/Engine/Source/Editor/EditorMenus/Private/EditorMenuSubsystem.cpp b/Engine/Source/Editor/EditorMenus/Private/EditorMenuSubsystem.cpp
--- a/Engine/Source/Editor/EditorMenus/Private/EditorMenuSubsystem.cpp
+++ b/Engine/Source/Editor/EditorMenus/Private/EditorMenuSubsystem.cpp
@@ -172,42 +172,78 @@ void UEditorMenuSubsystem::AssembleMenuSection(UEditorMenu* GeneratedMenu, const
 {
 	// Build list of blocks in expected order including blocks created by construct delegates
 	TArray<FEditorMenuEntry> RemainingBlocks;
+
+	UEditorMenu* ConstructedEntries = nullptr;
 	for (const FEditorMenuEntry& Block : OtherSection.Blocks)
 	{
-		if (Block.IsScriptObjectDynamicConstruct() || Block.Construct.IsBound())
+		if (!Block.IsNonLegacyDynamicConstruct())
 		{
-			UEditorMenu* Constructed = NewObject<UEditorMenu>(this);
-			Constructed->Context = DestSection->Context;
+			RemainingBlocks.Add(Block);
+			continue;
+		}
 
-			if (Block.IsScriptObjectDynamicConstruct())
-			{
-				Block.ScriptObject->ConstructMenuEntry(Constructed, DestSection->Name, DestSection->Context);
-			}
-			else
-			{
-				FEditorMenuSection& ConstructedSection = Constructed->AddSection(DestSection->Name);
-				ConstructedSection.Context = Constructed->Context;
-				Block.Construct.Execute(ConstructedSection);
-			}
+		if (ConstructedEntries == nullptr)
+		{
+			ConstructedEntries = NewObject<UEditorMenu>(this);
+			ConstructedEntries->Context = DestSection->Context;
+		}
+
+		TArray<FEditorMenuEntry> GeneratedEntries;
+		GeneratedEntries.Add(Block);
 
-			// Combine all user's choice of selections here into the current section target
-			// If the user wants to add items to different sections they will need to create dynamic section instead (for now)
-			for (FEditorMenuSection& ConstructedSection : Constructed->Sections)
+		int32 NumIterations = 0;
+		while (GeneratedEntries.Num() > 0)
+		{
+			FEditorMenuEntry& GeneratedEntry = GeneratedEntries[0];
+			if (GeneratedEntry.IsNonLegacyDynamicConstruct())
 			{
-				for (FEditorMenuEntry& ConstructedBlock : ConstructedSection.Blocks)
+				if (NumIterations++ > 5000)
+				{
+					UE_LOG(LogEditorMenus, Warning, TEXT("Possible infinite loop for menu: %s, section: %s, block: %s"), *Other->MenuName.ToString(), *OtherSection.Name.ToString(), *Block.Name.ToString());
+					break;
+				}
+				
+				ConstructedEntries->Sections.Reset();
+				if (GeneratedEntry.IsScriptObjectDynamicConstruct())
+				{
+					GeneratedEntry.ScriptObject->ConstructMenuEntry(ConstructedEntries, DestSection->Name, DestSection->Context);
+				}
+				else
+				{
+					FEditorMenuSection& ConstructedSection = ConstructedEntries->AddSection(DestSection->Name);
+					ConstructedSection.Context = ConstructedEntries->Context;
+					GeneratedEntry.Construct.Execute(ConstructedSection);
+				}
+				GeneratedEntries.RemoveAt(0, 1, false);
+
+				// Combine all user's choice of selections here into the current section target
+				// If the user wants to add items to different sections they will need to create dynamic section instead (for now)
+				int32 NumBlocksInserted = 0;
+				for (FEditorMenuSection& ConstructedSection : ConstructedEntries->Sections)
 				{
-					if (ConstructedBlock.InsertPosition.IsDefault())
+					for (FEditorMenuEntry& ConstructedBlock : ConstructedSection.Blocks)
 					{
-						ConstructedBlock.InsertPosition = Block.InsertPosition;
+						if (ConstructedBlock.InsertPosition.IsDefault())
+						{
+							ConstructedBlock.InsertPosition = Block.InsertPosition;
+						}
 					}
-					RemainingBlocks.Add(ConstructedBlock);
+					GeneratedEntries.Insert(ConstructedSection.Blocks, NumBlocksInserted);
+					NumBlocksInserted += ConstructedSection.Blocks.Num();
 				}
 			}
+			else
+			{
+				RemainingBlocks.Add(GeneratedEntry);
+				GeneratedEntries.RemoveAt(0, 1, false);
+			}
 		}
-		else
-		{
-			RemainingBlocks.Add(Block);
-		}
+	}
+
+	if (ConstructedEntries)
+	{
+		ConstructedEntries->Sections.Empty();
+		ConstructedEntries = nullptr;
 	}
 
 	// Repeatedly loop because insert location may not exist until later in list
@@ -242,35 +278,70 @@ void UEditorMenuSubsystem::AssembleMenuSection(UEditorMenu* GeneratedMenu, const
 void UEditorMenuSubsystem::AssembleMenu(UEditorMenu* GeneratedMenu, const UEditorMenu* Other)
 {
 	TArray<FEditorMenuSection> RemainingSections;
+
+	UEditorMenu* ConstructedSections = nullptr;
 	for (const FEditorMenuSection& OtherSection : Other->Sections)
 	{
-		if (OtherSection.EditorMenuSectionDynamic || OtherSection.Construct.NewEditorMenuDelegate.IsBound())
+		if (!OtherSection.IsNonLegacyDynamic())
+		{
+			RemainingSections.Add(OtherSection);
+			continue;
+		}
+		
+		if (ConstructedSections == nullptr)
 		{
-			UEditorMenu* ConstructedSections = NewObject<UEditorMenu>(this);
+			ConstructedSections = NewObject<UEditorMenu>(this);
 			ConstructedSections->Context = GeneratedMenu->Context;
+		}
 
-			if (OtherSection.EditorMenuSectionDynamic)
-			{
-				OtherSection.EditorMenuSectionDynamic->ConstructSections(ConstructedSections, GeneratedMenu->Context);
-			}
-			else if (OtherSection.Construct.NewEditorMenuDelegate.IsBound())
-			{
-				OtherSection.Construct.NewEditorMenuDelegate.Execute(ConstructedSections);
-			}
+		TArray<FEditorMenuSection> GeneratedSections;
+		GeneratedSections.Add(OtherSection);
 
-			for (FEditorMenuSection& ConstructedSection : ConstructedSections->Sections)
+		int32 NumIterations = 0;
+		while (GeneratedSections.Num() > 0)
+		{
+			if (GeneratedSections[0].IsNonLegacyDynamic())
 			{
-				if (ConstructedSection.InsertPosition.IsDefault())
+				if (NumIterations++ > 5000)
+				{
+					UE_LOG(LogEditorMenus, Warning, TEXT("Possible infinite loop for menu: %s, section: %s"), *Other->MenuName.ToString(), *OtherSection.Name.ToString());
+					break;
+				}
+
+				ConstructedSections->Sections.Reset();
+				
+				if (GeneratedSections[0].EditorMenuSectionDynamic)
+				{
+					GeneratedSections[0].EditorMenuSectionDynamic->ConstructSections(ConstructedSections, GeneratedMenu->Context);
+				}
+				else if (GeneratedSections[0].Construct.NewEditorMenuDelegate.IsBound())
+				{
+					GeneratedSections[0].Construct.NewEditorMenuDelegate.Execute(ConstructedSections);
+				}
+
+				for (FEditorMenuSection& ConstructedSection : ConstructedSections->Sections)
 				{
-					ConstructedSection.InsertPosition = OtherSection.InsertPosition;
+					if (ConstructedSection.InsertPosition.IsDefault())
+					{
+						ConstructedSection.InsertPosition = GeneratedSections[0].InsertPosition;
+					}
 				}
-				RemainingSections.Add(ConstructedSection);
+				
+				GeneratedSections.RemoveAt(0, 1, false);				
+				GeneratedSections.Insert(ConstructedSections->Sections, 0);
+			}
+			else
+			{
+				RemainingSections.Add(GeneratedSections[0]);
+				GeneratedSections.RemoveAt(0, 1, false);
 			}
 		}
-		else
-		{
-			RemainingSections.Add(OtherSection);
-		}
+	}
+
+	if (ConstructedSections)
+	{
+		ConstructedSections->Sections.Empty();
+		ConstructedSections = nullptr;
 	}
 
 	while (RemainingSections.Num() > 0)
diff --git a/Engine/Source/Editor/EditorMenus/Public/EditorMenuEntry.h b/Engine/Source/Editor/EditorMenus/Public/EditorMenuEntry.h
--- a/Engine/Source/Editor/EditorMenus/Public/EditorMenuEntry.h
+++ b/Engine/Source/Editor/EditorMenus/Public/EditorMenuEntry.h
@@ -87,7 +87,9 @@ struct EDITORMENUS_API FEditorMenuEntry
 	static FEditorMenuEntry InitMenuEntry(const FName InName, const TAttribute<FText>& InLabel, const TAttribute<FText>& InToolTip, const TAttribute<FSlateIcon>& InIcon, const FEditorUIActionChoice& InAction, const EUserInterfaceActionType UserInterfaceActionType = EUserInterfaceActionType::Button, const FName InTutorialHighlightName = NAME_None);
 	static FEditorMenuEntry InitMenuEntry(const TSharedPtr< const FUICommandInfo >& InCommand, const TAttribute<FText>& InLabelOverride = TAttribute<FText>(), const TAttribute<FText>& InToolTipOverride = TAttribute<FText>(), const TAttribute<FSlateIcon>& InIconOverride = TAttribute<FSlateIcon>(), const FName InTutorialHighlightName = NAME_None, const FName InNameOverride = NAME_None);
 	static FEditorMenuEntry InitMenuEntry(const FName InName, const FEditorUIActionChoice& InAction, const TSharedRef<SWidget>& Widget);
-	static FEditorMenuEntry InitSubMenu(const FName InParentMenu, const FName InName, const TAttribute<FText>& InLabel, const TAttribute<FText>& InToolTip, const FNewEditorMenuChoice& InMakeMenu, bool bInOpenSubMenuOnClick = false, const TAttribute<FSlateIcon>& InIcon = TAttribute<FSlateIcon>(), const bool ShouldCloseWindowAfterMenuSelection = true);
+
+	static FEditorMenuEntry InitSubMenu(const FName InParentMenu, const FName InName, const TAttribute<FText>& InLabel, const TAttribute<FText>& InToolTip, const FNewEditorMenuChoice& InMakeMenu, bool bInOpenSubMenuOnClick = false, const TAttribute<FSlateIcon>& InIcon = TAttribute<FSlateIcon>(), const bool bShouldCloseWindowAfterMenuSelection = true);
+	static FEditorMenuEntry InitSubMenu(const FName InParentMenu, const FName InName, const FEditorUIActionChoice& InAction, const TSharedRef<SWidget>& InWidget, const FNewEditorMenuChoice& InMakeMenu, bool bShouldCloseWindowAfterMenuSelection = true);
 
 	static FEditorMenuEntry InitToolBarButton(const FName InName, const FEditorUIActionChoice& InAction, const TAttribute<FText>& InLabel = TAttribute<FText>(), const TAttribute<FText>& InToolTip = TAttribute<FText>(), const TAttribute<FSlateIcon>& InIcon = TAttribute<FSlateIcon>(), const EUserInterfaceActionType UserInterfaceActionType = EUserInterfaceActionType::Button, FName InTutorialHighlightName = NAME_None);
 	static FEditorMenuEntry InitToolBarButton(const TSharedPtr< const FUICommandInfo >& InCommand, const TAttribute<FText>& InLabelOverride = TAttribute<FText>(), const TAttribute<FText>& InToolTipOverride = TAttribute<FText>(), const TAttribute<FSlateIcon>& InIconOverride = TAttribute<FSlateIcon>(), FName InTutorialHighlightName = NAME_None, const FName InNameOverride = NAME_None);
@@ -96,7 +98,7 @@ struct EDITORMENUS_API FEditorMenuEntry
 	static FEditorMenuEntry InitMenuSeparator(const FName InName);
 	static FEditorMenuEntry InitToolBarSeparator(const FName InName);
 
-	static FEditorMenuEntry InitWidget(const FName InName, const TSharedRef<SWidget>& Widget, const FText& Label, bool bNoIndent = false, bool bSearchable = true);
+	static FEditorMenuEntry InitWidget(const FName InName, const TSharedRef<SWidget>& InWidget, const FText& Label, bool bNoIndent = false, bool bSearchable = true);
 
 	bool IsSubMenu() const { return SubMenuData.bIsSubMenu; }
 
@@ -111,6 +113,8 @@ private:
 
 	bool IsScriptObjectDynamicConstruct() const;
 
+	bool IsNonLegacyDynamicConstruct() const;
+
 public:
 
 	UPROPERTY(EditAnywhere, BlueprintReadWrite, Category = "Editor UI")
diff --git a/Engine/Source/Editor/EditorMenus/Public/EditorMenuSection.h b/Engine/Source/Editor/EditorMenus/Public/EditorMenuSection.h
--- a/Engine/Source/Editor/EditorMenus/Public/EditorMenuSection.h
+++ b/Engine/Source/Editor/EditorMenus/Public/EditorMenuSection.h
@@ -59,6 +59,8 @@ private:
 
 	void AssembleBlock(const FEditorMenuEntry& InBlock);
 
+	bool IsNonLegacyDynamic() const;
+
 	friend class UEditorMenuSectionExtensions;
 	friend class UEditorMenuSubsystem;
 	friend class UEditorMenu;
diff --git a/Engine/Source/Editor/HierarchicalLODOutliner/HierarchicalLODOutliner.Build.cs b/Engine/Source/Editor/HierarchicalLODOutliner/HierarchicalLODOutliner.Build.cs
--- a/Engine/Source/Editor/HierarchicalLODOutliner/HierarchicalLODOutliner.Build.cs
+++ b/Engine/Source/Editor/HierarchicalLODOutliner/HierarchicalLODOutliner.Build.cs
@@ -45,6 +45,7 @@ public class HierarchicalLODOutliner : ModuleRules
                     "UnrealEd",
     				"SlateCore",
     				"Slate",
+					"EditorMenus",
                 }
             );
 
diff --git a/Engine/Source/Editor/HierarchicalLODOutliner/Private/HLODOutliner.cpp b/Engine/Source/Editor/HierarchicalLODOutliner/Private/HLODOutliner.cpp
--- a/Engine/Source/Editor/HierarchicalLODOutliner/Private/HLODOutliner.cpp
+++ b/Engine/Source/Editor/HierarchicalLODOutliner/Private/HLODOutliner.cpp
@@ -12,6 +12,7 @@
 #include "SlateOptMacros.h"
 #include "Framework/MultiBox/MultiBoxExtender.h"
 #include "Framework/MultiBox/MultiBoxBuilder.h"
+#include "EditorMenuSubsystem.h"
 #include "Widgets/Input/SButton.h"
 #include "Widgets/Images/SImage.h"
 #include "EditorStyleSet.h"
@@ -1468,12 +1469,6 @@ namespace HLODOutliner
 			return nullptr;
 		}
 
-		// Build up the menu for a selection
-		const bool bCloseAfterSelection = true;
-		TSharedPtr<FExtender> Extender = MakeShareable(new FExtender);
-
-		FMenuBuilder MenuBuilder(bCloseAfterSelection, TSharedPtr<FUICommandList>(), Extender);
-
 		// Multi-selection support, check if all selected items are of the same type, if so return the appropriate context menu
 		auto SelectedItems = TreeView->GetSelectedItems();
 		ITreeItem::TreeItemType Type = ITreeItem::Invalid;
@@ -1497,8 +1492,18 @@ namespace HLODOutliner
 
 		if (SelectedItems.Num() && bSameType)
 		{
-			TreeView->GetSelectedItems()[0]->GenerateContextMenu(MenuBuilder, *this);
-			return MenuBuilder.MakeWidget();
+			UEditorMenuSubsystem* EditorMenus = UEditorMenuSubsystem::Get();
+			static const FName MenuName = "HierarchicalLODOutliner.HLODOutlinerContextMenu";
+			if (!EditorMenus->IsMenuRegistered(MenuName))
+			{
+				EditorMenus->RegisterMenu(MenuName);
+			}
+
+			// Build up the menu for a selection
+			FEditorMenuContext Context;
+			UEditorMenu* Menu = EditorMenus->GenerateMenu(MenuName, Context);
+			TreeView->GetSelectedItems()[0]->GenerateContextMenu(Menu, *this);
+			return EditorMenus->GenerateWidget(Menu);
 		}
 
 		return TSharedPtr<SWidget>();
diff --git a/Engine/Source/Editor/HierarchicalLODOutliner/Private/ITreeItem.h b/Engine/Source/Editor/HierarchicalLODOutliner/Private/ITreeItem.h
--- a/Engine/Source/Editor/HierarchicalLODOutliner/Private/ITreeItem.h
+++ b/Engine/Source/Editor/HierarchicalLODOutliner/Private/ITreeItem.h
@@ -6,7 +6,7 @@
 #include "Widgets/SWidget.h"
 #include "Editor/HierarchicalLODOutliner/Private/HLODOutlinerDragDrop.h"
 
-class FMenuBuilder;
+class UEditorMenu;
 
 namespace HLODOutliner
 {
@@ -100,7 +100,7 @@ namespace HLODOutliner
 		virtual void OnExpansionChanged() {};
 
 		/** Generate a context menu for this item. Only called if *only* this item is selected. */
-		virtual void GenerateContextMenu(FMenuBuilder& MenuBuilder, SHLODOutliner& Outliner) {};
+		virtual void GenerateContextMenu(UEditorMenu* Menu, SHLODOutliner& Outliner) {};
 
 		/** Populate the specified drag/drop payload with any relevant information for this type */
 		virtual void PopulateDragDropPayload(FDragDropPayload& Payload) const = 0;
diff --git a/Engine/Source/Editor/HierarchicalLODOutliner/Private/LODActorItem.cpp b/Engine/Source/Editor/HierarchicalLODOutliner/Private/LODActorItem.cpp
--- a/Engine/Source/Editor/HierarchicalLODOutliner/Private/LODActorItem.cpp
+++ b/Engine/Source/Editor/HierarchicalLODOutliner/Private/LODActorItem.cpp
@@ -3,7 +3,7 @@
 #include "LODActorItem.h"
 #include "Textures/SlateIcon.h"
 #include "Framework/Commands/UIAction.h"
-#include "Framework/MultiBox/MultiBoxBuilder.h"
+#include "EditorMenuSubsystem.h"
 #include "Modules/ModuleManager.h"
 #include "HLODOutliner.h"
 #include "IHierarchicalLODUtilities.h"
@@ -24,25 +24,26 @@ bool HLODOutliner::FLODActorItem::CanInteract() const
 	return true;
 }
 
-void HLODOutliner::FLODActorItem::GenerateContextMenu(FMenuBuilder& MenuBuilder, SHLODOutliner& Outliner)
+void HLODOutliner::FLODActorItem::GenerateContextMenu(UEditorMenu* Menu, SHLODOutliner& Outliner)
 {
 	auto SharedOutliner = StaticCastSharedRef<SHLODOutliner>(Outliner.AsShared());
 
-	MenuBuilder.AddMenuEntry(LOCTEXT("SelectLODActor", "Select LOD Actor"), FText(), FSlateIcon(), FUIAction(FExecuteAction::CreateRaw(&Outliner, &SHLODOutliner::SelectLODActor)));
+	FEditorMenuSection& Section = Menu->AddSection("Section");
+	Section.AddMenuEntry("SelectLODActor", LOCTEXT("SelectLODActor", "Select LOD Actor"), FText(), FSlateIcon(), FUIAction(FExecuteAction::CreateRaw(&Outliner, &SHLODOutliner::SelectLODActor)));
 
-	MenuBuilder.AddMenuEntry(LOCTEXT("SelectContainedActors", "Select Contained Actors"), FText(), FSlateIcon(), FUIAction(FExecuteAction::CreateRaw(&Outliner, &SHLODOutliner::SelectContainedActors)));
+	Section.AddMenuEntry("SelectContainedActors", LOCTEXT("SelectContainedActors", "Select Contained Actors"), FText(), FSlateIcon(), FUIAction(FExecuteAction::CreateRaw(&Outliner, &SHLODOutliner::SelectContainedActors)));
 
 	if (!LODActor->IsBuilt())
 	{
-		MenuBuilder.AddMenuEntry(LOCTEXT("BuildLODActorMesh", "Build Proxy Mesh"), FText(), FSlateIcon(), FUIAction(FExecuteAction::CreateRaw(&Outliner, &SHLODOutliner::BuildLODActor)));
+		Section.AddMenuEntry("BuildLODActorMesh", LOCTEXT("BuildLODActorMesh", "Build Proxy Mesh"), FText(), FSlateIcon(), FUIAction(FExecuteAction::CreateRaw(&Outliner, &SHLODOutliner::BuildLODActor)));
 	}
 	else
 	{		
-		MenuBuilder.AddMenuEntry(LOCTEXT("ForceView", "ForceView"), FText(), FSlateIcon(), FUIAction(FExecuteAction::CreateRaw(&Outliner, &SHLODOutliner::ForceViewLODActor)));
-		MenuBuilder.AddMenuEntry(LOCTEXT("RebuildLODActorMesh", "Rebuild Proxy Mesh"), FText(), FSlateIcon(), FUIAction(FExecuteAction::CreateRaw(&Outliner, &SHLODOutliner::RebuildLODActor)));
+		Section.AddMenuEntry("ForceView", LOCTEXT("ForceView", "ForceView"), FText(), FSlateIcon(), FUIAction(FExecuteAction::CreateRaw(&Outliner, &SHLODOutliner::ForceViewLODActor)));
+		Section.AddMenuEntry("RebuildLODActorMesh", LOCTEXT("RebuildLODActorMesh", "Rebuild Proxy Mesh"), FText(), FSlateIcon(), FUIAction(FExecuteAction::CreateRaw(&Outliner, &SHLODOutliner::RebuildLODActor)));
 	}
 
-	MenuBuilder.AddMenuEntry(LOCTEXT("CreateHLODVolume", "Create Containing Hierarchical Volume"), FText(), FSlateIcon(), FUIAction(FExecuteAction::CreateRaw(&Outliner, &SHLODOutliner::CreateHierarchicalVolumeForActor)));
+	Section.AddMenuEntry("CreateHLODVolume", LOCTEXT("CreateHLODVolume", "Create Containing Hierarchical Volume"), FText(), FSlateIcon(), FUIAction(FExecuteAction::CreateRaw(&Outliner, &SHLODOutliner::CreateHierarchicalVolumeForActor)));
 
 	AActor* Actor = LODActor.Get();
 	FHierarchicalLODUtilitiesModule& Module = FModuleManager::LoadModuleChecked<FHierarchicalLODUtilitiesModule>("HierarchicalLODUtilities");
@@ -50,11 +51,11 @@ void HLODOutliner::FLODActorItem::GenerateContextMenu(FMenuBuilder& MenuBuilder,
 	ALODActor* ParentActor = Utilities->GetParentLODActor(Actor);
 	if (ParentActor && Parent.Pin()->GetTreeItemType() == TreeItemType::HierarchicalLODActor)
 	{		
-		MenuBuilder.AddMenuEntry(LOCTEXT("RemoveChildFromCluster", "Remove from cluster"), FText(), FSlateIcon(), FUIAction(FExecuteAction::CreateRaw(&Outliner, &SHLODOutliner::RemoveLODActorFromCluster)));
+		Section.AddMenuEntry("RemoveChildFromCluster", LOCTEXT("RemoveChildFromCluster", "Remove from cluster"), FText(), FSlateIcon(), FUIAction(FExecuteAction::CreateRaw(&Outliner, &SHLODOutliner::RemoveLODActorFromCluster)));
 	}
 	else
 	{
-		MenuBuilder.AddMenuEntry(LOCTEXT("DeleteCluster", "Delete Cluster"), FText(), FSlateIcon(), FUIAction(FExecuteAction::CreateRaw(&Outliner, &SHLODOutliner::DeleteCluster)));
+		Section.AddMenuEntry("DeleteCluster", LOCTEXT("DeleteCluster", "Delete Cluster"), FText(), FSlateIcon(), FUIAction(FExecuteAction::CreateRaw(&Outliner, &SHLODOutliner::DeleteCluster)));
 	}
 }
 
diff --git a/Engine/Source/Editor/HierarchicalLODOutliner/Private/LODActorItem.h b/Engine/Source/Editor/HierarchicalLODOutliner/Private/LODActorItem.h
--- a/Engine/Source/Editor/HierarchicalLODOutliner/Private/LODActorItem.h
+++ b/Engine/Source/Editor/HierarchicalLODOutliner/Private/LODActorItem.h
@@ -9,7 +9,7 @@
 #include "Engine/LODActor.h"
 #include "Editor/HierarchicalLODOutliner/Private/ITreeItem.h"
 
-class FMenuBuilder;
+class UEditorMenu;
 
 namespace HLODOutliner
 {
@@ -39,7 +39,7 @@ namespace HLODOutliner
 
 		//~ Begin ITreeItem Interface.
 		virtual bool CanInteract() const override;
-		virtual void GenerateContextMenu(FMenuBuilder& MenuBuilder, SHLODOutliner& Outliner) override;
+		virtual void GenerateContextMenu(UEditorMenu* Menu, SHLODOutliner& Outliner) override;
 		virtual FString GetDisplayString() const override;
 		virtual FSlateColor GetTint() const override;
 		virtual FTreeItemID GetID() override;
diff --git a/Engine/Source/Editor/HierarchicalLODOutliner/Private/LODLevelItem.cpp b/Engine/Source/Editor/HierarchicalLODOutliner/Private/LODLevelItem.cpp
--- a/Engine/Source/Editor/HierarchicalLODOutliner/Private/LODLevelItem.cpp
+++ b/Engine/Source/Editor/HierarchicalLODOutliner/Private/LODLevelItem.cpp
@@ -7,7 +7,7 @@
 #include "IHierarchicalLODUtilities.h"
 #include "HierarchicalLODUtilitiesModule.h"
 #include "Engine/Selection.h"
-#include "Framework/MultiBox/MultiBoxBuilder.h"
+#include "EditorMenuSubsystem.h"
 #include "HLODOutliner.h"
 #include "Editor.h"
 
@@ -25,7 +25,7 @@ bool HLODOutliner::FLODLevelItem::CanInteract() const
 	return true;
 }
 
-void HLODOutliner::FLODLevelItem::GenerateContextMenu(FMenuBuilder& MenuBuilder, SHLODOutliner& Outliner)
+void HLODOutliner::FLODLevelItem::GenerateContextMenu(UEditorMenu* Menu, SHLODOutliner& Outliner)
 {
 	// Collect actors / components available for creating a new HLOD cluster 
 	USelection* CurrentSelection = GEditor->GetSelectedActors();	
@@ -45,7 +45,8 @@ void HLODOutliner::FLODLevelItem::GenerateContextMenu(FMenuBuilder& MenuBuilder,
 		}
 
 		// Generate context menu with option to create a new cluster from the current viewport selection
-		MenuBuilder.AddMenuEntry(LOCTEXT("CreateClusterFromViewportSelection", "Create Cluster from World Selection"), FText(), FSlateIcon(), 
+		FEditorMenuSection& Section = Menu->AddSection("Section");
+		Section.AddMenuEntry("CreateClusterFromViewportSelection", LOCTEXT("CreateClusterFromViewportSelection", "Create Cluster from World Selection"), FText(), FSlateIcon(), 
 			FUIAction(FExecuteAction::CreateLambda([&Outliner, ValidActorsToCluster, this]()
 			{
 				Outliner.CreateClusterFromActors(ValidActorsToCluster, LODLevelIndex);
diff --git a/Engine/Source/Editor/HierarchicalLODOutliner/Private/LODLevelItem.h b/Engine/Source/Editor/HierarchicalLODOutliner/Private/LODLevelItem.h
--- a/Engine/Source/Editor/HierarchicalLODOutliner/Private/LODLevelItem.h
+++ b/Engine/Source/Editor/HierarchicalLODOutliner/Private/LODLevelItem.h
@@ -6,7 +6,7 @@
 #include "TreeItemID.h"
 #include "ITreeItem.h"
 
-class FMenuBuilder;
+class UEditorMenu;
 
 namespace HLODOutliner
 {
@@ -37,7 +37,7 @@ namespace HLODOutliner
 
 		//~ Begin ITreeItem Interface.
 		virtual bool CanInteract() const override;
-		virtual void GenerateContextMenu(FMenuBuilder& MenuBuilder, SHLODOutliner& Outliner) override;
+		virtual void GenerateContextMenu(UEditorMenu* Menu, SHLODOutliner& Outliner) override;
 		virtual FString GetDisplayString() const override;
 		virtual FTreeItemID GetID() override;
 		//~ Begin ITreeItem Interface.
diff --git a/Engine/Source/Editor/HierarchicalLODOutliner/Private/StaticMeshActorItem.cpp b/Engine/Source/Editor/HierarchicalLODOutliner/Private/StaticMeshActorItem.cpp
--- a/Engine/Source/Editor/HierarchicalLODOutliner/Private/StaticMeshActorItem.cpp
+++ b/Engine/Source/Editor/HierarchicalLODOutliner/Private/StaticMeshActorItem.cpp
@@ -4,7 +4,7 @@
 #include "GameFramework/Actor.h"
 #include "Textures/SlateIcon.h"
 #include "Framework/Commands/UIAction.h"
-#include "Framework/MultiBox/MultiBoxBuilder.h"
+#include "EditorMenuSubsystem.h"
 #include "HLODOutliner.h"
 
 #define LOCTEXT_NAMESPACE "StaticMeshActorItem"
@@ -20,11 +20,12 @@ bool HLODOutliner::FStaticMeshActorItem::CanInteract() const
 	return true;
 }
 
-void HLODOutliner::FStaticMeshActorItem::GenerateContextMenu(FMenuBuilder& MenuBuilder, SHLODOutliner& Outliner)
+void HLODOutliner::FStaticMeshActorItem::GenerateContextMenu(UEditorMenu* Menu, SHLODOutliner& Outliner)
 {
 	auto SharedOutliner = StaticCastSharedRef<SHLODOutliner>(Outliner.AsShared());
-	MenuBuilder.AddMenuEntry(LOCTEXT("RemoveSMActorFromCluster", "Remove From Cluster"), FText(), FSlateIcon(), FUIAction(FExecuteAction::CreateRaw(&Outliner, &SHLODOutliner::RemoveStaticMeshActorFromCluster)));
-	MenuBuilder.AddMenuEntry(LOCTEXT("ExcludeSMActorFromGeneration", "Exclude From Cluster Generation"), FText(), FSlateIcon(), FUIAction(FExecuteAction::CreateRaw(&Outliner, &SHLODOutliner::ExcludeFromClusterGeneration)));
+	FEditorMenuSection& Section = Menu->AddSection("Section");
+	Section.AddMenuEntry("RemoveSMActorFromCluster", LOCTEXT("RemoveSMActorFromCluster", "Remove From Cluster"), FText(), FSlateIcon(), FUIAction(FExecuteAction::CreateRaw(&Outliner, &SHLODOutliner::RemoveStaticMeshActorFromCluster)));
+	Section.AddMenuEntry("RemoveSMActorFromCluster", LOCTEXT("ExcludeSMActorFromGeneration", "Exclude From Cluster Generation"), FText(), FSlateIcon(), FUIAction(FExecuteAction::CreateRaw(&Outliner, &SHLODOutliner::ExcludeFromClusterGeneration)));
 }
 
 FString HLODOutliner::FStaticMeshActorItem::GetDisplayString() const
diff --git a/Engine/Source/Editor/HierarchicalLODOutliner/Private/StaticMeshActorItem.h b/Engine/Source/Editor/HierarchicalLODOutliner/Private/StaticMeshActorItem.h
--- a/Engine/Source/Editor/HierarchicalLODOutliner/Private/StaticMeshActorItem.h
+++ b/Engine/Source/Editor/HierarchicalLODOutliner/Private/StaticMeshActorItem.h
@@ -8,7 +8,7 @@
 #include "ITreeItem.h"
 
 class AActor;
-class FMenuBuilder;
+class UEditorMenu;
 
 namespace HLODOutliner
 {
@@ -41,7 +41,7 @@ namespace HLODOutliner
 
 		//~ Begin ITreeItem Interface.
 		virtual bool CanInteract() const override;
-		virtual void GenerateContextMenu(FMenuBuilder& MenuBuilder, SHLODOutliner& Outliner) override;
+		virtual void GenerateContextMenu(UEditorMenu* Menu, SHLODOutliner& Outliner) override;
 		virtual FString GetDisplayString() const override;
 		virtual FTreeItemID GetID() override;
 		//~ End ITreeItem Interface.
diff --git a/Engine/Source/Editor/Kismet/Kismet.Build.cs b/Engine/Source/Editor/Kismet/Kismet.Build.cs
--- a/Engine/Source/Editor/Kismet/Kismet.Build.cs
+++ b/Engine/Source/Editor/Kismet/Kismet.Build.cs
@@ -55,7 +55,8 @@ public class Kismet : ModuleRules
                 "BlueprintNativeCodeGen",
                 "UMGEditor",
                 "UMG", // for SBlueprintDiff
-                "WorkspaceMenuStructure"
+                "WorkspaceMenuStructure",
+				"EditorMenus",
             }
 			);
 
diff --git a/Engine/Source/Editor/Kismet/Private/SSCSEditor.cpp b/Engine/Source/Editor/Kismet/Private/SSCSEditor.cpp
--- a/Engine/Source/Editor/Kismet/Private/SSCSEditor.cpp
+++ b/Engine/Source/Editor/Kismet/Private/SSCSEditor.cpp
@@ -71,6 +71,10 @@
 #include "Algo/Find.h"
 #include "ActorEditorUtils.h"
 
+#include "EditorMenuSubsystem.h"
+#include "SSCSEditorMenuContext.h"
+#include "Kismet2/ComponentEditorContextMenuContex.h"
+
 #define LOCTEXT_NAMESPACE "SSCSEditor"
 
 DEFINE_LOG_CATEGORY_STATIC(LogSCSEditor, Log, All);
@@ -3917,15 +3921,12 @@ void SSCSEditor::GetSelectedItemsForContextMenu(TArray<FComponentEventConstructi
 	}
 }
 
-TSharedPtr< SWidget > SSCSEditor::CreateContextMenu()
+void SSCSEditor::PopulateContextMenu(UEditorMenu* Menu)
 {
 	TArray<FSCSEditorTreeNodePtrType> SelectedItems = SCSTreeWidget->GetSelectedItems();
 
 	if (SelectedItems.Num() > 0 || CanPasteNodes())
 	{
-		const bool CloseAfterSelection = true;
-		FMenuBuilder MenuBuilder( CloseAfterSelection, CommandList );
-
 		bool bOnlyShowPasteOption = false;
 
 		if (SelectedItems.Num() > 0)
@@ -3964,9 +3965,10 @@ TSharedPtr< SWidget > SSCSEditor::CreateContextMenu()
 
 					if (EditorMode == EComponentEditorMode::BlueprintSCS)
 					{
+						FEditorMenuSection& BlueprintSCSSection = Menu->AddSection("BlueprintSCS");
 						if (SelectedItems.Num() == 1)
 						{
-							MenuBuilder.AddMenuEntry(FGraphEditorCommands::Get().FindReferences);
+							BlueprintSCSSection.AddMenuEntry(FGraphEditorCommands::Get().FindReferences);
 						}
 
 						// Collect the classes of all selected objects
@@ -3987,16 +3989,19 @@ TSharedPtr< SWidget > SSCSEditor::CreateContextMenu()
 							// Build an event submenu if we can generate events
 							if( FBlueprintEditorUtils::CanClassGenerateEvents( SelectedClass ))
 							{
-								MenuBuilder.AddSubMenu(	LOCTEXT("AddEventSubMenu", "Add Event"), 
+								BlueprintSCSSection.AddEntry(FEditorMenuEntry::InitSubMenu(
+									Menu->GetMenuName(),
+									"AddEventSubMenu",
+									LOCTEXT("AddEventSubMenu", "Add Event"), 
 									LOCTEXT("ActtionsSubMenu_ToolTip", "Add Event"), 
 									FNewMenuDelegate::CreateStatic( &SSCSEditor::BuildMenuEventsSection,
 									GetBlueprint(), SelectedClass, FCanExecuteAction::CreateSP(this, &SSCSEditor::IsEditingAllowed),
-									FGetSelectedObjectsDelegate::CreateSP(this, &SSCSEditor::GetSelectedItemsForContextMenu)));
+									FGetSelectedObjectsDelegate::CreateSP(this, &SSCSEditor::GetSelectedItemsForContextMenu))));
 							}
 						}
 					}					
 
-					FComponentEditorUtils::FillComponentContextMenuOptions(MenuBuilder, SelectedComponents);
+					FComponentEditorUtils::FillComponentContextMenuOptions(Menu, SelectedComponents);
 				}
 			}
 		}
@@ -4007,14 +4012,42 @@ TSharedPtr< SWidget > SSCSEditor::CreateContextMenu()
 
 		if (bOnlyShowPasteOption)
 		{
-			MenuBuilder.BeginSection("PasteComponent", LOCTEXT("EditComponentHeading", "Edit") );
+			FEditorMenuSection& Section = Menu->AddSection("PasteComponent", LOCTEXT("EditComponentHeading", "Edit") );
 			{
-				MenuBuilder.AddMenuEntry( FGenericCommands::Get().Paste );
+				Section.AddMenuEntry( FGenericCommands::Get().Paste );
 			}
-			MenuBuilder.EndSection();
 		}
+	}
+}
 
-		return MenuBuilder.MakeWidget();
+void SSCSEditor::RegisterContextMenu()
+{
+	UEditorMenuSubsystem* EditorMenus = UEditorMenuSubsystem::Get();
+	if (!EditorMenus->IsMenuRegistered("Kismet.SCSEditorContextMenu"))
+	{
+		UEditorMenu* Menu = EditorMenus->RegisterMenu("Kismet.SCSEditorContextMenu");
+		Menu->AddDynamicSection("SCSEditorDynamic", FNewEditorMenuDelegate::CreateLambda([](UEditorMenu* InMenu)
+		{
+			USSCSEditorMenuContext* ContextObject = InMenu->FindContext<USSCSEditorMenuContext>();
+			if (ContextObject && ContextObject->SCSEditor.IsValid())
+			{
+				ContextObject->SCSEditor.Pin()->PopulateContextMenu(InMenu);
+			}
+		}));
+	}
+}
+
+TSharedPtr< SWidget > SSCSEditor::CreateContextMenu()
+{
+	TArray<FSCSEditorTreeNodePtrType> SelectedItems = SCSTreeWidget->GetSelectedItems();
+
+	if (SelectedItems.Num() > 0 || CanPasteNodes())
+	{
+		RegisterContextMenu();
+		USSCSEditorMenuContext* ContextObject = NewObject<USSCSEditorMenuContext>();
+		ContextObject->SCSEditor = SharedThis(this);
+		FEditorMenuContext EditorMenuContext(CommandList, TSharedPtr<FExtender>(), ContextObject);
+		return UEditorMenuSubsystem::Get()->GenerateWidget("Kismet.SCSEditorContextMenu", EditorMenuContext);
 	}
 	return TSharedPtr<SWidget>();
 }
diff --git a/Engine/Source/Editor/Kismet/Public/SSCSEditor.h b/Engine/Source/Editor/Kismet/Public/SSCSEditor.h
--- a/Engine/Source/Editor/Kismet/Public/SSCSEditor.h
+++ b/Engine/Source/Editor/Kismet/Public/SSCSEditor.h
@@ -25,6 +25,7 @@
 #include "ScopedTransaction.h"
 
 class FMenuBuilder;
+class UEditorMenu;
 class FSCSEditorTreeNode;
 class SSCSEditor;
 class UPrimitiveComponent;
@@ -37,7 +38,7 @@ using FSCSEditorActorNodePtrType = TSharedPtr<class FSCSEditorTreeNodeRootActor>
 /**
  * FSCSEditorTreeNode
  *
- * Wrapper class for component template nodes displayed in the SCS editor tree widget.
+ * Wrapper class for component template nodes displayed in the SCS (Simple Construction Script) editor tree widget.
  */
 class KISMET_API FSCSEditorTreeNode : public TSharedFromThis<FSCSEditorTreeNode>
 {
@@ -1017,6 +1018,12 @@ protected:
 	/** Called to display context menu when right clicking on the widget */
 	TSharedPtr< SWidget > CreateContextMenu();
 
+	/** Registers context menu by name for later access */
+	void RegisterContextMenu();
+
+	/** Populate context menu on the fly */
+	void PopulateContextMenu(UEditorMenu* InMenu);
+
 	/** Called when the level editor requests a component to be renamed. */
 	void OnLevelComponentRequestRename(const UActorComponent* InComponent);
 
diff --git a/Engine/Source/Editor/Kismet/Public/SSCSEditorMenuContext.h b/Engine/Source/Editor/Kismet/Public/SSCSEditorMenuContext.h
new file mode 100644
--- /dev/null
+++ b/Engine/Source/Editor/Kismet/Public/SSCSEditorMenuContext.h
@@ -0,0 +1,18 @@
+// Copyright 1998-2019 Epic Games, Inc. All Rights Reserved.
+
+#pragma once
+
+#include "CoreMinimal.h"
+
+#include "SSCSEditorMenuContext.generated.h"
+
+UCLASS()
+class USSCSEditorMenuContext : public UObject
+{
+	GENERATED_BODY()
+public:
+	
+	TWeakPtr<SSCSEditor> SCSEditor;
+
+	bool bOnlyShowPasteOption;
+};
diff --git a/Engine/Source/Editor/LevelEditor/Private/LevelEditorContextMenu.cpp b/Engine/Source/Editor/LevelEditor/Private/LevelEditorContextMenu.cpp
--- a/Engine/Source/Editor/LevelEditor/Private/LevelEditorContextMenu.cpp
+++ b/Engine/Source/Editor/LevelEditor/Private/LevelEditorContextMenu.cpp
@@ -13,6 +13,7 @@
 #include "Framework/Commands/UICommandList.h"
 #include "Framework/MultiBox/MultiBoxExtender.h"
 #include "Framework/MultiBox/MultiBoxBuilder.h"
+#include "EditorMenuSubsystem.h"
 #include "Components/ActorComponent.h"
 #include "GameFramework/Actor.h"
 #include "Kismet2/ComponentEditorUtils.h"
@@ -62,14 +63,14 @@ public:
 	 *
 	 * @param MenuBuilder	The menu to add items to
 	 */
-	static void FillSelectActorMenu( class FMenuBuilder& MenuBuilder );
+	static void FillSelectActorMenu(UEditorMenu* Menu);
 
 	/**
 	 * Fills in menu options for the actor visibility menu
 	 *
 	 * @param MenuBuilder	The menu to add items to
 	 */
-	static void FillActorVisibilityMenu( class FMenuBuilder& MenuBuilder );
+	static void FillActorVisibilityMenu(UEditorMenu* Menu);
 
 	/**
 	 * Fills in menu options for the actor level menu
@@ -78,42 +79,42 @@ public:
 	 * @param bAllInCurrentLevel	true if all selected actors are in the current level
 	 * @param MenuBuilder			The menu to add items to
 	 */
-	static void FillActorLevelMenu( class FMenuBuilder& MenuBuilder );
+	static void FillActorLevelMenu(UEditorMenu* Menu);
 
 	/**
 	 * Fills in menu options for the transform menu
 	 *
 	 * @param MenuBuilder	The menu to add items to
 	 */
-	static void FillTransformMenu( class FMenuBuilder& MenuBuilder );
+	static void FillTransformMenu(UEditorMenu* Menu);
 
 	/**
 	 * Fills in menu options for the Fill Actor menu
 	 *
 	 * @param MenuBuilder	The menu to add items to
 	 */
-	static void FillActorMenu( class FMenuBuilder& MenuBuilder );
+	static void FillActorMenu(UEditorMenu* Menu);
 
 	/**
 	 * Fills in menu options for the snap menu
 	 *
 	 * @param MenuBuilder	The menu to add items to
 	 */
-	static void FillSnapAlignMenu( class FMenuBuilder& MenuBuilder );
+	static void FillSnapAlignMenu(UEditorMenu* Menu);
 
 	/**
 	 * Fills in menu options for the pivot menu
 	 *
 	 * @param MenuBuilder	The menu to add items to
 	 */
-	static void FillPivotMenu( class FMenuBuilder& MenuBuilder );
+	static void FillPivotMenu(UEditorMenu* Menu);
 	
 	/**
 	 * Fills in menu options for the group menu
 	 *
 	 * @param MenuBuilder	The menu to add items to
 	 */
-	static void FillGroupMenu( class FMenuBuilder& MenuBuilder );
+	static void FillGroupMenu( UEditorMenu* Menu );
 
 	/**
 	 * Fills in menu options for the edit menu
@@ -121,7 +122,7 @@ public:
 	 * @param MenuBuilder	The menu to add items to
 	 * @param ContextType	The context for this editor menu
 	 */
-	static void FillEditMenu( class FMenuBuilder& MenuBuilder, LevelEditorMenuContext ContextType );
+	static void FillEditMenu(UEditorMenu* Menu);
 
 private:
 	/**
@@ -129,7 +130,7 @@ private:
 	 *
 	 * @param MenuBuilder	The menu to add items to
 	 */
-	static void FillMatineeSelectActorMenu( class FMenuBuilder& MenuBuilder );
+	static void FillMatineeSelectActorMenu(UEditorMenu* Menu);
 };
 
 FSelectedActorInfo FLevelEditorContextMenuImpl::SelectionInfo;
@@ -141,65 +142,80 @@ struct FLevelScriptEventMenuHelper
 	*
 	* @param MenuBuilder	The menu to add items to
 	*/
-	static void FillLevelBlueprintEventsMenu(class FMenuBuilder& MenuBuilder, const TArray<AActor*>& SelectedActors);
+	static void FillLevelBlueprintEventsMenu(UEditorMenu* Menu, const TArray<AActor*>& SelectedActors);
 };
 
-// NOTE: We intentionally receive a WEAK pointer here because we want to be callable by a delegate whose
-//       payload contains a weak reference to a level editor instance
-TSharedPtr< SWidget > FLevelEditorContextMenu::BuildMenuWidget( TWeakPtr< SLevelEditor > LevelEditor, LevelEditorMenuContext ContextType, TSharedPtr<FExtender> Extender )
-{
-	// Build up the menu
-	const bool bShouldCloseWindowAfterMenuSelection = true;
-	FMenuBuilder MenuBuilder(bShouldCloseWindowAfterMenuSelection, TSharedPtr<const FUICommandList>());
-
-	FillMenu(MenuBuilder, LevelEditor, ContextType, Extender);
-	
-	return MenuBuilder.MakeWidget();
-}
-
-void FLevelEditorContextMenu::FillMenu( FMenuBuilder& MenuBuilder, TWeakPtr<SLevelEditor> LevelEditor, LevelEditorMenuContext ContextType, TSharedPtr<FExtender> Extender )
+void FLevelEditorContextMenu::RegisterComponentContextMenu()
 {
-	auto LevelEditorActionsList = LevelEditor.Pin()->GetLevelEditorActions().ToSharedRef();
-	MenuBuilder.PushCommandList(LevelEditorActionsList);
+	UEditorMenuSubsystem* EditorMenus = UEditorMenuSubsystem::Get();
+	if (EditorMenus->IsMenuRegistered("LevelEditor.ComponentContextMenu"))
+	{
+		return;
+	}
 
-	if (GEditor->GetSelectedComponentCount() > 0)
+	UEditorMenu* Menu = EditorMenus->RegisterMenu("LevelEditor.ComponentContextMenu");
+	Menu->AddDynamicSection("ComponentControlDynamic", FNewEditorMenuDelegate::CreateLambda([](UEditorMenu* InMenu)
 	{
+		ULevelEditorContextMenuContext* LevelEditorContext = InMenu->FindContext<ULevelEditorContextMenuContext>();
+		if (!LevelEditorContext)
+		{
+			return;
+		}
+
 		TArray<UActorComponent*> SelectedComponents;
 		for (FSelectedEditableComponentIterator It(GEditor->GetSelectedEditableComponentIterator()); It; ++It)
 		{
 			SelectedComponents.Add(CastChecked<UActorComponent>(*It));
 		}
 
-		MenuBuilder.BeginSection("ComponentControl", LOCTEXT("ComponentControlHeading", "Component"));
 		{
+			FEditorMenuSection& Section = InMenu->AddSection("ComponentControl", LOCTEXT("ComponentControlHeading", "Component"));
+
 			AActor* OwnerActor = GEditor->GetSelectedActors()->GetTop<AActor>();
 			if(OwnerActor)
 			{
-				MenuBuilder.AddMenuEntry(
+				Section.AddMenuEntry(
 					FLevelEditorCommands::Get().SelectComponentOwnerActor,
-					NAME_None,
 					FText::Format(LOCTEXT("SelectComponentOwner", "Select Owner [{0}]"), FText::FromString(OwnerActor->GetHumanReadableName())),
 					TAttribute<FText>(),
 					FSlateIconFinder::FindIconForClass(OwnerActor->GetClass())
 				);
 			}
 
-			MenuBuilder.AddMenuEntry(FEditorViewportCommands::Get().FocusViewportToSelection);
+			Section.AddMenuEntry(FEditorViewportCommands::Get().FocusViewportToSelection);
 
 			const FVector* ClickLocation = &GEditor->ClickLocation;
 			FUIAction GoHereAction;
 			GoHereAction.ExecuteAction = FExecuteAction::CreateStatic(&FLevelEditorActionCallbacks::GoHere_Clicked, ClickLocation);
 
-			MenuBuilder.AddMenuEntry(FLevelEditorCommands::Get().GoHere);
-			MenuBuilder.AddMenuEntry(FLevelEditorCommands::Get().SnapCameraToObject);
-			MenuBuilder.AddMenuEntry(FLevelEditorCommands::Get().SnapObjectToCamera);
+			Section.AddMenuEntry(FLevelEditorCommands::Get().GoHere);
+			Section.AddMenuEntry(FLevelEditorCommands::Get().SnapCameraToObject);
+			Section.AddMenuEntry(FLevelEditorCommands::Get().SnapObjectToCamera);
 		}
-		MenuBuilder.EndSection();
 
-		FComponentEditorUtils::FillComponentContextMenuOptions(MenuBuilder, SelectedComponents);
+		FComponentEditorUtils::FillComponentContextMenuOptions(InMenu, SelectedComponents);
+	}));
+}
+
+void FLevelEditorContextMenu::RegisterActorContextMenu()
+{
+	UEditorMenuSubsystem* EditorMenus = UEditorMenuSubsystem::Get();
+	if (EditorMenus->IsMenuRegistered("LevelEditor.ActorContextMenu"))
+	{
+		return;
 	}
-	else if (GEditor->GetSelectedActorCount() > 0)
+
+	UEditorMenu* Menu = EditorMenus->RegisterMenu("LevelEditor.ActorContextMenu");
+	Menu->AddDynamicSection("ActorContextMenuDynamic", FNewEditorMenuDelegate::CreateLambda([](UEditorMenu* InMenu)
 	{
+		ULevelEditorContextMenuContext* LevelEditorContext = InMenu->FindContext<ULevelEditorContextMenuContext>();
+		if (!LevelEditorContext || !LevelEditorContext->LevelEditor.IsValid())
+		{
+			return;
+		}
+
+		TWeakPtr<SLevelEditor> LevelEditor = LevelEditorContext->LevelEditor;
+
 		// Generate information about our selection
 		TArray<AActor*> SelectedActors;
 		GEditor->GetSelectedActors()->GetSelectedObjects<AActor>(SelectedActors);
@@ -207,25 +223,6 @@ void FLevelEditorContextMenu::FillMenu( FMenuBuilder& MenuBuilder, TWeakPtr<SLev
 		FSelectedActorInfo& SelectionInfo = FLevelEditorContextMenuImpl::SelectionInfo;
 		SelectionInfo = AssetSelectionUtils::BuildSelectedActorInfo(SelectedActors);
 
-		// Get all menu extenders for this context menu from the level editor module
-		FLevelEditorModule& LevelEditorModule = FModuleManager::GetModuleChecked<FLevelEditorModule>(TEXT("LevelEditor"));
-		TArray<FLevelEditorModule::FLevelViewportMenuExtender_SelectedActors> MenuExtenderDelegates = LevelEditorModule.GetAllLevelViewportContextMenuExtenders();
-
-		TArray<TSharedPtr<FExtender>> Extenders;
-		if (Extender.IsValid())
-		{
-			Extenders.Add(Extender);
-		}
-
-		for (int32 i = 0; i < MenuExtenderDelegates.Num(); ++i)
-		{
-			if (MenuExtenderDelegates[i].IsBound())
-			{
-				Extenders.Add(MenuExtenderDelegates[i].Execute(LevelEditorActionsList, SelectedActors));
-			}
-		}
-		MenuBuilder.PushExtender(FExtender::Combine(Extenders).ToSharedRef());
-
 		// Check if current selection has any assets that can be browsed to
 		TArray< UObject* > ReferencedAssets;
 		GEditor->GetReferencedAssetsForEditorSelection(ReferencedAssets);
@@ -234,20 +231,19 @@ void FLevelEditorContextMenu::FillMenu( FMenuBuilder& MenuBuilder, TWeakPtr<SLev
 
 		if (bCanSyncToContentBrowser || ReferencedAssets.Num() > 0)
 		{
-			MenuBuilder.BeginSection("ActorAsset", LOCTEXT("AssetHeading", "Asset"));
 			{
+				FEditorMenuSection& Section = InMenu->AddSection("ActorAsset", LOCTEXT("AssetHeading", "Asset"));
 				if (bCanSyncToContentBrowser)
 				{
-					MenuBuilder.AddMenuEntry(FGlobalEditorCommonCommands::Get().FindInContentBrowser);
+					Section.AddMenuEntry(FGlobalEditorCommonCommands::Get().FindInContentBrowser);
 				}
 
 				if (ReferencedAssets.Num() == 1)
 				{
 					auto Asset = ReferencedAssets[0];
 
-					MenuBuilder.AddMenuEntry(
+					Section.AddMenuEntry(
 						FLevelEditorCommands::Get().EditAsset,
-						NAME_None,
 						FText::Format(LOCTEXT("EditAssociatedAsset", "Edit {0}"), FText::FromString(Asset->GetName())),
 						TAttribute<FText>(),
 						FSlateIconFinder::FindIconForClass(Asset->GetClass())
@@ -255,9 +251,8 @@ void FLevelEditorContextMenu::FillMenu( FMenuBuilder& MenuBuilder, TWeakPtr<SLev
 				}
 				else if (ReferencedAssets.Num() > 1)
 				{
-					MenuBuilder.AddMenuEntry(
+					Section.AddMenuEntry(
 						FLevelEditorCommands::Get().EditAssetNoConfirmMultiple,
-						NAME_None,
 						LOCTEXT("EditAssociatedAssetsMultiple", "Edit Multiple Assets"),
 						TAttribute<FText>(),
 						FSlateIcon(FEditorStyle::GetStyleSetName(), "ClassIcon.Default")
@@ -265,13 +260,12 @@ void FLevelEditorContextMenu::FillMenu( FMenuBuilder& MenuBuilder, TWeakPtr<SLev
 
 				}
 			}
-			MenuBuilder.EndSection();
 		}
 
 
-		MenuBuilder.BeginSection("ActorControl", LOCTEXT("ActorHeading", "Actor"));
 		{
-			MenuBuilder.AddMenuEntry(FEditorViewportCommands::Get().FocusViewportToSelection);
+			FEditorMenuSection& Section = InMenu->AddSection("ActorControl", LOCTEXT("ActorHeading", "Actor"));
+			Section.AddMenuEntry(FEditorViewportCommands::Get().FocusViewportToSelection);
 
 
 			const FVector* ClickLocation = &GEditor->ClickLocation;
@@ -279,9 +273,9 @@ void FLevelEditorContextMenu::FillMenu( FMenuBuilder& MenuBuilder, TWeakPtr<SLev
 			FUIAction GoHereAction;
 			GoHereAction.ExecuteAction = FExecuteAction::CreateStatic(&FLevelEditorActionCallbacks::GoHere_Clicked, ClickLocation);
 
-			MenuBuilder.AddMenuEntry(FLevelEditorCommands::Get().GoHere);
-			MenuBuilder.AddMenuEntry(FLevelEditorCommands::Get().SnapCameraToObject);
-			MenuBuilder.AddMenuEntry(FLevelEditorCommands::Get().SnapObjectToCamera);
+			Section.AddMenuEntry(FLevelEditorCommands::Get().GoHere);
+			Section.AddMenuEntry(FLevelEditorCommands::Get().SnapCameraToObject);
+			Section.AddMenuEntry(FLevelEditorCommands::Get().SnapObjectToCamera);
 
 			if (SelectedActors.Num() == 1)
 			{
@@ -296,17 +290,15 @@ void FLevelEditorContextMenu::FillMenu( FMenuBuilder& MenuBuilder, TWeakPtr<SLev
 					{
 						if (Viewport->IsSelectedActorLocked())
 						{
-							MenuBuilder.AddMenuEntry(
+							Section.AddMenuEntry(
 								Actions.EjectActorPilot,
-								NAME_None,
 								FText::Format(LOCTEXT("PilotActor_Stop", "Stop piloting '{0}'"), FText::FromString(SelectedActors[0]->GetActorLabel()))
 								);
 						}
 						else
 						{
-							MenuBuilder.AddMenuEntry(
+							Section.AddMenuEntry(
 								Actions.PilotSelectedActor,
-								NAME_None,
 								FText::Format(LOCTEXT("PilotActor", "Pilot '{0}'"), FText::FromString(SelectedActors[0]->GetActorLabel()))
 								);
 						}
@@ -314,7 +306,6 @@ void FLevelEditorContextMenu::FillMenu( FMenuBuilder& MenuBuilder, TWeakPtr<SLev
 				}
 			}
 		}
-		MenuBuilder.EndSection();
 
 		// Go to C++ Code
 		if (SelectionInfo.SelectionClass != NULL)
@@ -326,66 +317,70 @@ void FLevelEditorContextMenu::FillMenu( FMenuBuilder& MenuBuilder, TWeakPtr<SLev
 				{
 					const FString CodeFileName = FPaths::GetCleanFilename(*ClassHeaderPath);
 
-					MenuBuilder.BeginSection("ActorCode", LOCTEXT("ActorCodeHeading", "C++"));
+					FEditorMenuSection& Section = InMenu->AddSection("ActorCode", LOCTEXT("ActorCodeHeading", "C++"));
 					{
-						MenuBuilder.AddMenuEntry(FLevelEditorCommands::Get().GoToCodeForActor,
-							NAME_None,
+						Section.AddMenuEntry(FLevelEditorCommands::Get().GoToCodeForActor,
 							FText::Format(LOCTEXT("GoToCodeForActor", "Open {0}"), FText::FromString(CodeFileName)),
 							FText::Format(LOCTEXT("GoToCodeForActor_ToolTip", "Opens the header file for this actor ({0}) in a code editing program"), FText::FromString(CodeFileName)));
 					}
-					MenuBuilder.EndSection();
 				}
 			}
 
 			const FString DocumentationLink = FEditorClassUtils::GetDocumentationLink(SelectionInfo.SelectionClass);
 			if (!DocumentationLink.IsEmpty())
 			{
-				MenuBuilder.BeginSection("ActorDocumentation", LOCTEXT("ActorDocsHeading", "Documentation"));
+				FEditorMenuSection& Section = InMenu->AddSection("ActorDocumentation", LOCTEXT("ActorDocsHeading", "Documentation"));
 				{
-					MenuBuilder.AddMenuEntry(FLevelEditorCommands::Get().GoToDocsForActor,
-						NAME_None,
+					Section.AddMenuEntry(FLevelEditorCommands::Get().GoToDocsForActor,
 						LOCTEXT("GoToDocsForActor", "View Documentation"),
 						LOCTEXT("GoToDocsForActor_ToolTip", "Click to open documentation for this actor"),
 						FSlateIcon(FEditorStyle::GetStyleSetName(), "HelpIcon.Hovered"));
 				}
-				MenuBuilder.EndSection();
 			}
 		}
 
-		MenuBuilder.BeginSection("ActorSelectVisibilityLevels");
 		{
+			FEditorMenuSection& Section = InMenu->AddSection("ActorSelectVisibilityLevels");
+
 			// Add a sub-menu for "Select"
-			MenuBuilder.AddSubMenu(
+			Section.AddEntry(FEditorMenuEntry::InitSubMenu(
+				InMenu->MenuName,
+				"SelectSubMenu",
 				LOCTEXT("SelectSubMenu", "Select"),
 				LOCTEXT("SelectSubMenu_ToolTip", "Opens the actor selection menu"),
-				FNewMenuDelegate::CreateStatic(&FLevelEditorContextMenuImpl::FillSelectActorMenu));
+				FNewEditorMenuDelegate::CreateStatic(&FLevelEditorContextMenuImpl::FillSelectActorMenu)));
 
-			MenuBuilder.AddSubMenu(
+			Section.AddEntry(FEditorMenuEntry::InitSubMenu(
+				InMenu->MenuName,
+				"EditSubMenu",
 				LOCTEXT("EditSubMenu", "Edit"),
 				FText::GetEmpty(),
-				FNewMenuDelegate::CreateStatic(&FLevelEditorContextMenuImpl::FillEditMenu, ContextType));
+				FNewEditorMenuDelegate::CreateStatic(&FLevelEditorContextMenuImpl::FillEditMenu)));
 
-			MenuBuilder.AddSubMenu(
+			Section.AddEntry(FEditorMenuEntry::InitSubMenu(
+				InMenu->MenuName,
+				"VisibilitySubMenu",
 				LOCTEXT("VisibilitySubMenu", "Visibility"),
 				LOCTEXT("VisibilitySubMenu_ToolTip", "Selected actor visibility options"),
-				FNewMenuDelegate::CreateStatic(&FLevelEditorContextMenuImpl::FillActorVisibilityMenu));
+				FNewEditorMenuDelegate::CreateStatic(&FLevelEditorContextMenuImpl::FillActorVisibilityMenu)));
 
 			// Build the menu for grouping actors
-			BuildGroupMenu(MenuBuilder, SelectionInfo);
+			BuildGroupMenu(InMenu, SelectionInfo);
 
-			MenuBuilder.AddSubMenu(
+			Section.AddEntry(FEditorMenuEntry::InitSubMenu(
+				InMenu->MenuName,
+				"LevelSubMenu",
 				LOCTEXT("LevelSubMenu", "Level"),
 				LOCTEXT("LevelSubMenu_ToolTip", "Options for interacting with this actor's level"),
-				FNewMenuDelegate::CreateStatic(&FLevelEditorContextMenuImpl::FillActorLevelMenu));
+				FNewEditorMenuDelegate::CreateStatic(&FLevelEditorContextMenuImpl::FillActorLevelMenu)));
 		}
-		MenuBuilder.EndSection();
 
-		if (ContextType == LevelEditorMenuContext::Viewport)
+		if (LevelEditorContext->ContextType == ELevelEditorMenuContext::Viewport)
 		{
-			LevelEditorCreateActorMenu::FillAddReplaceViewportContextMenuSections(MenuBuilder);
+			LevelEditorCreateActorMenu::FillAddReplaceViewportContextMenuSections(InMenu);
 
-			MenuBuilder.AddMenuEntry(FLevelEditorCommands::Get().OpenMergeActor,
-					NAME_None,
+			FEditorMenuSection& Section = InMenu->AddSection("OpenMergeActor");
+			Section.AddMenuEntry(FLevelEditorCommands::Get().OpenMergeActor,
 					LOCTEXT("OpenMergeActor", "Merge Actors"),
 					LOCTEXT("OpenMergeActor_ToolTip", "Click to open the Merge Actor panel"));
 		}
@@ -394,68 +389,166 @@ void FLevelEditorContextMenu::FillMenu( FMenuBuilder& MenuBuilder, TWeakPtr<SLev
 		{
 			if (SelectionInfo.NumSelected > 0)
 			{
-				MenuBuilder.BeginSection("Simulation", NSLOCTEXT("LevelViewportContextMenu", "SimulationHeading", "Simulation"));
+				FEditorMenuSection& Section = InMenu->AddSection("Simulation", NSLOCTEXT("LevelViewportContextMenu", "SimulationHeading", "Simulation"));
 				{
-					MenuBuilder.AddMenuEntry(FLevelEditorCommands::Get().KeepSimulationChanges);
+					Section.AddMenuEntry(FLevelEditorCommands::Get().KeepSimulationChanges);
 				}
-				MenuBuilder.EndSection();
 			}
 		}
 
-		MenuBuilder.BeginSection("LevelViewportAttach");
 		{
+			FEditorMenuSection& Section = InMenu->AddSection("LevelViewportAttach");
+
 			// Only display the attach menu if we have actors selected
 			if (GEditor->GetSelectedActorCount())
 			{
 				if (SelectionInfo.bHaveAttachedActor)
 				{
-					MenuBuilder.AddMenuEntry(FLevelEditorCommands::Get().DetachFromParent);
+					Section.AddMenuEntry(FLevelEditorCommands::Get().DetachFromParent);
 				}
 
-				MenuBuilder.AddSubMenu(
+				Section.AddEntry(FEditorMenuEntry::InitSubMenu(
+					InMenu->GetMenuName(),
+					"ActorAttachToSubMenu",
 					LOCTEXT("ActorAttachToSubMenu", "Attach To"),
 					LOCTEXT("ActorAttachToSubMenu_ToolTip", "Attach Actor as child"),
-					FNewMenuDelegate::CreateStatic(&FLevelEditorContextMenuImpl::FillActorMenu));
+					FNewEditorMenuDelegate::CreateStatic(&FLevelEditorContextMenuImpl::FillActorMenu)));
 			}
 
 			// Add a heading for "Movement" if an actor is selected
 			if (GEditor->GetSelectedActorIterator())
 			{
 				// Add a sub-menu for "Transform"
-				MenuBuilder.AddSubMenu(
+				Section.AddEntry(FEditorMenuEntry::InitSubMenu(
+					InMenu->GetMenuName(),
+					"TransformSubMenu",
 					LOCTEXT("TransformSubMenu", "Transform"),
 					LOCTEXT("TransformSubMenu_ToolTip", "Actor transform utils"),
-					FNewMenuDelegate::CreateStatic(&FLevelEditorContextMenuImpl::FillTransformMenu));
+					FNewEditorMenuDelegate::CreateStatic(&FLevelEditorContextMenuImpl::FillTransformMenu)));
 			}
 
 			// Add a sub-menu for "Pivot"
-			MenuBuilder.AddSubMenu(
+			Section.AddEntry(FEditorMenuEntry::InitSubMenu(
+				InMenu->MenuName,
+				"PivotSubMenu",
 				LOCTEXT("PivotSubMenu", "Pivot"),
 				LOCTEXT("PivotSubMenu_ToolTip", "Actor pivoting utils"),
-				FNewMenuDelegate::CreateStatic(&FLevelEditorContextMenuImpl::FillPivotMenu));
+				FNewEditorMenuDelegate::CreateStatic(&FLevelEditorContextMenuImpl::FillPivotMenu)));
 		}
-		MenuBuilder.EndSection();
 
-		FLevelScriptEventMenuHelper::FillLevelBlueprintEventsMenu(MenuBuilder, SelectedActors);
+		FLevelScriptEventMenuHelper::FillLevelBlueprintEventsMenu(InMenu, SelectedActors);
+	}));
+}
 
-		MenuBuilder.PopExtender();
-	}	
-	else if (ContextType == LevelEditorMenuContext::SceneOutliner)
+void FLevelEditorContextMenu::RegisterSceneOutlinerContextMenu()
+{
+	UEditorMenuSubsystem* EditorMenus = UEditorMenuSubsystem::Get();
+	if (EditorMenus->IsMenuRegistered("LevelEditor.SceneOutlinerContextMenu"))
 	{
-		TWeakPtr<ISceneOutliner> SceneOutlinerPtr = LevelEditor.Pin()->GetSceneOutliner();
-		if (SceneOutlinerPtr.IsValid())
+		return;
+	}
+
+	UEditorMenu* Menu = EditorMenus->RegisterMenu("LevelEditor.SceneOutlinerContextMenu");
+	Menu->AddDynamicSection("SelectVisibilityLevels", FNewEditorMenuDelegate::CreateLambda([](UEditorMenu* InMenu)
+	{
+		if (ULevelEditorContextMenuContext* LevelEditorContext = InMenu->FindContext<ULevelEditorContextMenuContext>())
 		{
-			MenuBuilder.BeginSection("SelectVisibilityLevels");
+			TWeakPtr<ISceneOutliner> SceneOutlinerPtr = LevelEditorContext->LevelEditor.Pin()->GetSceneOutliner();
+			if (SceneOutlinerPtr.IsValid())
 			{
-				MenuBuilder.AddSubMenu(
+				FEditorMenuSection& Section = InMenu->AddSection("SelectVisibilityLevels");
+				Section.AddEntry(FEditorMenuEntry::InitSubMenu(
+					InMenu->MenuName,
+					"EditSubMenu",
 					LOCTEXT("EditSubMenu", "Edit"),
 					FText::GetEmpty(),
-					FNewMenuDelegate::CreateStatic(&FLevelEditorContextMenuImpl::FillEditMenu, ContextType));
+					FNewEditorMenuDelegate::CreateStatic(&FLevelEditorContextMenuImpl::FillEditMenu)
+				));
+			}
+		}
+	}));
+}
+
+FName FLevelEditorContextMenu::GetContextMenuName(ELevelEditorMenuContext ContextType)
+{
+	if (GEditor->GetSelectedComponentCount() > 0)
+	{
+		return "LevelEditor.ComponentContextMenu";
+	}
+	else if (GEditor->GetSelectedActorCount() > 0)
+	{
+		return "LevelEditor.ActorContextMenu";
+	}
+	else if (ContextType == ELevelEditorMenuContext::SceneOutliner)
+	{
+		return "LevelEditor.SceneOutlinerContextMenu";
+	}
+
+	return NAME_None;
+}
+
+FName FLevelEditorContextMenu::InitMenuContext(FEditorMenuContext& Context, TWeakPtr<SLevelEditor> LevelEditor, ELevelEditorMenuContext ContextType)
+{
+	RegisterComponentContextMenu();
+	RegisterActorContextMenu();
+	RegisterSceneOutlinerContextMenu();
+
+	Context.AppendCommandList(LevelEditor.Pin()->GetLevelEditorActions());
+
+	ULevelEditorContextMenuContext* ContextObject = NewObject<ULevelEditorContextMenuContext>();
+	ContextObject->LevelEditor = LevelEditor;
+	ContextObject->ContextType = ContextType;
+	for (FSelectedEditableComponentIterator It(GEditor->GetSelectedEditableComponentIterator()); It; ++It)
+	{
+		ContextObject->SelectedComponents.Add(CastChecked<UActorComponent>(*It));
+	}
+	Context.AddObject(ContextObject);
+
+	if (GEditor->GetSelectedComponentCount() == 0 && GEditor->GetSelectedActorCount() > 0)
+	{
+		TArray<AActor*> SelectedActors;
+		GEditor->GetSelectedActors()->GetSelectedObjects<AActor>(SelectedActors);
+
+		// Get all menu extenders for this context menu from the level editor module
+		FLevelEditorModule& LevelEditorModule = FModuleManager::GetModuleChecked<FLevelEditorModule>(TEXT("LevelEditor"));
+		TArray<FLevelEditorModule::FLevelViewportMenuExtender_SelectedActors> MenuExtenderDelegates = LevelEditorModule.GetAllLevelViewportContextMenuExtenders();
+		TSharedPtr<FUICommandList> LevelEditorActionsList = LevelEditor.Pin()->GetLevelEditorActions();
+		TArray<TSharedPtr<FExtender>> Extenders;
+		for (int32 i = 0; i < MenuExtenderDelegates.Num(); ++i)
+		{
+			if (MenuExtenderDelegates[i].IsBound())
+			{
+				Extenders.Add(MenuExtenderDelegates[i].Execute(LevelEditorActionsList.ToSharedRef(), SelectedActors));
 			}
 		}
+
+		if (Extenders.Num() > 0)
+		{
+			Context.AddExtender(FExtender::Combine(Extenders));
+		}
 	}
 
-	MenuBuilder.PopCommandList();
+	return GetContextMenuName(ContextType);
+}
+
+UEditorMenu* FLevelEditorContextMenu::GenerateMenu(TWeakPtr<SLevelEditor> LevelEditor, ELevelEditorMenuContext ContextType, TSharedPtr<FExtender> Extender)
+{
+	FEditorMenuContext Context;
+	if (Extender.IsValid())
+	{
+		Context.AddExtender(Extender);
+	}
+
+	FName ContextMenuName = InitMenuContext(Context, LevelEditor, ContextType);
+	return UEditorMenuSubsystem::Get()->GenerateMenu(ContextMenuName, Context);
+}
+
+// NOTE: We intentionally receive a WEAK pointer here because we want to be callable by a delegate whose
+//       payload contains a weak reference to a level editor instance
+TSharedPtr< SWidget > FLevelEditorContextMenu::BuildMenuWidget(TWeakPtr< SLevelEditor > LevelEditor, ELevelEditorMenuContext ContextType, TSharedPtr<FExtender> Extender)
+{
+	UEditorMenu* Menu = GenerateMenu(LevelEditor, ContextType, Extender);
+	return UEditorMenuSubsystem::Get()->GenerateWidget(Menu);
 }
 
 namespace EViewOptionType
@@ -559,7 +652,7 @@ void FLevelEditorContextMenu::SummonViewOptionMenu( const TSharedRef< SLevelEdit
 	BuildViewOptionMenu(LevelEditor, MakeViewOptionWidget(LevelEditor, bShouldCloseWindowAfterMenuSelection, ViewOptionType), MouseCursorLocation);
 }
 
-void FLevelEditorContextMenu::SummonMenu( const TSharedRef< SLevelEditor >& LevelEditor, LevelEditorMenuContext ContextType )
+void FLevelEditorContextMenu::SummonMenu(const TSharedRef< SLevelEditor >& LevelEditor, ELevelEditorMenuContext ContextType)
 {
 	struct Local
 	{
@@ -610,10 +703,12 @@ FSlateColor InvertOnHover( const TWeakPtr< SWidget > WidgetPtr )
 	return FSlateColor::UseForeground();
 }
 
-void FLevelEditorContextMenu::BuildGroupMenu( FMenuBuilder& MenuBuilder, const FSelectedActorInfo& SelectedActorInfo )
+void FLevelEditorContextMenu::BuildGroupMenu(UEditorMenu* Menu, const FSelectedActorInfo& SelectedActorInfo)
 {
 	if( UActorGroupingUtils::IsGroupingActive() )
 	{
+		FEditorMenuSection& Section = Menu->AddSection("GroupMenu");
+
 		// Whether or not we added a grouping sub-menu
 		bool bNeedGroupSubMenu = SelectedActorInfo.bHaveSelectedLockedGroup || SelectedActorInfo.bHaveSelectedUnlockedGroup;
 
@@ -623,7 +718,7 @@ void FLevelEditorContextMenu::BuildGroupMenu( FMenuBuilder& MenuBuilder, const F
 			if( !SelectedActorInfo.bHaveSelectedLockedGroup && !SelectedActorInfo.bHaveSelectedUnlockedGroup )
 			{
 				// Only one menu entry needed so dont use a sub-menu
-				MenuBuilder.AddMenuEntry( FLevelEditorCommands::Get().RegroupActors, NAME_None, FLevelEditorCommands::Get().GroupActors->GetLabel(), FLevelEditorCommands::Get().GroupActors->GetDescription() );
+				Section.AddMenuEntry(FLevelEditorCommands::Get().RegroupActors, FLevelEditorCommands::Get().GroupActors->GetLabel(), FLevelEditorCommands::Get().GroupActors->GetDescription());
 			}
 			else
 			{
@@ -634,142 +729,135 @@ void FLevelEditorContextMenu::BuildGroupMenu( FMenuBuilder& MenuBuilder, const F
 		
 		if( bNeedGroupSubMenu )
 		{
-			MenuBuilder.AddSubMenu( 
+			Section.AddEntry(FEditorMenuEntry::InitSubMenu(
+				Menu->GetMenuName(),
+				"GroupMenu",
 				LOCTEXT("GroupMenu", "Groups"),
 				LOCTEXT("GroupMenu_ToolTip", "Opens the actor grouping menu"),
-				FNewMenuDelegate::CreateStatic( &FLevelEditorContextMenuImpl::FillGroupMenu ) );
+				FNewEditorMenuDelegate::CreateStatic( &FLevelEditorContextMenuImpl::FillGroupMenu)));
 		}
 	}
 }
 
-void FLevelEditorContextMenuImpl::FillSelectActorMenu( FMenuBuilder& MenuBuilder )
+void FLevelEditorContextMenuImpl::FillSelectActorMenu(UEditorMenu* Menu)
 {
 	FText SelectAllActorStr = FText::Format( LOCTEXT("SelectActorsOfSameClass", "Select All {0}(s)"), FText::FromString( SelectionInfo.SelectionStr ) );
 	int32 NumSelectedSurfaces = AssetSelectionUtils::GetNumSelectedSurfaces( SelectionInfo.SharedWorld );
 
-	MenuBuilder.BeginSection("SelectActorGeneral", LOCTEXT("SelectAnyHeading", "General") );
 	{
-		MenuBuilder.AddMenuEntry( FGenericCommands::Get().SelectAll, NAME_None, TAttribute<FText>(), LOCTEXT("SelectAll_ToolTip", "Selects all actors") );
-		MenuBuilder.AddMenuEntry( FLevelEditorCommands::Get().SelectNone );
-		MenuBuilder.AddMenuEntry( FLevelEditorCommands::Get().InvertSelection );
+		FEditorMenuSection& Section = Menu->AddSection("SelectActorGeneral", LOCTEXT("SelectAnyHeading", "General"));
+		Section.AddMenuEntry( FGenericCommands::Get().SelectAll, TAttribute<FText>(), LOCTEXT("SelectAll_ToolTip", "Selects all actors") );
+		Section.AddMenuEntry( FLevelEditorCommands::Get().SelectNone );
+		Section.AddMenuEntry( FLevelEditorCommands::Get().InvertSelection );
 	}
-	MenuBuilder.EndSection();
 
 	if( !SelectionInfo.bHaveBrush && SelectionInfo.bAllSelectedActorsOfSameType && SelectionInfo.SelectionStr.Len() != 0 )
 	{
 		// These menu options appear if only if all the actors are the same type and we aren't selecting brush
-		MenuBuilder.AddMenuEntry( FLevelEditorCommands::Get().SelectAllActorsOfSameClass, NAME_None, SelectAllActorStr );
+		FEditorMenuSection& Section = Menu->AddSection("SelectAllActorsOfSameClass");
+		Section.AddMenuEntry(FLevelEditorCommands::Get().SelectAllActorsOfSameClass, SelectAllActorStr);
 	}
 
-	MenuBuilder.BeginSection("SelectActorHierarchy", LOCTEXT("SelectHierarchyHeading", "Hierarchy") );
 	{
-		MenuBuilder.AddMenuEntry( FLevelEditorCommands::Get().SelectImmediateChildren );
-		MenuBuilder.AddMenuEntry( FLevelEditorCommands::Get().SelectAllDescendants );
+		FEditorMenuSection& Section = Menu->AddSection("SelectActorHierarchy", LOCTEXT("SelectHierarchyHeading", "Hierarchy"));
+		Section.AddMenuEntry( FLevelEditorCommands::Get().SelectImmediateChildren );
+		Section.AddMenuEntry( FLevelEditorCommands::Get().SelectAllDescendants );
 	}
-	MenuBuilder.EndSection();
 
 	// Add brush commands when we have a brush or any surfaces selected
-	MenuBuilder.BeginSection("SelectBSP", LOCTEXT("SelectBSPHeading", "BSP") );
 	{
+		FEditorMenuSection& Section = Menu->AddSection("SelectBSP", LOCTEXT("SelectBSPHeading", "BSP"));
 		if( SelectionInfo.bHaveBrush || NumSelectedSurfaces > 0 )
 		{
 			if( SelectionInfo.bAllSelectedAreBrushes )
 			{
-				MenuBuilder.AddMenuEntry( FLevelEditorCommands::Get().SelectAllActorsOfSameClass, NAME_None, SelectAllActorStr );
+				Section.AddMenuEntry( FLevelEditorCommands::Get().SelectAllActorsOfSameClass, SelectAllActorStr);
 			}
 		}
 
-		MenuBuilder.AddMenuEntry( FLevelEditorCommands::Get().SelectAllAddditiveBrushes );
-		MenuBuilder.AddMenuEntry( FLevelEditorCommands::Get().SelectAllSubtractiveBrushes );
-		MenuBuilder.AddMenuEntry( FLevelEditorCommands::Get().SelectAllSurfaces );
+		Section.AddMenuEntry(FLevelEditorCommands::Get().SelectAllAddditiveBrushes);
+		Section.AddMenuEntry(FLevelEditorCommands::Get().SelectAllSubtractiveBrushes);
+		Section.AddMenuEntry(FLevelEditorCommands::Get().SelectAllSurfaces);
 	}
-	MenuBuilder.EndSection();
 
 	if( SelectionInfo.NumSelected > 0 || NumSelectedSurfaces > 0 )
 	{
 		// If any actors are selected add lights selection options
-		MenuBuilder.BeginSection("SelectLights", LOCTEXT("SelectLightHeading", "Lights") );
 		{
-			MenuBuilder.AddMenuEntry( FLevelEditorCommands::Get().SelectRelevantLights );
+			FEditorMenuSection& Section = Menu->AddSection("SelectLights", LOCTEXT("SelectLightHeading", "Lights"));
+			Section.AddMenuEntry(FLevelEditorCommands::Get().SelectRelevantLights);
 
 			if ( SelectionInfo.bHaveLight )
 			{
-				MenuBuilder.AddMenuEntry( FLevelEditorCommands::Get().SelectAllLights );
-				MenuBuilder.AddMenuEntry( FLevelEditorCommands::Get().SelectStationaryLightsExceedingOverlap );
+				Section.AddMenuEntry(FLevelEditorCommands::Get().SelectAllLights);
+				Section.AddMenuEntry(FLevelEditorCommands::Get().SelectStationaryLightsExceedingOverlap);
 			}
 		}
-		MenuBuilder.EndSection();
 
 		if( SelectionInfo.bHaveStaticMesh )
 		{
 			// if any static meshes are selected allow selecting actors using the same mesh
-			MenuBuilder.BeginSection("SelectMeshes", LOCTEXT("SelectStaticMeshHeading", "Static Meshes") );
 			{
-				MenuBuilder.AddMenuEntry( FLevelEditorCommands::Get().SelectStaticMeshesOfSameClass, NAME_None, LOCTEXT("SelectStaticMeshesOfSameClass_Menu", "Select Matching (Selected Classes)") );
-				MenuBuilder.AddMenuEntry( FLevelEditorCommands::Get().SelectStaticMeshesAllClasses, NAME_None, LOCTEXT("SelectStaticMeshesAllClasses_Menu", "Select Matching (All Classes)") );
+				FEditorMenuSection& Section = Menu->AddSection("SelectMeshes", LOCTEXT("SelectStaticMeshHeading", "Static Meshes"));
+				Section.AddMenuEntry(FLevelEditorCommands::Get().SelectStaticMeshesOfSameClass, LOCTEXT("SelectStaticMeshesOfSameClass_Menu", "Select Matching (Selected Classes)"));
+				Section.AddMenuEntry(FLevelEditorCommands::Get().SelectStaticMeshesAllClasses, LOCTEXT("SelectStaticMeshesAllClasses_Menu", "Select Matching (All Classes)"));
 			}
-			MenuBuilder.EndSection();
 
 			if (SelectionInfo.NumSelected == 1)
 			{
-				MenuBuilder.BeginSection("SelectHLODCluster", LOCTEXT("SelectHLODClusterHeading", "Hierachical LODs"));
 				{
-					MenuBuilder.AddMenuEntry(FLevelEditorCommands::Get().SelectOwningHierarchicalLODCluster, NAME_None, LOCTEXT("SelectOwningHierarchicalLODCluster_Menu", "Select Owning HierarchicalLODCluster"));
+					FEditorMenuSection& Section = Menu->AddSection("SelectHLODCluster", LOCTEXT("SelectHLODClusterHeading", "Hierachical LODs"));
+					Section.AddMenuEntry(FLevelEditorCommands::Get().SelectOwningHierarchicalLODCluster, LOCTEXT("SelectOwningHierarchicalLODCluster_Menu", "Select Owning HierarchicalLODCluster"));
 				}
-				MenuBuilder.EndSection();
 			}			
 		}
 
 		if( SelectionInfo.bHavePawn || SelectionInfo.bHaveSkeletalMesh )
 		{
 			// if any skeletal meshes are selected allow selecting actors using the same mesh
-			MenuBuilder.BeginSection("SelectSkeletalMeshes", LOCTEXT("SelectSkeletalMeshHeading", "Skeletal Meshes") );
 			{
-				MenuBuilder.AddMenuEntry( FLevelEditorCommands::Get().SelectSkeletalMeshesOfSameClass );
-				MenuBuilder.AddMenuEntry( FLevelEditorCommands::Get().SelectSkeletalMeshesAllClasses );
+				FEditorMenuSection& Section = Menu->AddSection("SelectSkeletalMeshes", LOCTEXT("SelectSkeletalMeshHeading", "Skeletal Meshes"));
+				Section.AddMenuEntry(FLevelEditorCommands::Get().SelectSkeletalMeshesOfSameClass);
+				Section.AddMenuEntry(FLevelEditorCommands::Get().SelectSkeletalMeshesAllClasses);
 			}
-			MenuBuilder.EndSection();
 		}
 
 		if( SelectionInfo.bHaveEmitter )
 		{
-			MenuBuilder.BeginSection("SelectEmitters", LOCTEXT("SelectEmitterHeading", "Emitters") );
 			{
-				MenuBuilder.AddMenuEntry( FLevelEditorCommands::Get().SelectMatchingEmitter );
+				FEditorMenuSection& Section = Menu->AddSection("SelectEmitters", LOCTEXT("SelectEmitterHeading", "Emitters"));
+				Section.AddMenuEntry(FLevelEditorCommands::Get().SelectMatchingEmitter);
 			}
-			MenuBuilder.EndSection();
 		}
 	}
 
 	if( SelectionInfo.bHaveBrush || SelectionInfo.NumSelected > 0 )
 	{
-		MenuBuilder.BeginSection("SelectMaterial", LOCTEXT("SelectMaterialHeading", "Materials") );
+		FEditorMenuSection& Section = Menu->AddSection("SelectMaterial", LOCTEXT("SelectMaterialHeading", "Materials"));
 		{
-			MenuBuilder.AddMenuEntry( FLevelEditorCommands::Get().SelectAllWithSameMaterial );
+			Section.AddMenuEntry(FLevelEditorCommands::Get().SelectAllWithSameMaterial);
 		}
-		MenuBuilder.EndSection();
 	}
 
 	// Add geometry collection commands
 	if (FModuleManager::Get().IsModuleLoaded("GeometryCollectionEditor"))
 	{
-		MenuBuilder.BeginSection("SelectBones", LOCTEXT("GeometryCollectionHeading", "Geometry Collection"));
+		FEditorMenuSection& Section = Menu->AddSection("SelectBones", LOCTEXT("GeometryCollectionHeading", "Geometry Collection"));
 		{
-			MenuBuilder.AddMenuEntry(FLevelEditorCommands::Get().GeometryCollectionSelectAllGeometry);
-			MenuBuilder.AddMenuEntry(FLevelEditorCommands::Get().GeometryCollectionSelectNone);
-			MenuBuilder.AddMenuEntry(FLevelEditorCommands::Get().GeometryCollectionSelectInverseGeometry);
+			Section.AddMenuEntry(FLevelEditorCommands::Get().GeometryCollectionSelectAllGeometry);
+			Section.AddMenuEntry(FLevelEditorCommands::Get().GeometryCollectionSelectNone);
+			Section.AddMenuEntry(FLevelEditorCommands::Get().GeometryCollectionSelectInverseGeometry);
 		}
-		MenuBuilder.EndSection();
 	}
 
 	// build matinee related selection menu
-	FillMatineeSelectActorMenu( MenuBuilder );
+	FillMatineeSelectActorMenu( Menu );
 }
 
-void FLevelEditorContextMenuImpl::FillMatineeSelectActorMenu( FMenuBuilder& MenuBuilder )
+void FLevelEditorContextMenuImpl::FillMatineeSelectActorMenu(UEditorMenu* Menu)
 {
-	MenuBuilder.BeginSection("SelectMatinee", LOCTEXT("SelectMatineeHeading", "Matinee") );
 	{
+		FEditorMenuSection& Section = Menu->AddSection("SelectMatinee", LOCTEXT("SelectMatineeHeading", "Matinee"));
 		// show list of Matinee Actors that controls this actor
 		// this is ugly but we don't have good way of knowing which Matinee actor controls me
 		// in the future this can be cached to TMap somewhere and use that list
@@ -824,7 +912,7 @@ void FLevelEditorContextMenuImpl::FillMatineeSelectActorMenu( FMenuBuilder& Menu
 						const FText Text = FText::Format( LOCTEXT("SelectMatineeActor", "Select {0}"), FText::FromString( CurMatineeActor->GetName() ) );
 
 						FUIAction CurMatineeActorAction( FExecuteAction::CreateStatic( &FLevelEditorActionCallbacks::OnSelectMatineeActor, CurMatineeActor ) );
-						MenuBuilder.AddMenuEntry( Text, Text, FSlateIcon(), CurMatineeActorAction );
+						Section.AddMenuEntry(NAME_None, Text, Text, FSlateIcon(), CurMatineeActorAction);
 
 						// if matinee is opened, and if that is CurMatineeActor, show option to go to group
 						if( GLevelEditorModeTools().IsModeActive( FBuiltinEditorModes::EM_InterpEdit ) )
@@ -834,7 +922,7 @@ void FLevelEditorContextMenuImpl::FillMatineeSelectActorMenu( FMenuBuilder& Menu
 							if ( InterpEditMode && InterpEditMode->MatineeActor == CurMatineeActor )
 							{
 								FUIAction SelectedActorAction( FExecuteAction::CreateStatic( &FLevelEditorActionCallbacks::OnSelectMatineeGroup, SelectedActor ) );
-								MenuBuilder.AddMenuEntry( LOCTEXT("SelectMatineeGroupForActorMenuTitle", "Select Matinee Group For This Actor"), LOCTEXT("SelectMatineeGroupForActorMenuTooltip", "Selects matinee group controlling this actor"), FSlateIcon(), SelectedActorAction );
+								Section.AddMenuEntry(NAME_None, LOCTEXT("SelectMatineeGroupForActorMenuTitle", "Select Matinee Group For This Actor"), LOCTEXT("SelectMatineeGroupForActorMenuTooltip", "Selects matinee group controlling this actor"), FSlateIcon(), SelectedActorAction);
 							}
 						}
 					}
@@ -845,109 +933,103 @@ void FLevelEditorContextMenuImpl::FillMatineeSelectActorMenu( FMenuBuilder& Menu
 		// if this class is Matinee Actor, add option to allow select all controlled actors
 		if ( SelectionInfo.bHaveMatinee )
 		{
-			MenuBuilder.AddMenuEntry( FLevelEditorCommands::Get().SelectAllActorsControlledByMatinee );	
+			Section.AddMenuEntry(FLevelEditorCommands::Get().SelectAllActorsControlledByMatinee);	
 		}
 	}
-	MenuBuilder.EndSection();
 }
 
-void FLevelEditorContextMenuImpl::FillActorVisibilityMenu( FMenuBuilder& MenuBuilder )
+void FLevelEditorContextMenuImpl::FillActorVisibilityMenu(UEditorMenu* Menu)
 {
-	MenuBuilder.BeginSection("VisibilitySelected");
 	{
+		FEditorMenuSection& Section = Menu->AddSection("VisibilitySelected");
 		// Show 'Show Selected' only if the selection has any hidden actors
 		if ( SelectionInfo.bHaveHidden )
 		{
-			MenuBuilder.AddMenuEntry( FLevelEditorCommands::Get().ShowSelected );
+			Section.AddMenuEntry(FLevelEditorCommands::Get().ShowSelected);
 		}
-		MenuBuilder.AddMenuEntry( FLevelEditorCommands::Get().HideSelected );
+		Section.AddMenuEntry(FLevelEditorCommands::Get().HideSelected);
 	}
-	MenuBuilder.EndSection();
 
-	MenuBuilder.BeginSection("VisibilityAll");
 	{
-		MenuBuilder.AddMenuEntry( FLevelEditorCommands::Get().ShowSelectedOnly );
-		MenuBuilder.AddMenuEntry( FLevelEditorCommands::Get().ShowAll );
+		FEditorMenuSection& Section = Menu->AddSection("VisibilityAll");
+		Section.AddMenuEntry(FLevelEditorCommands::Get().ShowSelectedOnly);
+		Section.AddMenuEntry(FLevelEditorCommands::Get().ShowAll);
 	}
-	MenuBuilder.EndSection();
 
-	MenuBuilder.BeginSection("VisibilityStartup");
 	{
-		MenuBuilder.AddMenuEntry( FLevelEditorCommands::Get().ShowAllStartup );
-		MenuBuilder.AddMenuEntry( FLevelEditorCommands::Get().ShowSelectedStartup );
-		MenuBuilder.AddMenuEntry( FLevelEditorCommands::Get().HideSelectedStartup );
+		FEditorMenuSection& Section = Menu->AddSection("VisibilityStartup");
+		Section.AddMenuEntry(FLevelEditorCommands::Get().ShowAllStartup);
+		Section.AddMenuEntry(FLevelEditorCommands::Get().ShowSelectedStartup);
+		Section.AddMenuEntry(FLevelEditorCommands::Get().HideSelectedStartup);
 	}
 }
 
-void FLevelEditorContextMenuImpl::FillActorLevelMenu( FMenuBuilder& MenuBuilder )
+void FLevelEditorContextMenuImpl::FillActorLevelMenu(UEditorMenu* Menu)
 {
-	MenuBuilder.BeginSection("ActorLevel", LOCTEXT("ActorLevel", "Actor Level"));
 	{
+		FEditorMenuSection& Section = Menu->AddSection("ActorLevel", LOCTEXT("ActorLevel", "Actor Level"));
 		if( SelectionInfo.SharedLevel && SelectionInfo.SharedWorld && SelectionInfo.SharedWorld->GetCurrentLevel() != SelectionInfo.SharedLevel )
 		{
 			// All actors are in the same level and that level is not the current level 
 			// so add a menu entry to make the shared level current
 
 			FText MakeCurrentLevelText = FText::Format( LOCTEXT("MakeCurrentLevelMenu", "Make Current Level: {0}"), FText::FromString( SelectionInfo.SharedLevel->GetOutermost()->GetName() ) );
-			MenuBuilder.AddMenuEntry( FLevelEditorCommands::Get().MakeActorLevelCurrent, NAME_None, MakeCurrentLevelText );
+			Section.AddMenuEntry(FLevelEditorCommands::Get().MakeActorLevelCurrent, MakeCurrentLevelText);
 		}
 
 		if( !SelectionInfo.bAllSelectedActorsBelongToCurrentLevel )
 		{
 			// Only show this menu entry if any actors are not in the current level
-			MenuBuilder.AddMenuEntry( FLevelEditorCommands::Get().MoveSelectedToCurrentLevel );
+			Section.AddMenuEntry(FLevelEditorCommands::Get().MoveSelectedToCurrentLevel);
 		}
 
-		MenuBuilder.AddMenuEntry(FLevelEditorCommands::Get().FindActorLevelInContentBrowser);
+		Section.AddMenuEntry(FLevelEditorCommands::Get().FindActorLevelInContentBrowser);
 	}
-	MenuBuilder.EndSection();
 
-	MenuBuilder.BeginSection("LevelBlueprint", LOCTEXT("LevelBlueprint", "Level Blueprint"));
 	{
-		MenuBuilder.AddMenuEntry(FLevelEditorCommands::Get().FindActorInLevelScript);
+		FEditorMenuSection& Section = Menu->AddSection("LevelBlueprint", LOCTEXT("LevelBlueprint", "Level Blueprint"));
+		Section.AddMenuEntry(FLevelEditorCommands::Get().FindActorInLevelScript);
 	}
-	MenuBuilder.EndSection();
 
-	MenuBuilder.BeginSection("LevelBrowser", LOCTEXT("LevelBrowser", "Level Browser"));
 	{
-		MenuBuilder.AddMenuEntry(FLevelEditorCommands::Get().FindLevelsInLevelBrowser);
-		MenuBuilder.AddMenuEntry(FLevelEditorCommands::Get().AddLevelsToSelection);
-		MenuBuilder.AddMenuEntry(FLevelEditorCommands::Get().RemoveLevelsFromSelection);
+		FEditorMenuSection& Section = Menu->AddSection("LevelBrowser", LOCTEXT("LevelBrowser", "Level Browser"));
+		Section.AddMenuEntry(FLevelEditorCommands::Get().FindLevelsInLevelBrowser);
+		Section.AddMenuEntry(FLevelEditorCommands::Get().AddLevelsToSelection);
+		Section.AddMenuEntry(FLevelEditorCommands::Get().RemoveLevelsFromSelection);
 	}
-	MenuBuilder.EndSection();
 }
 
 
-void FLevelEditorContextMenuImpl::FillTransformMenu( FMenuBuilder& MenuBuilder )
+void FLevelEditorContextMenuImpl::FillTransformMenu(UEditorMenu* Menu)
 {
 	if ( FLevelEditorActionCallbacks::ActorSelected_CanExecute() )
 	{
-		MenuBuilder.BeginSection("TransformSnapAlign");
 		{
-			MenuBuilder.AddSubMenu( 
+			FEditorMenuSection& Section = Menu->AddSection("TransformSnapAlign");
+			Section.AddEntry(FEditorMenuEntry::InitSubMenu(
+				Menu->GetMenuName(),
+				"SnapAlignSubMenu",
 				LOCTEXT("SnapAlignSubMenu", "Snap/Align"), 
 				LOCTEXT("SnapAlignSubMenu_ToolTip", "Actor snap/align utils"),
-				FNewMenuDelegate::CreateStatic( &FLevelEditorContextMenuImpl::FillSnapAlignMenu ) );
+				FNewEditorMenuDelegate::CreateStatic(&FLevelEditorContextMenuImpl::FillSnapAlignMenu)));
 		}
-		MenuBuilder.EndSection();
 
-		MenuBuilder.BeginSection("DeltaTransformToActors");
 		{
-			MenuBuilder.AddMenuEntry( FLevelEditorCommands::Get().DeltaTransformToActors );
+			FEditorMenuSection& Section = Menu->AddSection("DeltaTransformToActors");
+			Section.AddMenuEntry(FLevelEditorCommands::Get().DeltaTransformToActors);
 		}
-		MenuBuilder.EndSection();
 	}
 
-	MenuBuilder.BeginSection("MirrorLock");
 	{
-		MenuBuilder.AddMenuEntry( FLevelEditorCommands::Get().MirrorActorX );
-		MenuBuilder.AddMenuEntry( FLevelEditorCommands::Get().MirrorActorY );
-		MenuBuilder.AddMenuEntry( FLevelEditorCommands::Get().MirrorActorZ );
-		MenuBuilder.AddMenuEntry( FLevelEditorCommands::Get().LockActorMovement );
+		FEditorMenuSection& Section = Menu->AddSection("MirrorLock");
+		Section.AddMenuEntry(FLevelEditorCommands::Get().MirrorActorX);
+		Section.AddMenuEntry(FLevelEditorCommands::Get().MirrorActorY);
+		Section.AddMenuEntry(FLevelEditorCommands::Get().MirrorActorZ);
+		Section.AddMenuEntry(FLevelEditorCommands::Get().LockActorMovement);
 	}
 }
 
-void FLevelEditorContextMenuImpl::FillActorMenu( FMenuBuilder& MenuBuilder )
+void FLevelEditorContextMenuImpl::FillActorMenu(UEditorMenu* Menu)
 {
 	struct Local
 	{
@@ -969,9 +1051,10 @@ void FLevelEditorContextMenuImpl::FillActorMenu( FMenuBuilder& MenuBuilder )
 		InitOptions.Filters->AddFilterPredicate( SceneOutliner::FActorFilterPredicate::CreateStatic( &FLevelEditorActionCallbacks::IsAttachableActor) );
 	}		
 
+	FEditorMenuSection& Section = Menu->AddSection("Actor");
 	if(SelectionInfo.bHaveAttachedActor)
 	{
-		MenuBuilder.AddMenuEntry( FLevelEditorCommands::Get().DetachFromParent, NAME_None, LOCTEXT( "None", "None" ) );
+		Section.AddMenuEntry(FLevelEditorCommands::Get().DetachFromParent, LOCTEXT("None", "None"));
 	}
 
 	// Actor selector to allow the user to choose a parent actor
@@ -1019,23 +1102,24 @@ void FLevelEditorContextMenuImpl::FillActorMenu( FMenuBuilder& MenuBuilder )
 			]
 		];
 
-	MenuBuilder.AddWidget(MenuWidget, FText::GetEmpty(), false);
+	Section.AddEntry(FEditorMenuEntry::InitWidget("PickParentActor", MenuWidget, FText::GetEmpty(), false));
 }
 
-void FLevelEditorContextMenuImpl::FillSnapAlignMenu( FMenuBuilder& MenuBuilder )
+void FLevelEditorContextMenuImpl::FillSnapAlignMenu(UEditorMenu* Menu)
 {
-	MenuBuilder.AddMenuEntry( FLevelEditorCommands::Get().SnapOriginToGrid );
-	MenuBuilder.AddMenuEntry( FLevelEditorCommands::Get().SnapOriginToGridPerActor );
-	MenuBuilder.AddMenuEntry( FLevelEditorCommands::Get().AlignOriginToGrid );
-	MenuBuilder.AddMenuEntry( FLevelEditorCommands::Get().SnapTo2DLayer );
-	MenuBuilder.AddMenuEntry( FLevelEditorCommands::Get().SnapToFloor );
-	MenuBuilder.AddMenuEntry( FLevelEditorCommands::Get().AlignToFloor );
-	MenuBuilder.AddMenuEntry( FLevelEditorCommands::Get().SnapPivotToFloor );
-	MenuBuilder.AddMenuEntry( FLevelEditorCommands::Get().AlignPivotToFloor );
-	MenuBuilder.AddMenuEntry( FLevelEditorCommands::Get().SnapBottomCenterBoundsToFloor );
-	MenuBuilder.AddMenuEntry( FLevelEditorCommands::Get().AlignBottomCenterBoundsToFloor );
+	FEditorMenuSection& Section = Menu->AddSection("SnapAlign");
+	Section.AddMenuEntry( FLevelEditorCommands::Get().SnapOriginToGrid );
+	Section.AddMenuEntry( FLevelEditorCommands::Get().SnapOriginToGridPerActor );
+	Section.AddMenuEntry( FLevelEditorCommands::Get().AlignOriginToGrid );
+	Section.AddMenuEntry( FLevelEditorCommands::Get().SnapTo2DLayer );
+	Section.AddMenuEntry( FLevelEditorCommands::Get().SnapToFloor );
+	Section.AddMenuEntry( FLevelEditorCommands::Get().AlignToFloor );
+	Section.AddMenuEntry( FLevelEditorCommands::Get().SnapPivotToFloor );
+	Section.AddMenuEntry( FLevelEditorCommands::Get().AlignPivotToFloor );
+	Section.AddMenuEntry( FLevelEditorCommands::Get().SnapBottomCenterBoundsToFloor );
+	Section.AddMenuEntry( FLevelEditorCommands::Get().AlignBottomCenterBoundsToFloor );
 /*
-	MenuBuilder.AddMenuSeparator();
+	Section.AddSeparator();
 	AActor* Actor = GEditor->GetSelectedActors()->GetBottom<AActor>();
 	if( Actor && FLevelEditorActionCallbacks::ActorsSelected_CanExecute())
 	{
@@ -1059,53 +1143,53 @@ void FLevelEditorContextMenuImpl::FillSnapAlignMenu( FMenuBuilder& MenuBuilder )
 		SnapBottomCenterBoundsToActor->Label = FString::Printf( *LOCTEXT("Snap Bottom Center Bounds To", "Snap Bottom Center Bounds to %s"), *Label);
 		AlignBottomCenterBoundsToActor->Label = FString::Printf( *LOCTEXT("Align Bottom Center Bounds To", "Align Bottom Center Bounds to %s"), *Label);
 
-		MenuBuilder.AddMenuEntry( SnapOriginToActor );
-		MenuBuilder.AddMenuEntry( AlignOriginToActor );
-		MenuBuilder.AddMenuEntry( SnapToActor );
-		MenuBuilder.AddMenuEntry( AlignToActor );
-		MenuBuilder.AddMenuEntry( SnapPivotToActor );
-		MenuBuilder.AddMenuEntry( AlignPivotToActor );
-		MenuBuilder.AddMenuEntry( SnapBottomCenterBoundsToActor );
-		MenuBuilder.AddMenuEntry( AlignBottomCenterBoundsToActor );
+		Section.AddMenuEntry( SnapOriginToActor );
+		Section.AddMenuEntry( AlignOriginToActor );
+		Section.AddMenuEntry( SnapToActor );
+		Section.AddMenuEntry( AlignToActor );
+		Section.AddMenuEntry( SnapPivotToActor );
+		Section.AddMenuEntry( AlignPivotToActor );
+		Section.AddMenuEntry( SnapBottomCenterBoundsToActor );
+		Section.AddMenuEntry( AlignBottomCenterBoundsToActor );
 	}
 	else
 	{
-		MenuBuilder.AddMenuEntry( FLevelEditorCommands::Get().SnapOriginToActor );
-		MenuBuilder.AddMenuEntry( FLevelEditorCommands::Get().AlignOriginToActor );
-		MenuBuilder.AddMenuEntry( FLevelEditorCommands::Get().SnapToActor );
-		MenuBuilder.AddMenuEntry( FLevelEditorCommands::Get().AlignToActor );
-		MenuBuilder.AddMenuEntry( FLevelEditorCommands::Get().SnapPivotToActor );
-		MenuBuilder.AddMenuEntry( FLevelEditorCommands::Get().AlignPivotToActor );
-		MenuBuilder.AddMenuEntry( FLevelEditorCommands::Get().SnapBottomCenterBoundsToActor );
-		MenuBuilder.AddMenuEntry( FLevelEditorCommands::Get().AlignBottomCenterBoundsToActor );
+		Section.AddMenuEntry( FLevelEditorCommands::Get().SnapOriginToActor );
+		Section.AddMenuEntry( FLevelEditorCommands::Get().AlignOriginToActor );
+		Section.AddMenuEntry( FLevelEditorCommands::Get().SnapToActor );
+		Section.AddMenuEntry( FLevelEditorCommands::Get().AlignToActor );
+		Section.AddMenuEntry( FLevelEditorCommands::Get().SnapPivotToActor );
+		Section.AddMenuEntry( FLevelEditorCommands::Get().AlignPivotToActor );
+		Section.AddMenuEntry( FLevelEditorCommands::Get().SnapBottomCenterBoundsToActor );
+		Section.AddMenuEntry( FLevelEditorCommands::Get().AlignBottomCenterBoundsToActor );
 	}
 */
 }
 
-void FLevelEditorContextMenuImpl::FillPivotMenu( FMenuBuilder& MenuBuilder )
+void FLevelEditorContextMenuImpl::FillPivotMenu( UEditorMenu* Menu )
 {
-	MenuBuilder.BeginSection("SaveResetPivot");
 	{
-		MenuBuilder.AddMenuEntry( FLevelEditorCommands::Get().SavePivotToPrePivot );
-		MenuBuilder.AddMenuEntry( FLevelEditorCommands::Get().ResetPrePivot );
-		MenuBuilder.AddMenuEntry(FLevelEditorCommands::Get().MovePivotHere);
-		MenuBuilder.AddMenuEntry(FLevelEditorCommands::Get().MovePivotHereSnapped);
+		FEditorMenuSection& Section = Menu->AddSection("SaveResetPivot");
+		Section.AddMenuEntry(FLevelEditorCommands::Get().SavePivotToPrePivot);
+		Section.AddMenuEntry(FLevelEditorCommands::Get().ResetPrePivot);
+		Section.AddMenuEntry(FLevelEditorCommands::Get().MovePivotHere);
+		Section.AddMenuEntry(FLevelEditorCommands::Get().MovePivotHereSnapped);
 	}
-	MenuBuilder.EndSection();
 
-	MenuBuilder.BeginSection("MovePivot");
 	{
-		MenuBuilder.AddMenuEntry( FLevelEditorCommands::Get().MovePivotToCenter );
+		FEditorMenuSection& Section = Menu->AddSection("MovePivot");
+		Section.AddMenuEntry(FLevelEditorCommands::Get().MovePivotToCenter);
 	}
-	MenuBuilder.EndSection();
 }
 
-void FLevelEditorContextMenuImpl::FillGroupMenu( class FMenuBuilder& MenuBuilder )
+void FLevelEditorContextMenuImpl::FillGroupMenu( UEditorMenu* Menu )
 {
+	FEditorMenuSection& Section = Menu->AddSection("Group");
+
 	if( SelectionInfo.NumSelectedUngroupedActors > 1 )
 	{
 		// Only show this menu item if we have more than one actor.
-		MenuBuilder.AddMenuEntry( FLevelEditorCommands::Get().GroupActors  );
+		Section.AddMenuEntry( FLevelEditorCommands::Get().GroupActors );
 	}
 
 	if( SelectionInfo.bHaveSelectedLockedGroup || SelectionInfo.bHaveSelectedUnlockedGroup )
@@ -1116,50 +1200,55 @@ void FLevelEditorContextMenuImpl::FillGroupMenu( class FMenuBuilder& MenuBuilder
 		// Only allow regrouping if multiple groups are selected, or a group and ungrouped actors are selected
 		if( NumActiveGroups > 1 || (NumActiveGroups && SelectionInfo.NumSelectedUngroupedActors) )
 		{
-			MenuBuilder.AddMenuEntry( FLevelEditorCommands::Get().RegroupActors );
+			Section.AddMenuEntry( FLevelEditorCommands::Get().RegroupActors );
 		}
 
-		MenuBuilder.AddMenuEntry( FLevelEditorCommands::Get().UngroupActors );
+		Section.AddMenuEntry( FLevelEditorCommands::Get().UngroupActors );
 
 		if( SelectionInfo.bHaveSelectedUnlockedGroup )
 		{
 			// Only allow removal of loose actors or locked subgroups
 			if( !SelectionInfo.bHaveSelectedLockedGroup || ( SelectionInfo.bHaveSelectedLockedGroup && SelectionInfo.bHaveSelectedSubGroup ) )
 			{
-				MenuBuilder.AddMenuEntry( FLevelEditorCommands::Get().RemoveActorsFromGroup );
+				Section.AddMenuEntry( FLevelEditorCommands::Get().RemoveActorsFromGroup );
 			}
-			MenuBuilder.AddMenuEntry( FLevelEditorCommands::Get().LockGroup );
+			Section.AddMenuEntry( FLevelEditorCommands::Get().LockGroup );
 		}
 
 		if( SelectionInfo.bHaveSelectedLockedGroup )
 		{
-			MenuBuilder.AddMenuEntry( FLevelEditorCommands::Get().UnlockGroup );
+			Section.AddMenuEntry( FLevelEditorCommands::Get().UnlockGroup );
 		}
 
 		// Only allow group adds if a single group is selected in addition to ungrouped actors
 		if( AGroupActor::NumActiveGroups(true, false) == 1 && SelectionInfo.NumSelectedUngroupedActors )
 		{ 
-			MenuBuilder.AddMenuEntry( FLevelEditorCommands::Get().AddActorsToGroup );
+			Section.AddMenuEntry( FLevelEditorCommands::Get().AddActorsToGroup );
 		}
 	}
 }
 
-void FLevelEditorContextMenuImpl::FillEditMenu( class FMenuBuilder& MenuBuilder, LevelEditorMenuContext ContextType )
+void FLevelEditorContextMenuImpl::FillEditMenu( UEditorMenu* Menu )
 {
-	MenuBuilder.AddMenuEntry( FGenericCommands::Get().Cut );
-	MenuBuilder.AddMenuEntry( FGenericCommands::Get().Copy );
-	MenuBuilder.AddMenuEntry( FGenericCommands::Get().Paste );
-	if (ContextType == LevelEditorMenuContext::Viewport)
+	FEditorMenuSection& Section = Menu->AddSection("Section");
+
+	Section.AddMenuEntry( FGenericCommands::Get().Cut );
+	Section.AddMenuEntry( FGenericCommands::Get().Copy );
+	Section.AddMenuEntry( FGenericCommands::Get().Paste );
+	if (ULevelEditorContextMenuContext* LevelEditorContext = Menu->FindContext<ULevelEditorContextMenuContext>())
 	{
-		MenuBuilder.AddMenuEntry( FLevelEditorCommands::Get().PasteHere );
+		if (LevelEditorContext->ContextType == ELevelEditorMenuContext::Viewport)
+		{
+			Section.AddMenuEntry(FLevelEditorCommands::Get().PasteHere);
+		}
 	}
 
-	MenuBuilder.AddMenuEntry( FGenericCommands::Get().Duplicate );
-	MenuBuilder.AddMenuEntry( FGenericCommands::Get().Delete );
-	MenuBuilder.AddMenuEntry( FGenericCommands::Get().Rename );
+	Section.AddMenuEntry( FGenericCommands::Get().Duplicate );
+	Section.AddMenuEntry( FGenericCommands::Get().Delete );
+	Section.AddMenuEntry( FGenericCommands::Get().Rename );
 }
 
-void FLevelScriptEventMenuHelper::FillLevelBlueprintEventsMenu(class FMenuBuilder& MenuBuilder, const TArray<AActor*>& SelectedActors)
+void FLevelScriptEventMenuHelper::FillLevelBlueprintEventsMenu(UEditorMenu* Menu, const TArray<AActor*>& SelectedActors)
 {
 	AActor* SelectedActor = (1 == SelectedActors.Num()) ? SelectedActors[0] : NULL;
 
@@ -1172,33 +1261,35 @@ void FLevelScriptEventMenuHelper::FillLevelBlueprintEventsMenu(class FMenuBuilde
 		{
 			TWeakObjectPtr<AActor> ActorPtr(SelectedActor);
 
-			MenuBuilder.BeginSection("LevelBlueprintEvents", LOCTEXT("LevelBlueprintEvents", "Level Blueprint Events"));
+			FEditorMenuSection& Section = Menu->AddSection("LevelBlueprintEvents", LOCTEXT("LevelBlueprintEvents", "Level Blueprint Events"));
 
 			if (bAnyEventExists)
 			{
-				MenuBuilder.AddSubMenu(
+				Section.AddEntry(FEditorMenuEntry::InitSubMenu(
+					Menu->GetMenuName(),
+					"JumpEventSubMenu",
 					LOCTEXT("JumpEventSubMenu", "Jump to Event"),
 					FText::GetEmpty(),
-					FNewMenuDelegate::CreateStatic(&FKismetEditorUtilities::AddLevelScriptEventOptionsForActor
+					FNewEditorMenuDelegate::CreateStatic(&FKismetEditorUtilities::AddLevelScriptEventOptionsForActor
 					, ActorPtr
 					, true
 					, false
-					, true));
+					, true)));
 			}
 
 			if (bAnyEventCanBeAdded)
 			{
-				MenuBuilder.AddSubMenu(
+				Section.AddEntry(FEditorMenuEntry::InitSubMenu(
+					Menu->GetMenuName(),
+					"AddEventSubMenu",
 					LOCTEXT("AddEventSubMenu", "Add Event"),
 					FText::GetEmpty(),
-					FNewMenuDelegate::CreateStatic(&FKismetEditorUtilities::AddLevelScriptEventOptionsForActor
+					FNewEditorMenuDelegate::CreateStatic(&FKismetEditorUtilities::AddLevelScriptEventOptionsForActor
 					, ActorPtr
 					, false
 					, true
-					, true));
+					, true)));
 			}
-
-			MenuBuilder.EndSection();
 		}
 	}
 }
diff --git a/Engine/Source/Editor/LevelEditor/Private/LevelEditorContextMenu.h b/Engine/Source/Editor/LevelEditor/Private/LevelEditorContextMenu.h
--- a/Engine/Source/Editor/LevelEditor/Private/LevelEditorContextMenu.h
+++ b/Engine/Source/Editor/LevelEditor/Private/LevelEditorContextMenu.h
@@ -4,21 +4,14 @@
 
 #include "CoreMinimal.h"
 #include "Editor/UnrealEdTypes.h"
+#include "LevelEditorMenuContext.h"
 
 class FExtender;
-class FMenuBuilder;
+class UEditorMenu;
+struct FEditorMenuContext;
 class SLevelEditor;
 class SWidget;
 
-/** Enum to describe what a menu should be built for */
-enum class LevelEditorMenuContext
-{
-	/** This context menu is applicable to a viewport */
-	Viewport,
-	/** This context menu is applicable to the Scene Outliner (disables click-position-based menu items) */
-	SceneOutliner,
-};
-
 /**
  * Context menu construction class 
  */
@@ -32,7 +25,7 @@ public:
 	 * @param	LevelEditor		The level editor using this menu.
 	 * @param	ContextType		The context we should use to specialize this menu
 	 */
-	static void SummonMenu( const TSharedRef< class SLevelEditor >& LevelEditor, LevelEditorMenuContext ContextType );
+	static void SummonMenu( const TSharedRef< class SLevelEditor >& LevelEditor, ELevelEditorMenuContext ContextType );
 
 	/**
 	 * Summons the viewport view option menu
@@ -48,25 +41,35 @@ public:
 	 * @param	Extender		Allows extension of this menu based on context.
 	 * @return	Widget for this context menu
 	 */
-	static TSharedPtr< SWidget > BuildMenuWidget(TWeakPtr< SLevelEditor > LevelEditor, LevelEditorMenuContext ContextType, TSharedPtr<FExtender> Extender = TSharedPtr<FExtender>());
+	static TSharedPtr< SWidget > BuildMenuWidget(TWeakPtr< SLevelEditor > LevelEditor, ELevelEditorMenuContext ContextType, TSharedPtr<FExtender> Extender = TSharedPtr<FExtender>());
 
 	/**
 	 * Populates the specified menu builder for the context menu that can be inserted into a pop-up window
 	 *
-	 * @param	MenuBuilder		The menu builder to fill the menu with
+	 * @param	Menu			The menu to fill
 	 * @param	LevelEditor		The level editor using this menu.
 	 * @param	ContextType		The context we should use to specialize this menu
 	 * @param	Extender		Allows extension of this menu based on context.
 	 */
-	static void FillMenu(FMenuBuilder& MenuBuilder, TWeakPtr< SLevelEditor > LevelEditor, LevelEditorMenuContext ContextType, TSharedPtr<FExtender> Extender);
+	static UEditorMenu* GenerateMenu(TWeakPtr< SLevelEditor > LevelEditor, ELevelEditorMenuContext ContextType, TSharedPtr<FExtender> Extender = TSharedPtr<FExtender>());
+
+	/* Adds required information to Context for build menu based on current selection */
+	static FName InitMenuContext(FEditorMenuContext& Context, TWeakPtr<SLevelEditor> LevelEditor, ELevelEditorMenuContext ContextType);
+
+	/* Returns name of menu to display based on current selection */
+	static FName GetContextMenuName(ELevelEditorMenuContext ContextType);
 
 private:
 
+	static void RegisterComponentContextMenu();
+	static void RegisterActorContextMenu();
+	static void RegisterSceneOutlinerContextMenu();
+
 	/**
 	 * Builds the actor group menu
 	 *
-	 * @param MenuBuilder		The menu builder to add items to.
+	 * @param Menu		The menu to add items to.
 	 * @param SelectedActorInfo	Information about the selected actors.
 	 */
-	static void BuildGroupMenu( FMenuBuilder& MenuBuilder, const struct FSelectedActorInfo& SelectedActorInfo );
+	static void BuildGroupMenu(UEditorMenu* Menu, const struct FSelectedActorInfo& SelectedActorInfo);
 };
diff --git a/Engine/Source/Editor/LevelEditor/Private/LevelEditorCreateActorMenu.cpp b/Engine/Source/Editor/LevelEditor/Private/LevelEditorCreateActorMenu.cpp
--- a/Engine/Source/Editor/LevelEditor/Private/LevelEditorCreateActorMenu.cpp
+++ b/Engine/Source/Editor/LevelEditor/Private/LevelEditorCreateActorMenu.cpp
@@ -9,6 +9,7 @@
 #include "Widgets/SBoxPanel.h"
 #include "Textures/SlateIcon.h"
 #include "Framework/Commands/UIAction.h"
+#include "EditorMenuSubsystem.h"
 #include "Widgets/Images/SImage.h"
 #include "Widgets/Text/STextBlock.h"
 #include "Widgets/Layout/SBox.h"
@@ -279,8 +280,9 @@ static void GetContentBrowserSelectionFactoryMenuEntries( FAssetData& TargetAsse
  * @param	AssetMenuOptions	Menu options to use
  * @param	CreateMode			The creation mode to use
  */
-static void FillAssetAddReplaceActorMenu( FMenuBuilder& MenuBuilder, const FAssetData Asset, const TArray< FActorFactoryAssetProxy::FMenuItem > AssetMenuOptions, EActorCreateMode::Type CreateMode )
+static void FillAssetAddReplaceActorMenu(UEditorMenu* Menu, const FAssetData Asset, const TArray< FActorFactoryAssetProxy::FMenuItem > AssetMenuOptions, EActorCreateMode::Type CreateMode)
 {
+	FEditorMenuSection& Section = Menu->AddSection("Section");
 	for( int32 ItemIndex = 0; ItemIndex < AssetMenuOptions.Num(); ++ItemIndex )
 	{
 		const FActorFactoryAssetProxy::FMenuItem& MenuItem = AssetMenuOptions[ItemIndex];
@@ -305,7 +307,7 @@ static void FillAssetAddReplaceActorMenu( FMenuBuilder& MenuBuilder, const FAsse
 			Action = FUIAction( FExecuteAction::CreateStatic( &FLevelEditorActionCallbacks::AddActor_Clicked, MenuItem.FactoryToUse,  MenuItem.AssetData, CreateMode == EActorCreateMode::Placement ) );
 		}
 
-		MenuBuilder.AddMenuEntry( Label, ToolTip, Icon, Action );
+		Section.AddMenuEntry(NAME_None, Label, ToolTip, Icon, Action);
 	}
 }
 
@@ -317,7 +319,7 @@ static void FillAssetAddReplaceActorMenu( FMenuBuilder& MenuBuilder, const FAsse
  * @param	CreateMode			The creation mode to use
  * @param	LabelOverride		The lable to use, if any.
  */
-static void BuildSingleAssetAddReplaceActorMenu( FMenuBuilder& MenuBuilder, const FAssetData& Asset, const TArray< FActorFactoryAssetProxy::FMenuItem >& AssetMenuOptions, EActorCreateMode::Type CreateMode, const FText& LabelOverride = FText::GetEmpty() )
+static void BuildSingleAssetAddReplaceActorMenu(FEditorMenuSection& Section, const FAssetData& Asset, const TArray< FActorFactoryAssetProxy::FMenuItem >& AssetMenuOptions, EActorCreateMode::Type CreateMode, const FText& LabelOverride = FText::GetEmpty())
 {
 	if ( !Asset.IsValid() || AssetMenuOptions.Num() == 0 )
 	{
@@ -338,17 +340,18 @@ static void BuildSingleAssetAddReplaceActorMenu( FMenuBuilder& MenuBuilder, cons
 			Action = FUIAction( FExecuteAction::CreateStatic( &FLevelEditorActionCallbacks::AddActor_Clicked, MenuItem.FactoryToUse,  MenuItem.AssetData, CreateMode == EActorCreateMode::Placement ) );
 		}
 
-		MenuBuilder.AddMenuEntry( Action, SNew( SAssetMenuEntry, Asset, AssetMenuOptions ).LabelOverride( LabelOverride ) );
+		Section.AddEntry(FEditorMenuEntry::InitMenuEntry(NAME_None, Action, SNew(SAssetMenuEntry, Asset, AssetMenuOptions).LabelOverride(LabelOverride)));
 	}
 	else
 	{
-		MenuBuilder.AddSubMenu( 
-			SNew( SAssetMenuEntry, Asset, AssetMenuOptions ).LabelOverride( LabelOverride ), 
-			FNewMenuDelegate::CreateStatic( &FillAssetAddReplaceActorMenu, Asset, AssetMenuOptions, CreateMode ) );
+		FEditorMenuEntry Entry = FEditorMenuEntry::InitWidget(NAME_None, SNew( SAssetMenuEntry, Asset, AssetMenuOptions ).LabelOverride( LabelOverride ), FText());
+		Entry.SubMenuData.bIsSubMenu = true;
+		Entry.SubMenuData.ConstructMenu = FNewEditorMenuDelegate::CreateStatic(&FillAssetAddReplaceActorMenu, Asset, AssetMenuOptions, CreateMode);
+		Section.AddEntry(Entry);
 	}
 }
 
-void LevelEditorCreateActorMenu::FillAddReplaceViewportContextMenuSections( FMenuBuilder& MenuBuilder )
+void LevelEditorCreateActorMenu::FillAddReplaceViewportContextMenuSections(UEditorMenu* Menu)
 {
 	FAssetData TargetAssetData;
 	TArray< FActorFactoryAssetProxy::FMenuItem > AssetMenuOptions;
@@ -356,22 +359,25 @@ void LevelEditorCreateActorMenu::FillAddReplaceViewportContextMenuSections( FMen
 
 	if ( AssetMenuOptions.Num() == 0 )
 	{
-		MenuBuilder.BeginSection("ActorType");
 		{
-			MenuBuilder.AddSubMenu( 
+			FEditorMenuSection& Section = Menu->AddSection("ActorType");
+			Section.AddEntry(FEditorMenuEntry::InitSubMenu(
+				Menu->GetMenuName(),
+				"LevelViewportContextMenu",
 				NSLOCTEXT("LevelViewportContextMenu", "AddActorHeading", "Place Actor") , 
 				NSLOCTEXT("LevelViewportContextMenu", "AddActorMenu_ToolTip", "Templates for adding a new actor to the world"),
-				FNewMenuDelegate::CreateStatic( &LevelEditorCreateActorMenu::FillAddReplaceActorMenu, EActorCreateMode::Add ) );
+				FNewEditorMenuDelegate::CreateStatic(&LevelEditorCreateActorMenu::FillAddReplaceActorMenu, EActorCreateMode::Add)));
 
 			if ( CanReplaceActors() )
 			{
-				MenuBuilder.AddSubMenu( 
+				Section.AddEntry(FEditorMenuEntry::InitSubMenu(
+					Menu->GetMenuName(),
+					"LevelViewportContextMenu",
 					NSLOCTEXT("LevelViewportContextMenu", "ReplaceActorHeading", "Replace Selected Actors with") , 
 					NSLOCTEXT("LevelViewportContextMenu", "ReplaceActorMenu_ToolTip", "Templates for replacing selected with new actors in the world"),
-					FNewMenuDelegate::CreateStatic( &LevelEditorCreateActorMenu::FillAddReplaceActorMenu, EActorCreateMode::Replace ) );
+					FNewEditorMenuDelegate::CreateStatic(&LevelEditorCreateActorMenu::FillAddReplaceActorMenu, EActorCreateMode::Replace)));
 			}
 		}
-		MenuBuilder.EndSection();
 	}
 	else
 	{
@@ -380,39 +386,50 @@ void LevelEditorCreateActorMenu::FillAddReplaceViewportContextMenuSections( FMen
 			AssetMenuOptions.Pop();
 		}
 
-		MenuBuilder.BeginSection("AddActor", NSLOCTEXT("LevelViewportContextMenu", "AddActorHeading", "Place Actor") );
 		{
+			FEditorMenuSection& Section = Menu->AddSection("AddActor", NSLOCTEXT("LevelViewportContextMenu", "AddActorHeading", "Place Actor"));
 			FUIAction Action( FExecuteAction::CreateStatic( &FLevelEditorActionCallbacks::AddActor_Clicked, AssetMenuOptions[0].FactoryToUse,  AssetMenuOptions[0].AssetData, false ) );
-			MenuBuilder.AddSubMenu( Action, SNew( SAssetMenuEntry, TargetAssetData, AssetMenuOptions ), FNewMenuDelegate::CreateStatic( &LevelEditorCreateActorMenu::FillAddReplaceActorMenu, EActorCreateMode::Add ) );
+			TSharedRef<SWidget> Widget = SNew(SAssetMenuEntry, TargetAssetData, AssetMenuOptions);
+			Section.AddEntry(FEditorMenuEntry::InitSubMenu(
+				Menu->GetMenuName(),
+				"AddActor",
+				Action,
+				Widget,
+				FNewEditorMenuDelegate::CreateStatic(&LevelEditorCreateActorMenu::FillAddReplaceActorMenu, EActorCreateMode::Add)
+			));
 		}
-		MenuBuilder.EndSection();
 
 		if ( CanReplaceActors() )
 		{
-			MenuBuilder.BeginSection("ReplaceActor", NSLOCTEXT("LevelViewportContextMenu", "ReplaceActorHeading", "Replace Selected Actors with") );
 			{
+				FEditorMenuSection& Section = Menu->AddSection("ReplaceActor", NSLOCTEXT("LevelViewportContextMenu", "ReplaceActorHeading", "Replace Selected Actors with"));
 				FUIAction Action( FExecuteAction::CreateStatic( &FLevelEditorActionCallbacks::ReplaceActors_Clicked, AssetMenuOptions[0].FactoryToUse,  AssetMenuOptions[0].AssetData ) );
-				MenuBuilder.AddSubMenu( Action, SNew( SAssetMenuEntry, TargetAssetData, AssetMenuOptions ), FNewMenuDelegate::CreateStatic( &LevelEditorCreateActorMenu::FillAddReplaceActorMenu, EActorCreateMode::Replace ) );
+				TSharedRef<SWidget> Widget = SNew(SAssetMenuEntry, TargetAssetData, AssetMenuOptions);
+				Section.AddEntry(FEditorMenuEntry::InitSubMenu(
+					Menu->GetMenuName(),
+					"ReplaceActor",
+					Action,
+					Widget,
+					FNewEditorMenuDelegate::CreateStatic(&LevelEditorCreateActorMenu::FillAddReplaceActorMenu, EActorCreateMode::Replace)
+				));
 			}
-			MenuBuilder.EndSection();
 		}
 	}
 }
 
-void LevelEditorCreateActorMenu::FillAddReplaceActorMenu( FMenuBuilder& MenuBuilder, EActorCreateMode::Type CreateMode )
+void LevelEditorCreateActorMenu::FillAddReplaceActorMenu(UEditorMenu* Menu, EActorCreateMode::Type CreateMode)
 {
-	MenuBuilder.BeginSection("ContentBrowserActor", NSLOCTEXT("LevelViewportContextMenu", "AssetSelectionSection", "Selection") );
 	{
+		FEditorMenuSection& Section = Menu->AddSection("ContentBrowserActor", NSLOCTEXT("LevelViewportContextMenu", "AssetSelectionSection", "Selection"));
 		FAssetData TargetAssetData;
 		TArray< FActorFactoryAssetProxy::FMenuItem > AssetMenuOptions;
 		GetContentBrowserSelectionFactoryMenuEntries( /*OUT*/TargetAssetData, /*OUT*/AssetMenuOptions );
 
-		BuildSingleAssetAddReplaceActorMenu( MenuBuilder, TargetAssetData, AssetMenuOptions, CreateMode );
+		BuildSingleAssetAddReplaceActorMenu( Section, TargetAssetData, AssetMenuOptions, CreateMode );
 	}
-	MenuBuilder.EndSection();
 
-	MenuBuilder.BeginSection("RecentlyPlaced", NSLOCTEXT("LevelViewportContextMenu", "RecentlyPlacedSection", "Recently Placed") );
 	{
+		FEditorMenuSection& Section = Menu->AddSection("RecentlyPlaced", NSLOCTEXT("LevelViewportContextMenu", "RecentlyPlacedSection", "Recently Placed"));
 		if ( IPlacementModeModule::IsAvailable() )
 		{
 			FAssetRegistryModule& AssetRegistryModule = FModuleManager::LoadModuleChecked<FAssetRegistryModule>(TEXT("AssetRegistry"));
@@ -440,15 +457,14 @@ void LevelEditorCreateActorMenu::FillAddReplaceActorMenu( FMenuBuilder& MenuBuil
 						}
 					}
 
-					BuildSingleAssetAddReplaceActorMenu( MenuBuilder, Asset, AssetMenuOptions, CreateMode );
+					BuildSingleAssetAddReplaceActorMenu(Section, Asset, AssetMenuOptions, CreateMode);
 				}
 			}
 		}
 	}
-	MenuBuilder.EndSection();
 
-	MenuBuilder.BeginSection("Lights", NSLOCTEXT("LevelViewportContextMenu", "LightsSection", "Lights") );
 	{
+		FEditorMenuSection& Section = Menu->AddSection("Lights", NSLOCTEXT("LevelViewportContextMenu", "LightsSection", "Lights"));
 		TArray< FActorFactoryAssetProxy::FMenuItem > AssetMenuOptions;
 
 		{
@@ -456,7 +472,7 @@ void LevelEditorCreateActorMenu::FillAddReplaceActorMenu( FMenuBuilder& MenuBuil
 			UActorFactory* Factory = GEditor->FindActorFactoryByClass( UActorFactoryDirectionalLight::StaticClass() );
 			FAssetData AssetData = FAssetData( Factory->GetDefaultActorClass( FAssetData() ) );
 			AssetMenuOptions.Add( FActorFactoryAssetProxy::FMenuItem( Factory, AssetData ) );
-			BuildSingleAssetAddReplaceActorMenu( MenuBuilder, AssetData, AssetMenuOptions, CreateMode );
+			BuildSingleAssetAddReplaceActorMenu(Section, AssetData, AssetMenuOptions, CreateMode);
 		}
 
 		{
@@ -464,7 +480,7 @@ void LevelEditorCreateActorMenu::FillAddReplaceActorMenu( FMenuBuilder& MenuBuil
 			UActorFactory* Factory = GEditor->FindActorFactoryByClass( UActorFactorySpotLight::StaticClass() );
 			FAssetData AssetData = FAssetData( Factory->GetDefaultActorClass( FAssetData() ) );
 			AssetMenuOptions.Add( FActorFactoryAssetProxy::FMenuItem( Factory, AssetData ) );
-			BuildSingleAssetAddReplaceActorMenu( MenuBuilder, AssetData, AssetMenuOptions, CreateMode );
+			BuildSingleAssetAddReplaceActorMenu(Section, AssetData, AssetMenuOptions, CreateMode);
 		}
 
 		{
@@ -472,7 +488,7 @@ void LevelEditorCreateActorMenu::FillAddReplaceActorMenu( FMenuBuilder& MenuBuil
 			UActorFactory* Factory = GEditor->FindActorFactoryByClass( UActorFactoryPointLight::StaticClass() );
 			FAssetData AssetData = FAssetData( Factory->GetDefaultActorClass( FAssetData() ) );
 			AssetMenuOptions.Add( FActorFactoryAssetProxy::FMenuItem( Factory, AssetData ) );
-			BuildSingleAssetAddReplaceActorMenu( MenuBuilder, AssetData, AssetMenuOptions, CreateMode );
+			BuildSingleAssetAddReplaceActorMenu(Section, AssetData, AssetMenuOptions, CreateMode);
 		}
 
 		{
@@ -480,20 +496,19 @@ void LevelEditorCreateActorMenu::FillAddReplaceActorMenu( FMenuBuilder& MenuBuil
 			UActorFactory* Factory = GEditor->FindActorFactoryByClass( UActorFactoryRectLight::StaticClass() );
 			FAssetData AssetData = FAssetData( Factory->GetDefaultActorClass( FAssetData() ) );
 			AssetMenuOptions.Add( FActorFactoryAssetProxy::FMenuItem( Factory, AssetData ) );
-			BuildSingleAssetAddReplaceActorMenu( MenuBuilder, AssetData, AssetMenuOptions, CreateMode );
+			BuildSingleAssetAddReplaceActorMenu(Section, AssetData, AssetMenuOptions, CreateMode);
 		}
 	}
-	MenuBuilder.EndSection();
 
-	MenuBuilder.BeginSection("Primitives", NSLOCTEXT("LevelViewportContextMenu", "PrimitivesSection", "Primitives") );
 	{
+		FEditorMenuSection& Section = Menu->AddSection("Primitives", NSLOCTEXT("LevelViewportContextMenu", "PrimitivesSection", "Primitives"));
 		TArray< FActorFactoryAssetProxy::FMenuItem > AssetMenuOptions;
 		{
 			AssetMenuOptions.Empty();
 			UActorFactory* Factory = GEditor->FindActorFactoryByClass( UActorFactoryCameraActor::StaticClass() );
 			FAssetData AssetData = FAssetData( Factory->GetDefaultActorClass( FAssetData() ) );
 			AssetMenuOptions.Add( FActorFactoryAssetProxy::FMenuItem( Factory, AssetData ) );
-			BuildSingleAssetAddReplaceActorMenu( MenuBuilder, AssetData, AssetMenuOptions, CreateMode );
+			BuildSingleAssetAddReplaceActorMenu(Section, AssetData, AssetMenuOptions, CreateMode);
 		}
 
 		{
@@ -501,7 +516,7 @@ void LevelEditorCreateActorMenu::FillAddReplaceActorMenu( FMenuBuilder& MenuBuil
 			UActorFactory* Factory = GEditor->FindActorFactoryByClass( UActorFactoryPlayerStart::StaticClass() );
 			FAssetData AssetData = FAssetData( Factory->GetDefaultActorClass( FAssetData() ) );
 			AssetMenuOptions.Add( FActorFactoryAssetProxy::FMenuItem( Factory, AssetData ) );
-			BuildSingleAssetAddReplaceActorMenu( MenuBuilder, AssetData, AssetMenuOptions, CreateMode );
+			BuildSingleAssetAddReplaceActorMenu(Section, AssetData, AssetMenuOptions, CreateMode);
 		}
 
 		{
@@ -517,7 +532,7 @@ void LevelEditorCreateActorMenu::FillAddReplaceActorMenu( FMenuBuilder& MenuBuil
 
 			Factory = GEditor->FindActorFactoryByClassForActorClass( UActorFactoryCylinderVolume::StaticClass(), BlockingVolumeClass );
 			AssetMenuOptions.Add( FActorFactoryAssetProxy::FMenuItem( Factory, AssetData ) );
-			BuildSingleAssetAddReplaceActorMenu( MenuBuilder, AssetData, AssetMenuOptions, CreateMode );
+			BuildSingleAssetAddReplaceActorMenu(Section, AssetData, AssetMenuOptions, CreateMode);
 		}
 
 		{
@@ -534,13 +549,12 @@ void LevelEditorCreateActorMenu::FillAddReplaceActorMenu( FMenuBuilder& MenuBuil
 			AssetData = FAssetData( Factory->GetDefaultActorClass( FAssetData() ) );
 			AssetMenuOptions.Add( FActorFactoryAssetProxy::FMenuItem( Factory, AssetData ) );
 
-			BuildSingleAssetAddReplaceActorMenu( MenuBuilder, FAssetData( ATriggerBase::StaticClass() ), AssetMenuOptions, CreateMode, NSLOCTEXT("LevelViewportContextMenu", "TriggersGroup", "Trigger") );
+			BuildSingleAssetAddReplaceActorMenu(Section, FAssetData(ATriggerBase::StaticClass()), AssetMenuOptions, CreateMode, NSLOCTEXT("LevelViewportContextMenu", "TriggersGroup", "Trigger"));
 		}
 	}
-	MenuBuilder.EndSection();
 
-	MenuBuilder.BeginSection("Custom", NSLOCTEXT("LevelViewportContextMenu", "CustomSection", "Custom Actors") );
 	{
+		FEditorMenuSection& Section = Menu->AddSection("Custom", NSLOCTEXT("LevelViewportContextMenu", "CustomSection", "Custom Actors"));
 		TArray< FActorFactoryAssetProxy::FMenuItem > AssetMenuOptions;
 		FText UnusedErrorMessage;
 		const FAssetData NoAssetData;
@@ -555,9 +569,8 @@ void LevelEditorCreateActorMenu::FillAddReplaceActorMenu( FMenuBuilder& MenuBuil
 			if ( FactoryWorksWithoutAsset && Factory->bShowInEditorQuickMenu )
 			{
 				AssetMenuOptions.Add( FActorFactoryAssetProxy::FMenuItem( Factory, NoAssetData ) );
-				BuildSingleAssetAddReplaceActorMenu( MenuBuilder, AssetData, AssetMenuOptions, CreateMode );
+				BuildSingleAssetAddReplaceActorMenu(Section, AssetData, AssetMenuOptions, CreateMode);
 			}
 		}
 	}
-	MenuBuilder.EndSection();
 }
diff --git a/Engine/Source/Editor/LevelEditor/Private/LevelEditorCreateActorMenu.h b/Engine/Source/Editor/LevelEditor/Private/LevelEditorCreateActorMenu.h
--- a/Engine/Source/Editor/LevelEditor/Private/LevelEditorCreateActorMenu.h
+++ b/Engine/Source/Editor/LevelEditor/Private/LevelEditorCreateActorMenu.h
@@ -4,7 +4,7 @@
 
 #include "CoreMinimal.h"
 
-class FMenuBuilder;
+class UEditorMenu;
 
 /**
  * The mode to use when creating an actor
@@ -30,11 +30,11 @@ namespace LevelEditorCreateActorMenu
 	 * Fill the context menu section(s) for adding or replacing an actor in the viewport
 	 * @param	MenuBuilder		The menu builder used to generate the context menu
 	 */
-	void FillAddReplaceViewportContextMenuSections( FMenuBuilder& MenuBuilder );
+	void FillAddReplaceViewportContextMenuSections(UEditorMenu* Menu);
 
 	/**
 	 * Fill the context menu for adding or replacing an actor. Used for in-viewport and level editor toolbar menus.
 	 * @param	MenuBuilder		The menu builder used to generate the context menu
 	 */
-	void FillAddReplaceActorMenu( FMenuBuilder& MenuBuilder, EActorCreateMode::Type CreateMode );
+	void FillAddReplaceActorMenu(UEditorMenu* Menu, EActorCreateMode::Type CreateMode);
 };
diff --git a/Engine/Source/Editor/LevelEditor/Private/LevelEditorToolBar.cpp b/Engine/Source/Editor/LevelEditor/Private/LevelEditorToolBar.cpp
--- a/Engine/Source/Editor/LevelEditor/Private/LevelEditorToolBar.cpp
+++ b/Engine/Source/Editor/LevelEditor/Private/LevelEditorToolBar.cpp
@@ -1506,7 +1506,7 @@ TSharedRef< SWidget > FLevelEditorToolBar::MakeLevelEditorToolBar( const TShared
 	FLevelEditorModule& LevelEditorModule = FModuleManager::LoadModuleChecked<FLevelEditorModule>("LevelEditor");
 	FEditorMenuContext MenuContext(InCommandList, LevelEditorModule.GetToolBarExtensibilityManager()->GetAllExtenders());
 	ULevelEditorMenuContext* LevelEditorMenuContext = NewObject<ULevelEditorMenuContext>();
-	LevelEditorMenuContext->SlateLevelEditor = InLevelEditor;
+	LevelEditorMenuContext->LevelEditor = InLevelEditor;
 	MenuContext.AddObject(LevelEditorMenuContext);
 
 	// Create the tool bar!
@@ -2282,7 +2282,7 @@ TSharedRef< SWidget > FLevelEditorToolBar::GenerateCinematicsMenuContent( TShare
 	FLevelEditorModule& LevelEditorModule = FModuleManager::LoadModuleChecked<FLevelEditorModule>("LevelEditor");
 	FEditorMenuContext Context(InCommandList, FExtender::Combine(LevelEditorModule.GetAllLevelEditorToolbarCinematicsMenuExtenders()));
 	ULevelEditorMenuContext* LevelEditorMenuContext = NewObject<ULevelEditorMenuContext>();
-	LevelEditorMenuContext->SlateLevelEditor = LevelEditorWeakPtr;
+	LevelEditorMenuContext->LevelEditor = LevelEditorWeakPtr;
 	Context.AddObject(LevelEditorMenuContext);
 	return UEditorMenuSubsystem::Get()->GenerateWidget("LevelEditor.LevelEditorToolBar.Cinematics", Context);
 }
@@ -2309,7 +2309,7 @@ void FLevelEditorToolBar::RegisterCinematicsMenu()
 			return;
 		}
 
-		UWorld* World = FoundContext->SlateLevelEditor.IsValid() ? FoundContext->SlateLevelEditor.Pin()->GetWorld() : nullptr;
+		UWorld* World = FoundContext->LevelEditor.IsValid() ? FoundContext->LevelEditor.Pin()->GetWorld() : nullptr;
 		const bool bHasAnyCinematicsActors = !!TActorIterator<AMatineeActor>(World) || !!TActorIterator<ALevelSequenceActor>(World);
 		if (!bHasAnyCinematicsActors)
 		{
diff --git a/Engine/Source/Editor/LevelEditor/Private/SLevelEditor.cpp b/Engine/Source/Editor/LevelEditor/Private/SLevelEditor.cpp
--- a/Engine/Source/Editor/LevelEditor/Private/SLevelEditor.cpp
+++ b/Engine/Source/Editor/LevelEditor/Private/SLevelEditor.cpp
@@ -2,6 +2,7 @@
 
 #include "SLevelEditor.h"
 #include "Framework/MultiBox/MultiBoxExtender.h"
+#include "EditorMenuSubsystem.h"
 #include "Framework/Docking/LayoutService.h"
 #include "EditorModeRegistry.h"
 #include "EdMode.h"
@@ -693,15 +694,34 @@ TSharedRef<SDockTab> SLevelEditor::SpawnLevelEditorTab( const FSpawnTabArgs& Arg
 		InitOptions.bShowTransient = true;
 		InitOptions.Mode = ESceneOutlinerMode::ActorBrowsing;
 		{
+			UEditorMenuSubsystem* EditorMenus = UEditorMenuSubsystem::Get();
+			static const FName MenuName = "LevelEditor.LevelEditorSceneOutliner.ContextMenu";
+			if (!EditorMenus->IsMenuRegistered(MenuName))
+			{
+				UEditorMenu* Menu = EditorMenus->RegisterMenu(MenuName, "SceneOutliner.DefaultContextMenuBase");
+				FEditorMenuSection& Section = Menu->AddDynamicSection("LevelEditorContextMenu", FNewEditorMenuDelegate::CreateLambda([](UEditorMenu* InMenu)
+				{
+					FName LevelContextMenuName = FLevelEditorContextMenu::GetContextMenuName(ELevelEditorMenuContext::SceneOutliner);
+					if (LevelContextMenuName != NAME_None)
+					{
+						// Extend the menu even if no actors selected, as Edit menu should always exist for scene outliner
+						UEditorMenu* OtherMenu = UEditorMenuSubsystem::Get()->GenerateMenu(LevelContextMenuName, InMenu->Context);
+						InMenu->Sections.Append(OtherMenu->Sections);
+					}
+				}));
+				Section.InsertPosition = FEditorMenuInsert("MainSection", EEditorMenuInsertType::Before);
+			}
+
 			TWeakPtr<SLevelEditor> WeakLevelEditor = SharedThis(this);
-			InitOptions.DefaultMenuExtender = MakeShareable(new FExtender);
-			InitOptions.DefaultMenuExtender->AddMenuExtension(
-				"MainSection", EExtensionHook::Before, GetLevelEditorActions(),
-				FMenuExtensionDelegate::CreateStatic([](FMenuBuilder& MenuBuilder, TWeakPtr<SLevelEditor> InWeakLevelEditor){
-					// Extend the menu even if no actors selected, as Edit menu should always exist for scene outliner
-					FLevelEditorContextMenu::FillMenu(MenuBuilder, InWeakLevelEditor, LevelEditorMenuContext::SceneOutliner, TSharedPtr<FExtender>());
-				}, WeakLevelEditor)
-			);
+			InitOptions.ModifyContextMenu.BindLambda([=](FName& OutMenuName, FEditorMenuContext& MenuContext)
+			{
+				OutMenuName = MenuName;
+
+				if (WeakLevelEditor.IsValid())
+				{
+					FLevelEditorContextMenu::InitMenuContext(MenuContext, WeakLevelEditor, ELevelEditorMenuContext::SceneOutliner);
+				}
+			});
 		}
 
 
@@ -1492,7 +1512,7 @@ void SLevelEditor::OnLayoutHasChanged()
 
 void SLevelEditor::SummonLevelViewportContextMenu()
 {
-	FLevelEditorContextMenu::SummonMenu( SharedThis( this ), LevelEditorMenuContext::Viewport );
+	FLevelEditorContextMenu::SummonMenu( SharedThis( this ), ELevelEditorMenuContext::Viewport );
 }
 
 void SLevelEditor::SummonLevelViewportViewOptionMenu(const ELevelViewportType ViewOption)
diff --git a/Engine/Source/Editor/LevelEditor/Public/LevelEditorMenuContext.h b/Engine/Source/Editor/LevelEditor/Public/LevelEditorMenuContext.h
--- a/Engine/Source/Editor/LevelEditor/Public/LevelEditorMenuContext.h
+++ b/Engine/Source/Editor/LevelEditor/Public/LevelEditorMenuContext.h
@@ -9,11 +9,33 @@
 #include "LevelEditorMenuContext.generated.h"
 
 class SLevelEditor;
+class UActorComponent;
 
 UCLASS()
-class ULevelEditorMenuContext : public UObject
+class LEVELEDITOR_API ULevelEditorMenuContext : public UObject
 {
 	GENERATED_BODY()
 public:
-	TWeakPtr<SLevelEditor> SlateLevelEditor;
+	TWeakPtr<SLevelEditor> LevelEditor;
 };
+
+/** Enum to describe what a level editor context menu should be built for */
+enum class ELevelEditorMenuContext
+{
+	/** This context menu is applicable to a viewport */
+	Viewport,
+	/** This context menu is applicable to the Scene Outliner (disables click-position-based menu items) */
+	SceneOutliner,
+};
+
+UCLASS()
+class LEVELEDITOR_API ULevelEditorContextMenuContext : public UObject
+{
+	GENERATED_BODY()
+public:
+
+	TWeakPtr<SLevelEditor> LevelEditor;
+	TArray<UActorComponent*> SelectedComponents;
+	ELevelEditorMenuContext ContextType;
+};
+
diff --git a/Engine/Source/Editor/SceneOutliner/Private/FolderTreeItem.cpp b/Engine/Source/Editor/SceneOutliner/Private/FolderTreeItem.cpp
--- a/Engine/Source/Editor/SceneOutliner/Private/FolderTreeItem.cpp
+++ b/Engine/Source/Editor/SceneOutliner/Private/FolderTreeItem.cpp
@@ -3,7 +3,7 @@
 #include "FolderTreeItem.h"
 #include "Textures/SlateIcon.h"
 #include "Framework/Commands/UIAction.h"
-#include "Framework/MultiBox/MultiBoxBuilder.h"
+#include "EditorMenuSubsystem.h"
 #include "EditorStyleSet.h"
 #include "ScopedTransaction.h"
 #include "DragAndDrop/ActorDragDropGraphEdOp.h"
@@ -320,14 +320,15 @@ void FFolderTreeItem::OnExpansionChanged()
 	}
 }
 
-void FFolderTreeItem::GenerateContextMenu(FMenuBuilder& MenuBuilder, SSceneOutliner& Outliner)
+void FFolderTreeItem::GenerateContextMenu(UEditorMenu* Menu, SSceneOutliner& Outliner)
 {
 	auto SharedOutliner = StaticCastSharedRef<SSceneOutliner>(Outliner.AsShared());
 
 	const FSlateIcon NewFolderIcon(FEditorStyle::GetStyleSetName(), "SceneOutliner.NewFolderIcon");
 	
-	MenuBuilder.AddMenuEntry(LOCTEXT("CreateSubFolder", "Create Sub Folder"), FText(), NewFolderIcon, FUIAction(FExecuteAction::CreateSP(this, &FFolderTreeItem::CreateSubFolder, TWeakPtr<SSceneOutliner>(SharedOutliner))));
-	MenuBuilder.AddMenuEntry(LOCTEXT("DuplicateFolderHierarchy", "Duplicate Hierarchy"), FText(), FSlateIcon(), FUIAction(FExecuteAction::CreateSP(&Outliner, &SSceneOutliner::DuplicateFoldersHierarchy)));
+	FEditorMenuSection& Section = Menu->AddSection("Section");
+	Section.AddMenuEntry("CreateSubFolder", LOCTEXT("CreateSubFolder", "Create Sub Folder"), FText(), NewFolderIcon, FUIAction(FExecuteAction::CreateSP(this, &FFolderTreeItem::CreateSubFolder, TWeakPtr<SSceneOutliner>(SharedOutliner))));
+	Section.AddMenuEntry("DuplicateFolderHierarchy", LOCTEXT("DuplicateFolderHierarchy", "Duplicate Hierarchy"), FText(), FSlateIcon(), FUIAction(FExecuteAction::CreateSP(&Outliner, &SSceneOutliner::DuplicateFoldersHierarchy)));
 }
 
 void FFolderTreeItem::PopulateDragDropPayload(FDragDropPayload& Payload) const
diff --git a/Engine/Source/Editor/SceneOutliner/Private/SSceneOutliner.cpp b/Engine/Source/Editor/SceneOutliner/Private/SSceneOutliner.cpp
--- a/Engine/Source/Editor/SceneOutliner/Private/SSceneOutliner.cpp
+++ b/Engine/Source/Editor/SceneOutliner/Private/SSceneOutliner.cpp
@@ -12,6 +12,7 @@
 #include "Widgets/Layout/SSeparator.h"
 #include "Widgets/Images/SImage.h"
 #include "Framework/MultiBox/MultiBoxBuilder.h"
+#include "EditorMenuSubsystem.h"
 #include "Widgets/Input/SButton.h"
 #include "Widgets/Input/SComboButton.h"
 #include "EditorStyleSet.h"
@@ -34,6 +35,7 @@
 #include "Widgets/Input/SSearchBox.h"
 
 #include "SceneOutlinerDragDrop.h"
+#include "SceneOutlinerMenuContext.h"
 
 
 #include "ActorEditorUtils.h"
@@ -1706,83 +1708,125 @@ namespace SceneOutliner
 			return nullptr;
 		}
 
+		RegisterDefaultContextMenu();
+
 		FItemSelection ItemSelection(*OutlinerTreeView);
 
-		const auto NumSelectedItems = OutlinerTreeView->GetNumItemsSelected();
-		bool bPasteCommandOnly = (NumSelectedItems == 0);
+		USceneOutlinerMenuContext* ContextObject = NewObject<USceneOutlinerMenuContext>();
+		ContextObject->SceneOutliner = SharedThis(this);
+		ContextObject->bShowParentTree = SharedData->bShowParentTree;
+		ContextObject->NumSelectedItems = OutlinerTreeView->GetNumItemsSelected();
+		ContextObject->NumSelectedFolders = ItemSelection.Folders.Num();
+		ContextObject->NumWorldsSelected = ItemSelection.Worlds.Num();
+		FEditorMenuContext Context(ContextObject);
 
-		// Build up the menu for a selection
-		const bool bCloseAfterSelection = true;
-		FMenuBuilder MenuBuilder(bCloseAfterSelection, TSharedPtr<FUICommandList>(), SharedData->DefaultMenuExtender);
+		// Allow other systems to override menu name and provide additional context
+		static const FName DefaultContextMenuName("SceneOutliner.DefaultContextMenu");
+		FName MenuName = DefaultContextMenuName;
+		SharedData->ModifyContextMenu.ExecuteIfBound(MenuName, Context);
 
-		bool MenuBuilderHasContent = false;
+		// Build up the menu for a selection
+		UEditorMenuSubsystem* EditorMenus = UEditorMenuSubsystem::Get();
+		UEditorMenu* Menu = EditorMenus->GenerateMenu(MenuName, Context);
 
-		if (SharedData->bShowParentTree)
+		for (const FEditorMenuSection& Section : Menu->Sections)
 		{
-			if (NumSelectedItems == 0)
+			if (Section.Blocks.Num() > 0)
 			{
-				const FSlateIcon NewFolderIcon(FEditorStyle::GetStyleSetName(), "SceneOutliner.NewFolderIcon");
-				MenuBuilder.AddMenuEntry(LOCTEXT("CreateFolder", "Create Folder"), FText(), NewFolderIcon, FUIAction(FExecuteAction::CreateSP(this, &SSceneOutliner::CreateFolder)));
-				MenuBuilderHasContent = true;
+				return EditorMenus->GenerateWidget(Menu);
 			}
-			else
+		}
+
+		return nullptr;
+	}
+
+	void SSceneOutliner::RegisterDefaultContextMenu()
+	{
+		static const FName DefaultContextBaseMenuName("SceneOutliner.DefaultContextMenuBase");
+		static const FName DefaultContextMenuName("SceneOutliner.DefaultContextMenu");
+
+		UEditorMenuSubsystem* EditorMenus = UEditorMenuSubsystem::Get();
+
+		if (!EditorMenus->IsMenuRegistered(DefaultContextBaseMenuName))
+		{
+			UEditorMenu* Menu = EditorMenus->RegisterMenu(DefaultContextBaseMenuName);
+
+			Menu->AddDynamicSection("DynamicSection1", FNewEditorMenuDelegate::CreateLambda([](UEditorMenu* InMenu)
 			{
-				if (NumSelectedItems == 1)
+				USceneOutlinerMenuContext* Context = InMenu->FindContext<USceneOutlinerMenuContext>();
+				if (!Context || !Context->SceneOutliner.IsValid())
 				{
-					OutlinerTreeView->GetSelectedItems()[0]->GenerateContextMenu(MenuBuilder, *this);
-					MenuBuilderHasContent = true;
+					return;
 				}
 
-				// If we've only got folders selected, show the selection and edit sub menus
-				if (NumSelectedItems > 0 && ItemSelection.Folders.Num() == NumSelectedItems)
+				SSceneOutliner* SceneOutliner = Context->SceneOutliner.Pin().Get();
+				if (Context->bShowParentTree)
 				{
-					MenuBuilder.AddSubMenu(
-						LOCTEXT("SelectSubmenu", "Select"),
-						LOCTEXT("SelectSubmenu_Tooltip", "Select the contents of the current selection"),
-						FNewMenuDelegate::CreateSP(this, &SSceneOutliner::FillSelectionSubMenu));
-					MenuBuilderHasContent = true;
+					if (Context->NumSelectedItems == 0)
+					{
+						InMenu->FindOrAddSection("Section").AddMenuEntry(
+							"CreateFolder",
+							LOCTEXT("CreateFolder", "Create Folder"),
+							FText(),
+							FSlateIcon(FEditorStyle::GetStyleSetName(), "SceneOutliner.NewFolderIcon"),
+							FUIAction(FExecuteAction::CreateSP(SceneOutliner, &SSceneOutliner::CreateFolder)));
+					}
+					else
+					{
+						if (Context->NumSelectedItems == 1)
+						{
+							SceneOutliner->GetTree().GetSelectedItems()[0]->GenerateContextMenu(InMenu, *SceneOutliner);
+						}
+
+						// If we've only got folders selected, show the selection and edit sub menus
+						if (Context->NumSelectedItems > 0 && Context->NumSelectedFolders == Context->NumSelectedItems)
+						{
+							InMenu->FindOrAddSection("Section").AddEntry(FEditorMenuEntry::InitSubMenu(
+								InMenu->MenuName,
+								"SelectSubMenu",
+								LOCTEXT("SelectSubmenu", "Select"),
+								LOCTEXT("SelectSubmenu_Tooltip", "Select the contents of the current selection"),
+								FNewEditorMenuDelegate::CreateSP(SceneOutliner, &SSceneOutliner::FillSelectionSubMenu)));
+						}
+					}
 				}
-			}
-		}
+			}));
 
-			// We always create a section here, even if there is no parent so that clients can still extend the menu
-			MenuBuilder.BeginSection("MainSection");
+			Menu->AddDynamicSection("DynamicMainSection", FNewEditorMenuDelegate::CreateLambda([](UEditorMenu* InMenu)
 			{
-				// Don't add any of these menu items if we're not showing the parent tree
-				if (SharedData->bShowParentTree)
+				// We always create a section here, even if there is no parent so that clients can still extend the menu
+				FEditorMenuSection& Section = InMenu->AddSection("MainSection");
+
+				if (USceneOutlinerMenuContext* Context = InMenu->FindContext<USceneOutlinerMenuContext>())
 				{
+					// Don't add any of these menu items if we're not showing the parent tree
 					// Can't move worlds or level blueprints
-					if (NumSelectedItems > 0)
+					if (Context->bShowParentTree && Context->NumSelectedItems > 0 && Context->NumWorldsSelected == 0 && Context->SceneOutliner.IsValid())
 					{
-						const bool bCanMoveSelection = ItemSelection.Worlds.Num() == 0;
-						if (bCanMoveSelection)
-						{
-							MenuBuilder.AddSubMenu(
-								LOCTEXT("MoveActorsTo", "Move To"),
-								LOCTEXT("MoveActorsTo_Tooltip", "Move selection to another folder"),
-								FNewMenuDelegate::CreateSP(this, &SSceneOutliner::FillFoldersSubMenu));
-							MenuBuilderHasContent = true;
-						}
+						Section.AddEntry(FEditorMenuEntry::InitSubMenu(
+							InMenu->MenuName,
+							"MoveActorsTo",
+							LOCTEXT("MoveActorsTo", "Move To"),
+							LOCTEXT("MoveActorsTo_Tooltip", "Move selection to another folder"),
+							FNewEditorMenuDelegate::CreateSP(Context->SceneOutliner.Pin().Get(), &SSceneOutliner::FillFoldersSubMenu)));
 					}
 				}
-			}
-			MenuBuilder.EndSection();
-		
+			}));
+		}
 
-		if (MenuBuilderHasContent)
+		if (!EditorMenus->IsMenuRegistered(DefaultContextMenuName))
 		{
-			return MenuBuilder.MakeWidget();
+			EditorMenus->RegisterMenu(DefaultContextMenuName, DefaultContextBaseMenuName);
 		}
-
-		return nullptr;
 	}
 
-	void SSceneOutliner::FillFoldersSubMenu(FMenuBuilder& MenuBuilder) const
+	void SSceneOutliner::FillFoldersSubMenu(UEditorMenu* Menu) const
 	{
-		MenuBuilder.AddMenuEntry(LOCTEXT( "CreateNew", "Create New Folder" ), LOCTEXT( "CreateNew_ToolTip", "Move to a new folder" ),
+		FEditorMenuSection& Section = Menu->AddSection("Section");
+		Section.AddMenuEntry("CreateNew", LOCTEXT( "CreateNew", "Create New Folder" ), LOCTEXT( "CreateNew_ToolTip", "Move to a new folder" ),
 			FSlateIcon(FEditorStyle::GetStyleSetName(), "SceneOutliner.NewFolderIcon"), FExecuteAction::CreateSP(const_cast<SSceneOutliner*>(this), &SSceneOutliner::CreateFolder));
 
-		AddMoveToFolderOutliner(MenuBuilder);
+		AddMoveToFolderOutliner(Menu);
 	}
 
 	TSharedRef<TSet<FName>> SSceneOutliner::GatherInvalidMoveToDestinations() const
@@ -1849,7 +1893,7 @@ namespace SceneOutliner
 		return ExcludedParents;
 	}
 
-	void SSceneOutliner::AddMoveToFolderOutliner(FMenuBuilder& MenuBuilder) const
+	void SSceneOutliner::AddMoveToFolderOutliner(UEditorMenu* Menu) const
 	{
 		// We don't show this if there aren't any folders in the world
 		if (!FActorFolders::Get().GetFolderPropertiesForWorld(*SharedData->RepresentingWorld).Num())
@@ -1909,19 +1953,25 @@ namespace SceneOutliner
 				.OnItemPickedDelegate(FOnSceneOutlinerItemPicked::CreateSP(const_cast<SSceneOutliner*>(this), &SSceneOutliner::MoveSelectionTo))
 			];
 
-		MenuBuilder.BeginSection(FName(), LOCTEXT("ExistingFolders", "Existing:"));
-		MenuBuilder.AddWidget(MiniSceneOutliner, FText::GetEmpty(), false);
-		MenuBuilder.EndSection();
+		FEditorMenuSection& Section = Menu->AddSection(FName(), LOCTEXT("ExistingFolders", "Existing:"));
+		Section.AddEntry(FEditorMenuEntry::InitWidget(
+			"MiniSceneOutliner",
+			MiniSceneOutliner,
+			FText::GetEmpty(),
+			false));
 	}
 
-	void SSceneOutliner::FillSelectionSubMenu(FMenuBuilder& MenuBuilder) const
+	void SSceneOutliner::FillSelectionSubMenu(UEditorMenu* Menu) const
 	{
-		MenuBuilder.AddMenuEntry(
+		FEditorMenuSection& Section = Menu->AddSection("Section");
+		Section.AddMenuEntry(
+			"AddChildrenToSelection",
 			LOCTEXT( "AddChildrenToSelection", "Immediate Children" ),
 			LOCTEXT( "AddChildrenToSelection_ToolTip", "Select all immediate actor children of the selected folders" ),
 			FSlateIcon(),
 			FExecuteAction::CreateSP(const_cast<SSceneOutliner*>(this), &SSceneOutliner::SelectFoldersDescendants, /*bSelectImmediateChildrenOnly=*/ true));
-		MenuBuilder.AddMenuEntry(
+		Section.AddMenuEntry(
+			"AddDescendantsToSelection",
 			LOCTEXT( "AddDescendantsToSelection", "All Descendants" ),
 			LOCTEXT( "AddDescendantsToSelection_ToolTip", "Select all actor descendants of the selected folders" ),
 			FSlateIcon(),
diff --git a/Engine/Source/Editor/SceneOutliner/Private/WorldTreeItem.cpp b/Engine/Source/Editor/SceneOutliner/Private/WorldTreeItem.cpp
--- a/Engine/Source/Editor/SceneOutliner/Private/WorldTreeItem.cpp
+++ b/Engine/Source/Editor/SceneOutliner/Private/WorldTreeItem.cpp
@@ -9,7 +9,7 @@
 #include "SceneOutlinerDragDrop.h"
 #include "SSceneOutliner.h"
 
-#include "Framework/MultiBox/MultiBoxBuilder.h"
+#include "EditorMenuSubsystem.h"
 #include "LevelEditor.h"
 #include "EditorActorFolders.h"
 
@@ -82,15 +82,15 @@ bool FWorldTreeItem::CanInteract() const
 	return Flags.bInteractive;
 }
 
-void FWorldTreeItem::GenerateContextMenu(FMenuBuilder& MenuBuilder, SSceneOutliner& Outliner)
+void FWorldTreeItem::GenerateContextMenu(UEditorMenu* Menu, SSceneOutliner& Outliner)
 {
 	auto SharedOutliner = StaticCastSharedRef<SSceneOutliner>(Outliner.AsShared());
 	
 	const FSlateIcon WorldSettingsIcon(FEditorStyle::GetStyleSetName(), "LevelEditor.WorldProperties.Tab");
 	const FSlateIcon NewFolderIcon(FEditorStyle::GetStyleSetName(), "SceneOutliner.NewFolderIcon");
-
-	MenuBuilder.AddMenuEntry(LOCTEXT("CreateFolder", "Create Folder"), FText(), NewFolderIcon, FUIAction(FExecuteAction::CreateSP(this, &FWorldTreeItem::CreateFolder, TWeakPtr<SSceneOutliner>(SharedOutliner))));
-	MenuBuilder.AddMenuEntry(LOCTEXT("OpenWorldSettings", "World Settings"), FText(), WorldSettingsIcon, FExecuteAction::CreateSP(this, &FWorldTreeItem::OpenWorldSettings));
+	FEditorMenuSection& Section = Menu->AddSection("Section");
+	Section.AddMenuEntry("CreateFolder", LOCTEXT("CreateFolder", "Create Folder"), FText(), NewFolderIcon, FUIAction(FExecuteAction::CreateSP(this, &FWorldTreeItem::CreateFolder, TWeakPtr<SSceneOutliner>(SharedOutliner))));
+	Section.AddMenuEntry("OpenWorldSettings", LOCTEXT("OpenWorldSettings", "World Settings"), FText(), WorldSettingsIcon, FExecuteAction::CreateSP(this, &FWorldTreeItem::OpenWorldSettings));
 }
 
 void FWorldTreeItem::CreateFolder(TWeakPtr<SSceneOutliner> WeakOutliner)
diff --git a/Engine/Source/Editor/SceneOutliner/Public/FolderTreeItem.h b/Engine/Source/Editor/SceneOutliner/Public/FolderTreeItem.h
--- a/Engine/Source/Editor/SceneOutliner/Public/FolderTreeItem.h
+++ b/Engine/Source/Editor/SceneOutliner/Public/FolderTreeItem.h
@@ -6,7 +6,7 @@
 #include "Widgets/SWidget.h"
 #include "ITreeItem.h"
 
-class FMenuBuilder;
+class UEditorMenu;
 
 namespace SceneOutliner
 {
@@ -68,7 +68,7 @@ namespace SceneOutliner
 		virtual bool CanInteract() const override;
 		
 		/** Generate a context menu for this item. Only called if *only* this item is selected. */
-		virtual void GenerateContextMenu(FMenuBuilder& MenuBuilder, SSceneOutliner& Outliner) override;
+		virtual void GenerateContextMenu(UEditorMenu* Menu, SSceneOutliner& Outliner) override;
 
 		/** Populate the specified drag/drop payload with any relevant information for this type */
 		virtual void PopulateDragDropPayload(FDragDropPayload& Payload) const override;
diff --git a/Engine/Source/Editor/SceneOutliner/Public/ITreeItem.h b/Engine/Source/Editor/SceneOutliner/Public/ITreeItem.h
--- a/Engine/Source/Editor/SceneOutliner/Public/ITreeItem.h
+++ b/Engine/Source/Editor/SceneOutliner/Public/ITreeItem.h
@@ -6,7 +6,7 @@
 #include "Widgets/SWidget.h"
 #include "SceneOutlinerVisitorTypes.h"
 
-class FMenuBuilder;
+class UEditorMenu;
 
 namespace SceneOutliner
 {
@@ -156,7 +156,7 @@ namespace SceneOutliner
 		virtual void OnExpansionChanged() {};
 
 		/** Generate a context menu for this item. Only called if *only* this item is selected. */
-		virtual void GenerateContextMenu(FMenuBuilder& MenuBuilder, SSceneOutliner& Outliner) {}
+		virtual void GenerateContextMenu(UEditorMenu* Menu, SSceneOutliner& Outliner) {}
 
 		/** Populate the specified drag/drop payload with any relevant information for this type */
 		virtual void PopulateDragDropPayload(FDragDropPayload& Payload) const = 0;
diff --git a/Engine/Source/Editor/SceneOutliner/Public/SSceneOutliner.h b/Engine/Source/Editor/SceneOutliner/Public/SSceneOutliner.h
--- a/Engine/Source/Editor/SceneOutliner/Public/SSceneOutliner.h
+++ b/Engine/Source/Editor/SceneOutliner/Public/SSceneOutliner.h
@@ -20,6 +20,7 @@
 #include "Framework/Commands/UICommandList.h"
 
 class FMenuBuilder;
+class UEditorMenu;
 class ISceneOutlinerColumn;
 class SComboButton;
 
@@ -209,6 +210,7 @@ namespace SceneOutliner
 			}
 		}
 
+		void RegisterDefaultContextMenu();
 
 		/** Visitor that is used to set up type-specific data after tree items are added to the tree */
 		struct FOnItemAddedToTree : IMutableTreeItemVisitor
@@ -495,9 +497,9 @@ namespace SceneOutliner
 
 		/** Build a context menu for right-clicking an item in the tree */
 		TSharedPtr<SWidget> BuildDefaultContextMenu();
-		void FillFoldersSubMenu(FMenuBuilder& MenuBuilder) const;
-		void AddMoveToFolderOutliner(FMenuBuilder& MenuBuilder) const;
-		void FillSelectionSubMenu(FMenuBuilder& MenuBuilder) const;
+		void FillFoldersSubMenu(UEditorMenu* Menu) const;
+		void AddMoveToFolderOutliner(UEditorMenu* Menu) const;
+		void FillSelectionSubMenu(UEditorMenu* Menun) const;
 		TSharedRef<TSet<FName>> GatherInvalidMoveToDestinations() const;
 
 	private:
diff --git a/Engine/Source/Editor/SceneOutliner/Public/SceneOutlinerMenuContext.h b/Engine/Source/Editor/SceneOutliner/Public/SceneOutlinerMenuContext.h
new file mode 100644
--- /dev/null
+++ b/Engine/Source/Editor/SceneOutliner/Public/SceneOutlinerMenuContext.h
@@ -0,0 +1,36 @@
+// Copyright 1998-2019 Epic Games, Inc. All Rights Reserved.
+#pragma once
+
+#include "CoreMinimal.h"
+#include "UObject/ObjectMacros.h"
+#include "UObject/Object.h"
+
+#include "SceneOutlinerMenuContext.generated.h"
+
+namespace SceneOutliner
+{
+	class SSceneOutliner;
+}
+
+UCLASS()
+class USceneOutlinerMenuContext : public UObject
+{
+	GENERATED_BODY()
+public:
+
+	USceneOutlinerMenuContext(const FObjectInitializer& ObjectInitializer)
+		: Super(ObjectInitializer)
+	{
+		bShowParentTree = false;
+		NumSelectedItems = 0;
+		NumSelectedFolders = 0;
+		NumWorldsSelected = 0;
+	}
+
+	TWeakPtr<SceneOutliner::SSceneOutliner> SceneOutliner;
+
+	bool bShowParentTree;
+	int32 NumSelectedItems;
+	int32 NumSelectedFolders;
+	int32 NumWorldsSelected;
+};
diff --git a/Engine/Source/Editor/SceneOutliner/Public/SceneOutlinerPublicTypes.h b/Engine/Source/Editor/SceneOutliner/Public/SceneOutlinerPublicTypes.h
--- a/Engine/Source/Editor/SceneOutliner/Public/SceneOutlinerPublicTypes.h
+++ b/Engine/Source/Editor/SceneOutliner/Public/SceneOutlinerPublicTypes.h
@@ -10,6 +10,7 @@
 #include "SceneOutlinerFilters.h"
 
 class FExtender;
+struct FEditorMenuContext;
 
 enum class ESceneOutlinerMode : uint8
 {
@@ -23,6 +24,8 @@ enum class ESceneOutlinerMode : uint8
 	ComponentPicker,
 };
 
+DECLARE_DELEGATE_TwoParams(FSceneOutlinerModifyContextMenu, FName& /* MenuName */, FEditorMenuContext& /* MenuContext */);
+
 namespace SceneOutliner
 {
 	/** Container for built in column types. Function-static so they are available without linking */
@@ -103,8 +106,8 @@ namespace SceneOutliner
 		/** Override default context menu handling */
 		FOnContextMenuOpening ContextMenuOverride;
 
-		/** Extend default context menu handling */
-		TSharedPtr<FExtender> DefaultMenuExtender;
+		/** Modify context menu before display */
+		FSceneOutlinerModifyContextMenu ModifyContextMenu;
 
 		/** Map of column types available to the scene outliner, along with default ordering */
 		TMap<FName, FColumnInfo> ColumnMap;
diff --git a/Engine/Source/Editor/SceneOutliner/Public/WorldTreeItem.h b/Engine/Source/Editor/SceneOutliner/Public/WorldTreeItem.h
--- a/Engine/Source/Editor/SceneOutliner/Public/WorldTreeItem.h
+++ b/Engine/Source/Editor/SceneOutliner/Public/WorldTreeItem.h
@@ -7,7 +7,7 @@
 #include "ITreeItem.h"
 #include "UObject/ObjectKey.h"
 
-class FMenuBuilder;
+class UEditorMenu;
 
 namespace SceneOutliner
 {
@@ -48,7 +48,7 @@ namespace SceneOutliner
 		virtual bool CanInteract() const override;
 
 		/** Generate a context menu for this item. Only called if *only* this item is selected. */
-		virtual void GenerateContextMenu(FMenuBuilder& MenuBuilder, SSceneOutliner& Outliner) override;
+		virtual void GenerateContextMenu(UEditorMenu* Menu, SSceneOutliner& Outliner) override;
 
 		/** Populate the specified drag/drop payload with any relevant information for this type */
 		virtual void PopulateDragDropPayload(FDragDropPayload& Payload) const override
diff --git a/Engine/Source/Editor/SceneOutliner/SceneOutliner.Build.cs b/Engine/Source/Editor/SceneOutliner/SceneOutliner.Build.cs
--- a/Engine/Source/Editor/SceneOutliner/SceneOutliner.Build.cs
+++ b/Engine/Source/Editor/SceneOutliner/SceneOutliner.Build.cs
@@ -22,6 +22,7 @@ public class SceneOutliner : ModuleRules
 				"SlateCore",
 				"EditorStyle",
 				"UnrealEd",
+				"EditorMenus",
             }
 		);
 	}
diff --git a/Engine/Source/Editor/UnrealEd/Private/Kismet2/ComponentEditorUtils.cpp b/Engine/Source/Editor/UnrealEd/Private/Kismet2/ComponentEditorUtils.cpp
--- a/Engine/Source/Editor/UnrealEd/Private/Kismet2/ComponentEditorUtils.cpp
+++ b/Engine/Source/Editor/UnrealEd/Private/Kismet2/ComponentEditorUtils.cpp
@@ -27,6 +27,8 @@
 #include "Editor.h"
 #include "Toolkits/AssetEditorManager.h"
 #include "HAL/PlatformApplicationMisc.h"
+#include "EditorMenuSubsystem.h"
+#include "Kismet2/ComponentEditorContextMenuContex.h"
 
 #define LOCTEXT_NAMESPACE "ComponentEditorUtils"
 
@@ -913,19 +915,18 @@ FName FComponentEditorUtils::FindVariableNameGivenComponentInstance(const UActor
 	return NAME_None;
 }
 
-void FComponentEditorUtils::FillComponentContextMenuOptions(FMenuBuilder& MenuBuilder, const TArray<UActorComponent*>& SelectedComponents)
+void FComponentEditorUtils::FillComponentContextMenuOptions(UEditorMenu* Menu, const TArray<UActorComponent*>& SelectedComponents)
 {
 	// Basic commands
-	MenuBuilder.BeginSection("EditComponent", LOCTEXT("EditComponentHeading", "Edit"));
 	{
-		MenuBuilder.AddMenuEntry(FGenericCommands::Get().Cut);
-		MenuBuilder.AddMenuEntry(FGenericCommands::Get().Copy);
-		MenuBuilder.AddMenuEntry(FGenericCommands::Get().Paste);
-		MenuBuilder.AddMenuEntry(FGenericCommands::Get().Duplicate);
-		MenuBuilder.AddMenuEntry(FGenericCommands::Get().Delete);
-		MenuBuilder.AddMenuEntry(FGenericCommands::Get().Rename);
+		FEditorMenuSection& Section = Menu->AddSection("EditComponent", LOCTEXT("EditComponentHeading", "Edit"));
+		Section.AddMenuEntry(FGenericCommands::Get().Cut);
+		Section.AddMenuEntry(FGenericCommands::Get().Copy);
+		Section.AddMenuEntry(FGenericCommands::Get().Paste);
+		Section.AddMenuEntry(FGenericCommands::Get().Duplicate);
+		Section.AddMenuEntry(FGenericCommands::Get().Delete);
+		Section.AddMenuEntry(FGenericCommands::Get().Rename);
 	}
-	MenuBuilder.EndSection();
 
 	if (SelectedComponents.Num() == 1)
 	{
@@ -933,9 +934,10 @@ void FComponentEditorUtils::FillComponentContextMenuOptions(FMenuBuilder& MenuBu
 
 		if (Component->GetClass()->ClassGeneratedBy)
 		{
-			MenuBuilder.BeginSection("ComponentAsset", LOCTEXT("ComponentAssetHeading", "Asset"));
 			{
-				MenuBuilder.AddMenuEntry(
+				FEditorMenuSection& Section = Menu->AddSection("ComponentAsset", LOCTEXT("ComponentAssetHeading", "Asset"));
+				Section.AddMenuEntry(
+					"GoToBlueprintForComponent",
 					FText::Format(LOCTEXT("GoToBlueprintForComponent", "Edit {0}"), FText::FromString(Component->GetClass()->ClassGeneratedBy->GetName())),
 					LOCTEXT("EditBlueprintForComponent_ToolTip", "Edits the Blueprint Class that defines this component."),
 					FSlateIconFinder::FindIconForClass(Component->GetClass()),
@@ -943,7 +945,8 @@ void FComponentEditorUtils::FillComponentContextMenuOptions(FMenuBuilder& MenuBu
 					FExecuteAction::CreateStatic(&FComponentEditorUtils::OnEditBlueprintComponent, Component->GetClass()->ClassGeneratedBy),
 					FCanExecuteAction()));
 
-				MenuBuilder.AddMenuEntry(
+				Section.AddMenuEntry(
+					"GoToAssetForComponent",
 					LOCTEXT("GoToAssetForComponent", "Find Class in Content Browser"),
 					LOCTEXT("GoToAssetForComponent_ToolTip", "Summons the content browser and goes to the class for this component."),
 					FSlateIcon(FEditorStyle::GetStyleSetName(), "SystemWideCommands.FindInContentBrowser"),
@@ -951,12 +954,11 @@ void FComponentEditorUtils::FillComponentContextMenuOptions(FMenuBuilder& MenuBu
 					FExecuteAction::CreateStatic(&FComponentEditorUtils::OnGoToComponentAssetInBrowser, Component->GetClass()->ClassGeneratedBy),
 					FCanExecuteAction()));
 			}
-			MenuBuilder.EndSection();
 		}
 		else
 		{
-			MenuBuilder.BeginSection("ComponentCode", LOCTEXT("ComponentCodeHeading", "C++"));
 			{
+				FEditorMenuSection& Section = Menu->AddSection("ComponentCode", LOCTEXT("ComponentCodeHeading", "C++"));
 				if (FSourceCodeNavigation::IsCompilerAvailable())
 				{
 					FString ClassHeaderPath;
@@ -964,7 +966,8 @@ void FComponentEditorUtils::FillComponentContextMenuOptions(FMenuBuilder& MenuBu
 					{
 						const FString CodeFileName = FPaths::GetCleanFilename(*ClassHeaderPath);
 
-						MenuBuilder.AddMenuEntry(
+						Section.AddMenuEntry(
+							"GoToCodeForComponent",
 							FText::Format(LOCTEXT("GoToCodeForComponent", "Open {0}"), FText::FromString(CodeFileName)),
 							FText::Format(LOCTEXT("GoToCodeForComponent_ToolTip", "Opens the header file for this component ({0}) in a code editing program"), FText::FromString(CodeFileName)),
 							FSlateIcon(),
@@ -973,7 +976,8 @@ void FComponentEditorUtils::FillComponentContextMenuOptions(FMenuBuilder& MenuBu
 							FCanExecuteAction()));
 					}
 
-					MenuBuilder.AddMenuEntry(
+					Section.AddMenuEntry(
+						"GoToAssetForComponent",
 						LOCTEXT("GoToAssetForComponent", "Find Class in Content Browser"),
 						LOCTEXT("GoToAssetForComponent_ToolTip", "Summons the content browser and goes to the class for this component."),
 						FSlateIcon(FEditorStyle::GetStyleSetName(), "SystemWideCommands.FindInContentBrowser"),
@@ -982,7 +986,6 @@ void FComponentEditorUtils::FillComponentContextMenuOptions(FMenuBuilder& MenuBu
 						FCanExecuteAction()));
 				}
 			}
-			MenuBuilder.EndSection();
 		}
 	}
 }
diff --git a/Engine/Source/Editor/UnrealEd/Private/Kismet2/Kismet2.cpp b/Engine/Source/Editor/UnrealEd/Private/Kismet2/Kismet2.cpp
--- a/Engine/Source/Editor/UnrealEd/Private/Kismet2/Kismet2.cpp
+++ b/Engine/Source/Editor/UnrealEd/Private/Kismet2/Kismet2.cpp
@@ -81,6 +81,7 @@
 #include "Widgets/Notifications/SNotificationList.h"
 #include "Engine/InheritableComponentHandler.h"
 #include "Classes/EditorStyleSettings.h"
+#include "EditorMenuSubsystem.h"
 
 DECLARE_CYCLE_STAT(TEXT("Compile Blueprint"), EKismetCompilerStats_CompileBlueprint, STATGROUP_KismetCompiler);
 DECLARE_CYCLE_STAT(TEXT("Broadcast Precompile"), EKismetCompilerStats_BroadcastPrecompile, STATGROUP_KismetCompiler);
@@ -2448,7 +2449,7 @@ bool FKismetEditorUtilities::AnyBoundLevelScriptEventForActor(AActor* Actor, boo
 	return false;
 }
 
-void FKismetEditorUtilities::AddLevelScriptEventOptionsForActor(class FMenuBuilder& MenuBuilder, TWeakObjectPtr<AActor> ActorPtr, bool bExistingEvents, bool bNewEvents, bool bOnlyEventName)
+void FKismetEditorUtilities::AddLevelScriptEventOptionsForActor(UEditorMenu* Menu, TWeakObjectPtr<AActor> ActorPtr, bool bExistingEvents, bool bNewEvents, bool bOnlyEventName)
 {
 	struct FCreateEventForActorHelper
 	{
@@ -2518,7 +2519,7 @@ void FKismetEditorUtilities::AddLevelScriptEventOptionsForActor(class FMenuBuild
 		// Now build the menu
 		for(FEventCategory& Category : CategorizedEvents)
 		{
-			MenuBuilder.BeginSection(NAME_None, FText::FromString(Category.CategoryName));
+			FEditorMenuSection& Section = Menu->AddSection(NAME_None, FText::FromString(Category.CategoryName));
 
 			for(UProperty* Property : Category.EventProperties)
 			{
@@ -2551,15 +2552,14 @@ void FKismetEditorUtilities::AddLevelScriptEventOptionsForActor(class FMenuBuild
 				}
 
 				// create menu entry
-				MenuBuilder.AddMenuEntry(
+				Section.AddMenuEntry(
+					NAME_None,
 					EntryText,
 					Property->GetToolTipText(),
 					FSlateIcon(),
-					FUIAction(FExecuteAction::CreateStatic(&FCreateEventForActorHelper::CreateEventForActor, ActorPtr, EventName))
+					FExecuteAction::CreateStatic(&FCreateEventForActorHelper::CreateEventForActor, ActorPtr, EventName)
 					);
 			}
-
-			MenuBuilder.EndSection();
 		}
 	}
 }
diff --git a/Engine/Source/Editor/UnrealEd/Public/Kismet2/ComponentEditorContextMenuContex.h b/Engine/Source/Editor/UnrealEd/Public/Kismet2/ComponentEditorContextMenuContex.h
new file mode 100644
--- /dev/null
+++ b/Engine/Source/Editor/UnrealEd/Public/Kismet2/ComponentEditorContextMenuContex.h
@@ -0,0 +1,18 @@
+// Copyright 1998-2019 Epic Games, Inc. All Rights Reserved.
+
+#pragma once
+
+#include "CoreMinimal.h"
+#include "UObject/UObjectHash.h"
+#include "Components/ActorComponent.h"
+
+//#include "ComponentEditorContextMenuContex.generated.h"
+//
+//UCLASS()
+//class UNREALED_API UComponentEditorContextMenuContext : public UObject
+//{
+//	GENERATED_BODY()
+//public:
+//	
+//	TArray<UActorComponent*> SelectedComponents;
+//};
diff --git a/Engine/Source/Editor/UnrealEd/Public/Kismet2/ComponentEditorUtils.h b/Engine/Source/Editor/UnrealEd/Public/Kismet2/ComponentEditorUtils.h
--- a/Engine/Source/Editor/UnrealEd/Public/Kismet2/ComponentEditorUtils.h
+++ b/Engine/Source/Editor/UnrealEd/Public/Kismet2/ComponentEditorUtils.h
@@ -9,7 +9,7 @@
 #include "GameFramework/Actor.h"
 #include "UObject/UObjectHash.h"
 
-class FMenuBuilder;
+class UEditorMenu;
 class UMaterialInterface;
 
 class UNREALED_API FComponentEditorUtils
@@ -243,10 +243,10 @@ public:
 
 	/**
 	* Populates the given menu with basic options for operations on components in the world.
-	* @param MenuBuilder Used to create the menu options
+	* @param Menu Used to register the menu options
 	* @param SelectedComponents The selected components to create menu options for
 	*/
-	static void FillComponentContextMenuOptions(FMenuBuilder& MenuBuilder, const TArray<UActorComponent*>& SelectedComponents);
+	static void FillComponentContextMenuOptions(UEditorMenu* Menu, const TArray<UActorComponent*>& SelectedComponents);
 
 	/**
 	 * Tries to find a match for ComponentInstance in the ComponentList. First by name and then if multiple Components have a matching name try to match the SceneComponent hierarchy to find the best match.
diff --git a/Engine/Source/Editor/UnrealEd/Public/Kismet2/KismetEditorUtilities.h b/Engine/Source/Editor/UnrealEd/Public/Kismet2/KismetEditorUtilities.h
--- a/Engine/Source/Editor/UnrealEd/Public/Kismet2/KismetEditorUtilities.h
+++ b/Engine/Source/Editor/UnrealEd/Public/Kismet2/KismetEditorUtilities.h
@@ -10,6 +10,7 @@
 #include "Engine/Blueprint.h"
 #include "Engine/BlueprintGeneratedClass.h"
 
+class UEditorMenu;
 class IBlueprintEditor;
 class UEdGraph;
 struct Rect;
@@ -276,7 +277,7 @@ public:
 	static bool AnyBoundLevelScriptEventForActor(AActor* Actor, bool bCouldAddAny);
 
 	/** It lists bounded LevelScriptEvents for given actor */
-	static void AddLevelScriptEventOptionsForActor(class FMenuBuilder& MenuBuilder, TWeakObjectPtr<AActor> ActorPtr, bool bExistingEvents, bool bNewEvents, bool bOnlyEventName);
+	static void AddLevelScriptEventOptionsForActor(UEditorMenu* Menu, TWeakObjectPtr<AActor> ActorPtr, bool bExistingEvents, bool bNewEvents, bool bOnlyEventName);
 	
 	/** Return information about the given macro graph */
 	static void GetInformationOnMacro(UEdGraph* MacroGraph, /*out*/ class UK2Node_Tunnel*& EntryNode, /*out*/ class UK2Node_Tunnel*& ExitNode, bool& bIsPure);
diff --git a/Engine/Source/Editor/WorldBrowser/Private/LevelCollectionModel.cpp b/Engine/Source/Editor/WorldBrowser/Private/LevelCollectionModel.cpp
--- a/Engine/Source/Editor/WorldBrowser/Private/LevelCollectionModel.cpp
+++ b/Engine/Source/Editor/WorldBrowser/Private/LevelCollectionModel.cpp
@@ -1432,7 +1432,7 @@ FBox FLevelCollectionModel::GetVisibleLevelsBoundingBox(const FLevelModelList& I
 	return TotalBounds;
 }
 
-const TSharedRef<const FUICommandList> FLevelCollectionModel::GetCommandList() const
+const TSharedRef<FUICommandList> FLevelCollectionModel::GetCommandList() const
 {
 	return CommandList;
 }
diff --git a/Engine/Source/Editor/WorldBrowser/Private/LevelCollectionModel.h b/Engine/Source/Editor/WorldBrowser/Private/LevelCollectionModel.h
--- a/Engine/Source/Editor/WorldBrowser/Private/LevelCollectionModel.h
+++ b/Engine/Source/Editor/WorldBrowser/Private/LevelCollectionModel.h
@@ -331,7 +331,7 @@ public:
 	static FBox GetVisibleLevelsBoundingBox(const FLevelModelList& InList, bool bIncludeChildren);
 
 	/** @return	The UICommandList supported by this collection */
-	const TSharedRef<const FUICommandList> GetCommandList() const;
+	const TSharedRef<FUICommandList> GetCommandList() const;
 
 	/**  */
 	void LoadSettings();
diff --git a/Engine/Source/Editor/WorldBrowser/Private/SWorldHierarchyImpl.cpp b/Engine/Source/Editor/WorldBrowser/Private/SWorldHierarchyImpl.cpp
--- a/Engine/Source/Editor/WorldBrowser/Private/SWorldHierarchyImpl.cpp
+++ b/Engine/Source/Editor/WorldBrowser/Private/SWorldHierarchyImpl.cpp
@@ -10,6 +10,7 @@
 
 #include "Framework/Application/SlateApplication.h"
 #include "Framework/MultiBox/MultiBoxBuilder.h"
+#include "EditorMenuSubsystem.h"
 #include "Widgets/Layout/SSeparator.h"
 #include "Widgets/Images/SImage.h"
 #include "Widgets/Input/SComboButton.h"
@@ -458,22 +459,33 @@ TSharedPtr<SWidget> SWorldHierarchyImpl::ConstructLevelContextMenu() const
 
 	if (!WorldModel->IsReadOnly())
 	{
-		FMenuBuilder MenuBuilder(true, WorldModel->GetCommandList());
+		UEditorMenuSubsystem* EditorMenus = UEditorMenuSubsystem::Get();
+		static const FName MenuName = "WorldBrowser.WorldHierarchy.LevelContextMenu";
+		if (!EditorMenus->IsMenuRegistered(MenuName))
+		{
+			EditorMenus->RegisterMenu(MenuName);
+		}
+
+		FEditorMenuContext Context(WorldModel->GetCommandList(), TSharedPtr<FExtender>());
+		UEditorMenu* Menu = EditorMenus->GenerateMenu(MenuName, Context);
 
 		TArray<WorldHierarchy::FWorldTreeItemPtr> SelectedItems = GetSelectedTreeItems();
 
 		if (SelectedItems.Num() == 1)
 		{
 			// If exactly one item is selected, allow it to generate its own context menu
-			SelectedItems[0]->GenerateContextMenu(MenuBuilder, *this);
+			SelectedItems[0]->GenerateContextMenu(Menu, *this);
 		}
 		else if (SelectedItems.Num() == 0)
 		{
 			// If no items are selected, allow the first root level item to create a context menu
-			RootTreeItems[0]->GenerateContextMenu(MenuBuilder, *this);
+			RootTreeItems[0]->GenerateContextMenu(Menu, *this);
 		}
 
-		WorldModel->BuildHierarchyMenu(MenuBuilder);
+		Menu->AddDynamicSection("HierarchyDynamicSection", FNewEditorMenuDelegateLegacy::CreateLambda([=](FMenuBuilder& MenuBuilder, UEditorMenu* Menu)
+		{
+			WorldModel->BuildHierarchyMenu(MenuBuilder);
+		}));
 
 		// Generate the "Move To" and "Select" submenus based on the current selection
 		if (WorldModel->HasFolderSupport())
@@ -495,24 +507,30 @@ TSharedPtr<SWidget> SWorldHierarchyImpl::ConstructLevelContextMenu() const
 
 			if (bAllSelectedItemsCanMove && FLevelFolders::IsAvailable())
 			{
-				MenuBuilder.AddSubMenu(
+				FEditorMenuSection& Section = Menu->AddSection("Section");
+				Section.AddEntry(FEditorMenuEntry::InitSubMenu(
+					Menu->MenuName,
+					"MoveSelectionTo",
 					LOCTEXT("MoveSelectionTo", "Move To"),
 					LOCTEXT("MoveSelectionTo_Tooltip", "Move selection to another folder"),
 					FNewMenuDelegate::CreateSP(const_cast<SWorldHierarchyImpl*>(this), &SWorldHierarchyImpl::FillFoldersSubmenu)
-				);
+				));
 			}
 
 			if (bOnlyFoldersSelected)
 			{
-				MenuBuilder.AddSubMenu(
+				FEditorMenuSection& Section = Menu->AddSection("Section");
+				Section.AddEntry(FEditorMenuEntry::InitSubMenu(
+					Menu->MenuName,
+					"SelectSubmenu",
 					LOCTEXT("SelectSubmenu", "Select"),
 					LOCTEXT("SelectSubmenu_Tooltip", "Select child items of the current selection"),
 					FNewMenuDelegate::CreateSP(const_cast<SWorldHierarchyImpl*>(this), &SWorldHierarchyImpl::FillSelectionSubmenu)
-				);
+				));
 			}
 		}
 
-		MenuWidget = MenuBuilder.MakeWidget();
+		MenuWidget = EditorMenus->GenerateWidget(Menu);
 	}
 
 	return MenuWidget;
diff --git a/Engine/Source/Editor/WorldBrowser/Private/WorldTreeItemTypes.cpp b/Engine/Source/Editor/WorldBrowser/Private/WorldTreeItemTypes.cpp
--- a/Engine/Source/Editor/WorldBrowser/Private/WorldTreeItemTypes.cpp
+++ b/Engine/Source/Editor/WorldBrowser/Private/WorldTreeItemTypes.cpp
@@ -7,7 +7,7 @@
 #include "Misc/PackageName.h"
 #include "Engine/Engine.h"
 #include "SWorldHierarchyImpl.h"
-#include "Framework/MultiBox/MultiBoxBuilder.h"
+#include "EditorMenuSubsystem.h"
 #include "EditorStyleSet.h"
 #include "AssetSelection.h"
 #include "SWorldHierarchyImpl.h"
@@ -476,7 +476,7 @@ namespace WorldHierarchy
 		return Parent.IsValid();
 	}
 
-	void FLevelModelTreeItem::GenerateContextMenu(FMenuBuilder& MenuBuilder, const SWorldHierarchyImpl& Hierarchy)
+	void FLevelModelTreeItem::GenerateContextMenu(UEditorMenu* Menu, const SWorldHierarchyImpl& Hierarchy)
 	{
 		if (!Parent.IsValid() && WorldModel.Pin()->HasFolderSupport())
 		{
@@ -486,7 +486,8 @@ namespace WorldHierarchy
 			FName RootPath = LevelModel.Pin()->GetFolderPath();
 			auto NewFolderAction = FExecuteAction::CreateSP(const_cast<SWorldHierarchyImpl*>(&Hierarchy), &SWorldHierarchyImpl::CreateFolder, LevelModel.Pin(), RootPath);
 
-			MenuBuilder.AddMenuEntry(LOCTEXT("CreateFolder", "Create Folder"), FText(), NewFolderIcon, FUIAction(NewFolderAction));
+			FEditorMenuSection& Section = Menu->AddSection("Section");
+			Section.AddMenuEntry("CreateFolder", LOCTEXT("CreateFolder", "Create Folder"), FText(), NewFolderIcon, FUIAction(NewFolderAction));
 		}
 	}
 
@@ -903,7 +904,7 @@ namespace WorldHierarchy
 		return FEditorStyle::Get().GetBrush("WorldBrowser.FolderClosed");
 	}
 
-	void FFolderTreeItem::GenerateContextMenu(FMenuBuilder& MenuBuilder, const SWorldHierarchyImpl& Hierarchy)
+	void FFolderTreeItem::GenerateContextMenu(UEditorMenu* Menu, const SWorldHierarchyImpl& Hierarchy)
 	{
 		// Folder items should be able to create subfolders, rename themselves, or delete themselves from the tree
 		const FSlateIcon NewFolderIcon(FEditorStyle::GetStyleSetName(), "WorldBrowser.NewFolderIcon");
@@ -916,9 +917,10 @@ namespace WorldHierarchy
 		auto RenameFolderAction = FExecuteAction::CreateSP(const_cast<SWorldHierarchyImpl*>(&Hierarchy), &SWorldHierarchyImpl::InitiateRename, AsShared());
 		auto DeleteFolderAction = FExecuteAction::CreateSP(const_cast<SWorldHierarchyImpl*>(&Hierarchy), &SWorldHierarchyImpl::DeleteFolders, Folders, /*bTransactional*/ true);
 
-		MenuBuilder.AddMenuEntry(LOCTEXT("CreateSubFolder", "Create Subfolder"), FText(), NewFolderIcon, FUIAction(NewFolderAction));
-		MenuBuilder.AddMenuEntry(LOCTEXT("RenameFolder", "Rename"), FText(), FSlateIcon(FEditorStyle::GetStyleSetName(), "ContentBrowser.AssetActions.Rename"), FUIAction(RenameFolderAction));
-		MenuBuilder.AddMenuEntry(LOCTEXT("DeleteFolder", "Delete"), FText(), FSlateIcon(FEditorStyle::GetStyleSetName(), "ContentBrowser.AssetActions.Delete"), FUIAction(DeleteFolderAction));
+		FEditorMenuSection& Section = Menu->AddSection("Section");
+		Section.AddMenuEntry("CreateSubFolder", LOCTEXT("CreateSubFolder", "Create Subfolder"), FText(), NewFolderIcon, FUIAction(NewFolderAction));
+		Section.AddMenuEntry("RenameFolder", LOCTEXT("RenameFolder", "Rename"), FText(), FSlateIcon(FEditorStyle::GetStyleSetName(), "ContentBrowser.AssetActions.Rename"), FUIAction(RenameFolderAction));
+		Section.AddMenuEntry("DeleteFolder", LOCTEXT("DeleteFolder", "Delete"), FText(), FSlateIcon(FEditorStyle::GetStyleSetName(), "ContentBrowser.AssetActions.Delete"), FUIAction(DeleteFolderAction));
 	}
 
 	FValidationInfo FFolderTreeItem::ValidateDrop(const FDragDropEvent& DragEvent) const
diff --git a/Engine/Source/Editor/WorldBrowser/Private/WorldTreeItemTypes.h b/Engine/Source/Editor/WorldBrowser/Private/WorldTreeItemTypes.h
--- a/Engine/Source/Editor/WorldBrowser/Private/WorldTreeItemTypes.h
+++ b/Engine/Source/Editor/WorldBrowser/Private/WorldTreeItemTypes.h
@@ -87,7 +87,7 @@ namespace WorldHierarchy
 		virtual const FSlateBrush* GetHierarchyItemBrush() const override;
 
 		virtual bool CanChangeParents() const override;
-		virtual void GenerateContextMenu(FMenuBuilder& MenuBuilder, const SWorldHierarchyImpl& Hierarchy) override;
+		virtual void GenerateContextMenu(UEditorMenu* Menu, const SWorldHierarchyImpl& Hierarchy) override;
 
 	public:
 		virtual FLevelModelTreeItem* GetAsLevelModelTreeItem() const override { return const_cast<FLevelModelTreeItem*>(this); }
@@ -165,7 +165,7 @@ namespace WorldHierarchy
 		virtual float GetHierarchyItemBrushWidth() const { return 16.0f; }
 
 		virtual bool CanChangeParents() const { return true; }
-		virtual void GenerateContextMenu(FMenuBuilder& MenuBuilder, const SWorldHierarchyImpl& Hierarchy) override;
+		virtual void GenerateContextMenu(UEditorMenu* Menu, const SWorldHierarchyImpl& Hierarchy) override;
 
 		/** Sets the new path of the folder. This includes the leaf name. This will not rename any children in this folder */
 		void SetNewPath(FName NewPath);
diff --git a/Engine/Source/Editor/WorldBrowser/Public/IWorldTreeItem.h b/Engine/Source/Editor/WorldBrowser/Public/IWorldTreeItem.h
--- a/Engine/Source/Editor/WorldBrowser/Public/IWorldTreeItem.h
+++ b/Engine/Source/Editor/WorldBrowser/Public/IWorldTreeItem.h
@@ -14,7 +14,7 @@ class SWidget;
 class SWorldHierarchy;
 class FLevelCollectionModel;
 class SWorldHierarchyImpl;
-class FMenuBuilder;
+class UEditorMenu;
 
 namespace WorldHierarchy
 {
@@ -270,7 +270,7 @@ namespace WorldHierarchy
 		virtual bool CanChangeParents() const = 0;
 
 		/** Generates a context menu option for this item if and only if it's the only item selected in the hierarchy */
-		virtual void GenerateContextMenu(FMenuBuilder& MenuBuilder, const SWorldHierarchyImpl& Hierarchy) {}
+		virtual void GenerateContextMenu(UEditorMenu* Menu, const SWorldHierarchyImpl& Hierarchy) {}
 
 		/** Sets the item's visible status */
 		virtual void SetVisible(bool bVisible) = 0;
diff --git a/Engine/Source/Editor/WorldBrowser/WorldBrowser.Build.cs b/Engine/Source/Editor/WorldBrowser/WorldBrowser.Build.cs
--- a/Engine/Source/Editor/WorldBrowser/WorldBrowser.Build.cs
+++ b/Engine/Source/Editor/WorldBrowser/WorldBrowser.Build.cs
@@ -52,6 +52,7 @@ public class WorldBrowser : ModuleRules
                 "MaterialUtilities",
                 "RHI",
                 "Json",
+				"EditorMenus",
             }
 		);
 
