Better handling of editor menu command lists.
#rnx

diff --git a/Engine/Source/Editor/EditorMenus/Private/EditorMenuContext.cpp b/Engine/Source/Editor/EditorMenus/Private/EditorMenuContext.cpp
--- a/Engine/Source/Editor/EditorMenus/Private/EditorMenuContext.cpp
+++ b/Engine/Source/Editor/EditorMenus/Private/EditorMenuContext.cpp
@@ -54,16 +54,32 @@ UObject* FEditorMenuContext::FindByClass(UClass* InClass) const
 	return nullptr;
 }
 
-void FEditorMenuContext::AppendCommandList(TSharedPtr<FUICommandList> InCommandList)
+void FEditorMenuContext::AppendCommandList(const TSharedRef<FUICommandList>& InCommandList)
+{
+	const TSharedPtr<FUICommandList> List = InCommandList;
+	AppendCommandList(List);
+}
+
+void FEditorMenuContext::AppendCommandList(const TSharedPtr<FUICommandList>& InCommandList)
 {
 	if (InCommandList.IsValid())
 	{
-		if (!CommandList.IsValid())
+		CommandLists.Add(InCommandList);
+
+		if (CommandLists.Num() == 1)
+		{
+			CommandList = InCommandList;
+		}
+		else if (CommandLists.Num() == 2)
 		{
 			CommandList = MakeShared<FUICommandList>();
+			CommandList->Append(CommandLists[0].ToSharedRef());
+			CommandList->Append(InCommandList.ToSharedRef());
+		}
+		else
+		{
+			CommandList->Append(InCommandList.ToSharedRef());
 		}
-
-		CommandList->Append(InCommandList.ToSharedRef());
 	}
 }
 
diff --git a/Engine/Source/Editor/EditorMenus/Public/EditorMenuContext.h b/Engine/Source/Editor/EditorMenus/Public/EditorMenuContext.h
--- a/Engine/Source/Editor/EditorMenus/Public/EditorMenuContext.h
+++ b/Engine/Source/Editor/EditorMenus/Public/EditorMenuContext.h
@@ -52,7 +52,8 @@ public:
 	
 	UObject* FindByClass(UClass* InClass) const;
 
-	void AppendCommandList(TSharedPtr<FUICommandList> InCommandList);
+	void AppendCommandList(const TSharedRef<FUICommandList>& InCommandList);
+	void AppendCommandList(const TSharedPtr<FUICommandList>& InCommandList);
 	void AddExtender(const TSharedPtr<FExtender>& InExtender);
 	TSharedPtr<FExtender> GetAllExtenders();
 	void ReplaceExtenders(const TSharedPtr<FExtender>& InExtender);
@@ -69,10 +70,10 @@ private:
 	UPROPERTY()
 	TArray<UObject*> ContextObjects;
 
+	TArray<TSharedPtr<FUICommandList>> CommandLists;
+
 	TSharedPtr<FUICommandList> CommandList;
 
-	// Prevent CommandLists from being deleted
-	TArray<TSharedPtr<const FUICommandList>> ReferencedCommandLists;
 	TSharedPtr<FExtensibilityManager> ExtensibilityManager;
 };
 
diff --git a/Engine/Source/Editor/LevelEditor/Private/LevelEditorContextMenu.cpp b/Engine/Source/Editor/LevelEditor/Private/LevelEditorContextMenu.cpp
--- a/Engine/Source/Editor/LevelEditor/Private/LevelEditorContextMenu.cpp
+++ b/Engine/Source/Editor/LevelEditor/Private/LevelEditorContextMenu.cpp
@@ -493,7 +493,8 @@ FName FLevelEditorContextMenu::InitMenuContext(FEditorMenuContext& Context, TWea
 	RegisterActorContextMenu();
 	RegisterSceneOutlinerContextMenu();
 
-	Context.AppendCommandList(LevelEditor.Pin()->GetLevelEditorActions());
+	TSharedPtr<FUICommandList> LevelEditorActionsList = LevelEditor.Pin()->GetLevelEditorActions();
+	Context.AppendCommandList(LevelEditorActionsList);
 
 	ULevelEditorContextMenuContext* ContextObject = NewObject<ULevelEditorContextMenuContext>();
 	ContextObject->LevelEditor = LevelEditor;
@@ -512,7 +513,6 @@ FName FLevelEditorContextMenu::InitMenuContext(FEditorMenuContext& Context, TWea
 		// Get all menu extenders for this context menu from the level editor module
 		FLevelEditorModule& LevelEditorModule = FModuleManager::GetModuleChecked<FLevelEditorModule>(TEXT("LevelEditor"));
 		TArray<FLevelEditorModule::FLevelViewportMenuExtender_SelectedActors> MenuExtenderDelegates = LevelEditorModule.GetAllLevelViewportContextMenuExtenders();
-		TSharedPtr<FUICommandList> LevelEditorActionsList = LevelEditor.Pin()->GetLevelEditorActions();
 		TArray<TSharedPtr<FExtender>> Extenders;
 		for (int32 i = 0; i < MenuExtenderDelegates.Num(); ++i)
 		{
