Rename Editor\EditorMenus\ module to Runtime\ToolMenus\ so that other modules can use the API without having to depend on editor modules.

diff --git a/Engine/Source/Editor/ContentBrowser/Private/AssetContextMenu.cpp b/Engine/Source/Editor/ContentBrowser/Private/AssetContextMenu.cpp
--- a/Engine/Source/Editor/ContentBrowser/Private/AssetContextMenu.cpp
+++ b/Engine/Source/Editor/ContentBrowser/Private/AssetContextMenu.cpp
@@ -20,7 +20,7 @@
 #include "Widgets/Text/STextBlock.h"
 #include "Widgets/Input/SMultiLineEditableTextBox.h"
 #include "Framework/MultiBox/MultiBoxBuilder.h"
-#include "EditorMenuSubsystem.h"
+#include "ToolMenus.h"
 #include "ContentBrowserMenuContexts.h"
 #include "Widgets/Input/SButton.h"
 #include "EditorStyleSet.h"
@@ -136,13 +136,13 @@ TSharedRef<SWidget> FAssetContextMenu::MakeContextMenu(const TArray<FAssetData>&
 	UContentBrowserAssetContextMenuContext* ContextObject = NewObject<UContentBrowserAssetContextMenuContext>();
 	ContextObject->AssetContextMenu = SharedThis(this);
 
-	UEditorMenuSubsystem* EditorMenus = UEditorMenuSubsystem::Get();
+	UToolMenus* ToolMenus = UToolMenus::Get();
 
 	static const FName BaseMenuName("ContentBrowser.AssetContextMenu");
 	RegisterContextMenu(BaseMenuName);
 
 	// Create menu hierarchy based on class hierarchy
-	TArray<UEditorMenu*> MenuHierarchy;
+	TArray<UToolMenu*> MenuHierarchy;
 	FName MenuName = BaseMenuName;
 	{
 		// Objects must be loaded for this operation... for now
@@ -174,15 +174,15 @@ TSharedRef<SWidget> FAssetContextMenu::MakeContextMenu(const TArray<FAssetData>&
 			{
 				ContextObject->CommonAssetTypeActions = CommonAssetTypeActions;
 				AssetTypeActionsClass = CommonAssetTypeActions->GetSupportedClass();
-				MenuName = UEditorMenuSubsystem::JoinMenuPaths(BaseMenuName, AssetTypeActionsClass->GetFName());
+				MenuName = UToolMenus::JoinMenuPaths(BaseMenuName, AssetTypeActionsClass->GetFName());
 			}
 
 			// Find menus for common class
 			UClass* ClassWithMenu = CommonClass;
 			while (ClassWithMenu)
 			{
-				FName CheckName = UEditorMenuSubsystem::JoinMenuPaths(BaseMenuName, ClassWithMenu->GetFName());
-				UEditorMenu* FoundMenu = EditorMenus->FindMenu(CheckName);
+				FName CheckName = UToolMenus::JoinMenuPaths(BaseMenuName, ClassWithMenu->GetFName());
+				UToolMenu* FoundMenu = ToolMenus->FindMenu(CheckName);
 				if (FoundMenu && FoundMenu->Sections.Num() > 0)
 				{
 					// Change generated menu's name if necessary
@@ -201,13 +201,13 @@ TSharedRef<SWidget> FAssetContextMenu::MakeContextMenu(const TArray<FAssetData>&
 			}
 		}
 
-		MenuHierarchy.Add(EditorMenus->FindMenu(BaseMenuName));
+		MenuHierarchy.Add(ToolMenus->FindMenu(BaseMenuName));
 
 		Algo::Reverse(MenuHierarchy);
 	}
 
-	FEditorMenuContext MenuContext(InCommandList, MenuExtender, ContextObject);
-	UEditorMenu* Menu = EditorMenus->GenerateMenu(MenuHierarchy, MenuContext);
+	FToolMenuContext MenuContext(InCommandList, MenuExtender, ContextObject);
+	UToolMenu* Menu = ToolMenus->GenerateMenu(MenuHierarchy, MenuContext);
 	Menu->MenuName = MenuName;
 
 	// Only add something if at least one asset is selected
@@ -238,18 +238,18 @@ TSharedRef<SWidget> FAssetContextMenu::MakeContextMenu(const TArray<FAssetData>&
 		AddSourceControlMenuOptions(Menu);
 	}
 	
-	return EditorMenus->GenerateWidget(Menu);
+	return ToolMenus->GenerateWidget(Menu);
 }
 
 void FAssetContextMenu::RegisterContextMenu(const FName MenuName)
 {
-	UEditorMenuSubsystem* EditorMenus = UEditorMenuSubsystem::Get();
-	if (!EditorMenus->IsMenuRegistered(MenuName))
+	UToolMenus* ToolMenus = UToolMenus::Get();
+	if (!ToolMenus->IsMenuRegistered(MenuName))
 	{
-		UEditorMenu* Menu = EditorMenus->RegisterMenu(MenuName);
-		FEditorMenuSection& Section = Menu->FindOrAddSection("GetAssetActions");
+		UToolMenu* Menu = ToolMenus->RegisterMenu(MenuName);
+		FToolMenuSection& Section = Menu->FindOrAddSection("GetAssetActions");
 
-		Section.AddDynamicEntry("GetActions", FNewEditorMenuSectionDelegate::CreateLambda([](FEditorMenuSection& InSection)
+		Section.AddDynamicEntry("GetActions", FNewToolMenuSectionDelegate::CreateLambda([](FToolMenuSection& InSection)
 		{
 			UContentBrowserAssetContextMenuContext* Context = InSection.FindContext<UContentBrowserAssetContextMenuContext>();
 			if (Context && Context->CommonAssetTypeActions.IsValid())
@@ -258,7 +258,7 @@ void FAssetContextMenu::RegisterContextMenu(const FName MenuName)
 			}
 		}));
 
-		Section.AddDynamicEntry("GetActionsLegacy", FNewEditorMenuDelegateLegacy::CreateLambda([](FMenuBuilder& MenuBuilder, UEditorMenu* InMenu)
+		Section.AddDynamicEntry("GetActionsLegacy", FNewToolMenuDelegateLegacy::CreateLambda([](FMenuBuilder& MenuBuilder, UToolMenu* InMenu)
 		{
 			UContentBrowserAssetContextMenuContext* Context = InMenu->FindContext<UContentBrowserAssetContextMenuContext>();
 			if (Context && Context->CommonAssetTypeActions.IsValid())
@@ -299,7 +299,7 @@ void FAssetContextMenu::SetOnAssetViewRefreshRequested(const FOnAssetViewRefresh
 	OnAssetViewRefreshRequested = InOnAssetViewRefreshRequested;
 }
 
-bool FAssetContextMenu::AddImportedAssetMenuOptions(UEditorMenu* Menu)
+bool FAssetContextMenu::AddImportedAssetMenuOptions(UToolMenu* Menu)
 {
 	if (AreImportedAssetActionsVisible())
 	{
@@ -308,9 +308,9 @@ bool FAssetContextMenu::AddImportedAssetMenuOptions(UEditorMenu* Menu)
 		int32 ValidSelectedAssetCount = 0;
 		GetSelectedAssetSourceFilePaths(ResolvedFilePaths, SourceFileLabels, ValidSelectedAssetCount);
 
-		FEditorMenuSection& Section = Menu->AddSection("ImportedAssetActions", LOCTEXT("ImportedAssetActionsMenuHeading", "Imported Asset"));
+		FToolMenuSection& Section = Menu->AddSection("ImportedAssetActions", LOCTEXT("ImportedAssetActionsMenuHeading", "Imported Asset"));
 		{
-			auto CreateSubMenu = [this](UEditorMenu* SubMenu, bool bReimportWithNewFile)
+			auto CreateSubMenu = [this](UToolMenu* SubMenu, bool bReimportWithNewFile)
 			{
 				//Get the data, we cannot use the closure since the lambda will be call when the function scope will be gone
 				TArray<FString> ResolvedFilePaths;
@@ -319,7 +319,7 @@ bool FAssetContextMenu::AddImportedAssetMenuOptions(UEditorMenu* Menu)
 				GetSelectedAssetSourceFilePaths(ResolvedFilePaths, SourceFileLabels, ValidSelectedAssetCount);
 				if (SourceFileLabels.Num() > 0 )
 				{
-					FEditorMenuSection& SubSection = SubMenu->AddSection("Section");
+					FToolMenuSection& SubSection = SubMenu->AddSection("Section");
 					for (int32 SourceFileIndex = 0; SourceFileIndex < SourceFileLabels.Num(); ++SourceFileIndex)
 					{
 						FText ReimportLabel = FText::Format(LOCTEXT("ReimportNoLabel", "SourceFile {0}"), SourceFileIndex);
@@ -370,19 +370,19 @@ bool FAssetContextMenu::AddImportedAssetMenuOptions(UEditorMenu* Menu)
 			//Reimport Menu
 			if (ValidSelectedAssetCount == 1 && SourceFileLabels.Num() > 1)
 			{
-				Section.AddEntry(FEditorMenuEntry::InitSubMenu(
+				Section.AddEntry(FToolMenuEntry::InitSubMenu(
 					Menu->GetMenuName(),
 					"Reimport",
 					LOCTEXT("Reimport", "Reimport"),
 					LOCTEXT("ReimportEmptyTooltip", ""),
-					FNewEditorMenuDelegate::CreateLambda(CreateSubMenu, false) ));
+					FNewToolMenuDelegate::CreateLambda(CreateSubMenu, false) ));
 				//With new file
-				Section.AddEntry(FEditorMenuEntry::InitSubMenu(
+				Section.AddEntry(FToolMenuEntry::InitSubMenu(
 					Menu->GetMenuName(),
 					"ReimportWithNewFile",
 					LOCTEXT("ReimportWithNewFile", "Reimport With New File"),
 					LOCTEXT("ReimportEmptyTooltip", ""),
-					FNewEditorMenuDelegate::CreateLambda(CreateSubMenu, true)));
+					FNewToolMenuDelegate::CreateLambda(CreateSubMenu, true)));
 			}
 			else
 			{
@@ -444,13 +444,13 @@ bool FAssetContextMenu::AddImportedAssetMenuOptions(UEditorMenu* Menu)
 	return false;
 }
 
-bool FAssetContextMenu::AddCommonMenuOptions(UEditorMenu* Menu)
+bool FAssetContextMenu::AddCommonMenuOptions(UToolMenu* Menu)
 {
 	int32 NumAssetItems, NumClassItems;
 	ContentBrowserUtils::CountItemTypes(SelectedAssets, NumAssetItems, NumClassItems);
 
 	{
-		FEditorMenuSection& Section = Menu->AddSection("CommonAssetActions", LOCTEXT("CommonAssetActionsMenuHeading", "Common"));
+		FToolMenuSection& Section = Menu->AddSection("CommonAssetActions", LOCTEXT("CommonAssetActionsMenuHeading", "Common"));
 
 		// Edit
 		Section.AddMenuEntry(
@@ -493,12 +493,12 @@ bool FAssetContextMenu::AddCommonMenuOptions(UEditorMenu* Menu)
 				);
 
 			// Asset Actions sub-menu
-			Section.AddEntry(FEditorMenuEntry::InitSubMenu(
+			Section.AddEntry(FToolMenuEntry::InitSubMenu(
 				Menu->GetMenuName(),
 				"AssetActionsSubMenu",
 				LOCTEXT("AssetActionsSubMenuLabel", "Asset Actions"),
 				LOCTEXT("AssetActionsSubMenuToolTip", "Other asset actions"),
-				FNewEditorMenuDelegate::CreateSP(this, &FAssetContextMenu::MakeAssetActionsSubMenu),
+				FNewToolMenuDelegate::CreateSP(this, &FAssetContextMenu::MakeAssetActionsSubMenu),
 				FUIAction(
 					FExecuteAction(),
 					FCanExecuteAction::CreateSP( this, &FAssetContextMenu::CanExecuteAssetActions )
@@ -511,12 +511,12 @@ bool FAssetContextMenu::AddCommonMenuOptions(UEditorMenu* Menu)
 			if (NumClassItems == 0)
 			{
 				// Asset Localization sub-menu
-				Section.AddEntry(FEditorMenuEntry::InitSubMenu(
+				Section.AddEntry(FToolMenuEntry::InitSubMenu(
 					Menu->GetMenuName(),
 					"LocalizationSubMenu",
 					LOCTEXT("LocalizationSubMenuLabel", "Asset Localization"),
 					LOCTEXT("LocalizationSubMenuToolTip", "Manage the localization of this asset"),
-					FNewEditorMenuDelegate::CreateSP(this, &FAssetContextMenu::MakeAssetLocalizationSubMenu),
+					FNewToolMenuDelegate::CreateSP(this, &FAssetContextMenu::MakeAssetLocalizationSubMenu),
 					FUIAction(),
 					EUserInterfaceActionType::Button,
 					false,
@@ -529,9 +529,9 @@ bool FAssetContextMenu::AddCommonMenuOptions(UEditorMenu* Menu)
 	return true;
 }
 
-void FAssetContextMenu::AddExploreMenuOptions(UEditorMenu* Menu)
+void FAssetContextMenu::AddExploreMenuOptions(UToolMenu* Menu)
 {
-	FEditorMenuSection& Section = Menu->AddSection("AssetContextExploreMenuOptions", LOCTEXT("AssetContextExploreMenuOptionsHeading", "Explore"));
+	FToolMenuSection& Section = Menu->AddSection("AssetContextExploreMenuOptions", LOCTEXT("AssetContextExploreMenuOptionsHeading", "Explore"));
 	{
 		// Find in Content Browser
 		Section.AddMenuEntry(
@@ -554,10 +554,10 @@ void FAssetContextMenu::AddExploreMenuOptions(UEditorMenu* Menu)
 	}
 }
 
-void FAssetContextMenu::MakeAssetActionsSubMenu(UEditorMenu* Menu)
+void FAssetContextMenu::MakeAssetActionsSubMenu(UToolMenu* Menu)
 {	
 	{
-		FEditorMenuSection& Section = Menu->AddSection("AssetActionsSection");
+		FToolMenuSection& Section = Menu->AddSection("AssetActionsSection");
 
 		// Create BP Using This
 		Section.AddMenuEntry(
@@ -602,7 +602,7 @@ void FAssetContextMenu::MakeAssetActionsSubMenu(UEditorMenu* Menu)
 
 	// FIND ACTIONS
 	{
-		FEditorMenuSection& Section = Menu->AddSection("AssetContextFindActions", LOCTEXT("AssetContextFindActionsMenuHeading", "Find"));
+		FToolMenuSection& Section = Menu->AddSection("AssetContextFindActions", LOCTEXT("AssetContextFindActionsMenuHeading", "Find"));
 		// Select Actors Using This Asset
 		Section.AddMenuEntry(
 			"FindAssetInWorld",
@@ -618,7 +618,7 @@ void FAssetContextMenu::MakeAssetActionsSubMenu(UEditorMenu* Menu)
 
 	// MOVE ACTIONS
 	{
-		FEditorMenuSection& Section = Menu->AddSection("AssetContextMoveActions", LOCTEXT("AssetContextMoveActionsMenuHeading", "Move"));
+		FToolMenuSection& Section = Menu->AddSection("AssetContextMoveActions", LOCTEXT("AssetContextMoveActionsMenuHeading", "Move"));
 		bool bHasExportableAssets = false;
 		for (const FAssetData& AssetData : SelectedAssets)
 		{
@@ -670,7 +670,7 @@ void FAssetContextMenu::MakeAssetActionsSubMenu(UEditorMenu* Menu)
 
 	// ADVANCED ACTIONS
 	{
-		FEditorMenuSection& Section = Menu->AddSection("AssetContextAdvancedActions", LOCTEXT("AssetContextAdvancedActionsMenuHeading", "Advanced"));
+		FToolMenuSection& Section = Menu->AddSection("AssetContextAdvancedActions", LOCTEXT("AssetContextAdvancedActionsMenuHeading", "Advanced"));
 
 		// Reload
 		Section.AddMenuEntry(
@@ -751,12 +751,12 @@ void FAssetContextMenu::MakeAssetActionsSubMenu(UEditorMenu* Menu)
 				FUIAction( FExecuteAction::CreateSP(this, &FAssetContextMenu::ExecuteAssignChunkID) )
 				);
 
-			Section.AddEntry(FEditorMenuEntry::InitSubMenu(
+			Section.AddEntry(FToolMenuEntry::InitSubMenu(
 				Menu->GetMenuName(),
 				"RemoveAssetFromChunk",
 				LOCTEXT("RemoveAssetFromChunk", "Remove from Chunk..."),
 				LOCTEXT("RemoveAssetFromChunkTooltip", "Removed an asset from a Chunk it's assigned to."),
-				FNewEditorMenuDelegate::CreateRaw(this, &FAssetContextMenu::MakeChunkIDListMenu)
+				FNewToolMenuDelegate::CreateRaw(this, &FAssetContextMenu::MakeChunkIDListMenu)
 				));
 
 			Section.AddMenuEntry(
@@ -771,7 +771,7 @@ void FAssetContextMenu::MakeAssetActionsSubMenu(UEditorMenu* Menu)
 
 	if (GetDefault<UEditorExperimentalSettings>()->bTextAssetFormatSupport)
 	{
-		FEditorMenuSection& FormatActionsSection = Menu->AddSection("AssetContextTextAssetFormatActions", LOCTEXT("AssetContextTextAssetFormatActionsHeading", "Text Assets"));
+		FToolMenuSection& FormatActionsSection = Menu->AddSection("AssetContextTextAssetFormatActions", LOCTEXT("AssetContextTextAssetFormatActionsHeading", "Text Assets"));
 		{
 			FormatActionsSection.AddMenuEntry(
 				"ExportToTextFormat",
@@ -822,7 +822,7 @@ bool FAssetContextMenu::CanExecuteAssetActions() const
 	return !bAtLeastOneClassSelected;
 }
 
-void FAssetContextMenu::MakeAssetLocalizationSubMenu(UEditorMenu* Menu)
+void FAssetContextMenu::MakeAssetLocalizationSubMenu(UToolMenu* Menu)
 {
 	TArray<FCultureRef> CurrentCultures;
 
@@ -952,7 +952,7 @@ void FAssetContextMenu::MakeAssetLocalizationSubMenu(UEditorMenu* Menu)
 #if USE_STABLE_LOCALIZATION_KEYS
 	// Add the Localization ID options
 	{
-		FEditorMenuSection& Section = Menu->AddSection("LocalizationId", LOCTEXT("LocalizationIdHeading", "Localization ID"));
+		FToolMenuSection& Section = Menu->AddSection("LocalizationId", LOCTEXT("LocalizationIdHeading", "Localization ID"));
 		{
 			// Show the localization ID if we have a single asset selected
 			if (SelectedAssets.Num() == 1)
@@ -985,7 +985,7 @@ void FAssetContextMenu::MakeAssetLocalizationSubMenu(UEditorMenu* Menu)
 		FString PackageFilename;
 		if (FPackageName::DoesPackageExist(SelectedAssets[0].PackageName.ToString(), nullptr, &PackageFilename))
 		{
-			FEditorMenuSection& Section = Menu->AddSection("LocalizationCache", LOCTEXT("LocalizationCacheHeading", "Localization Cache"));
+			FToolMenuSection& Section = Menu->AddSection("LocalizationCache", LOCTEXT("LocalizationCacheHeading", "Localization Cache"));
 			{
 				// Always show the reset localization ID option
 				Section.AddMenuEntry(
@@ -1002,7 +1002,7 @@ void FAssetContextMenu::MakeAssetLocalizationSubMenu(UEditorMenu* Menu)
 	// If we found source assets for localized assets, then we can show the Source Asset options
 	if (SourceAssetsState.CurrentAssets.Num() > 0)
 	{
-		FEditorMenuSection& Section = Menu->AddSection("ManageSourceAsset", LOCTEXT("ManageSourceAssetHeading", "Manage Source Asset"));
+		FToolMenuSection& Section = Menu->AddSection("ManageSourceAsset", LOCTEXT("ManageSourceAssetHeading", "Manage Source Asset"));
 		{
 			Section.AddMenuEntry(
 				"ShowSourceAsset",
@@ -1025,14 +1025,14 @@ void FAssetContextMenu::MakeAssetLocalizationSubMenu(UEditorMenu* Menu)
 	// If we currently have source assets selected, then we can show the Localized Asset options
 	if (SourceAssetsState.SelectedAssets.Num() > 0)
 	{
-		FEditorMenuSection& Section = Menu->AddSection("ManageLocalizedAsset", LOCTEXT("ManageLocalizedAssetHeading", "Manage Localized Asset"));
+		FToolMenuSection& Section = Menu->AddSection("ManageLocalizedAsset", LOCTEXT("ManageLocalizedAssetHeading", "Manage Localized Asset"));
 		{
-			Section.AddEntry(FEditorMenuEntry::InitSubMenu(
+			Section.AddEntry(FToolMenuEntry::InitSubMenu(
 				Menu->GetMenuName(),
 				"CreateLocalizedAsset",
 				LOCTEXT("CreateLocalizedAsset", "Create Localized Asset"),
 				LOCTEXT("CreateLocalizedAssetTooltip", "Create a new localized asset."),
-				FNewEditorMenuDelegate::CreateSP(this, &FAssetContextMenu::MakeCreateLocalizedAssetSubMenu, SourceAssetsState.SelectedAssets, LocalizedAssetsState),
+				FNewToolMenuDelegate::CreateSP(this, &FAssetContextMenu::MakeCreateLocalizedAssetSubMenu, SourceAssetsState.SelectedAssets, LocalizedAssetsState),
 				FUIAction(),
 				EUserInterfaceActionType::Button,
 				false,
@@ -1047,24 +1047,24 @@ void FAssetContextMenu::MakeAssetLocalizationSubMenu(UEditorMenu* Menu)
 
 			if (NumLocalizedAssets > 0)
 			{
-				Section.AddEntry(FEditorMenuEntry::InitSubMenu(
+				Section.AddEntry(FToolMenuEntry::InitSubMenu(
 					Menu->GetMenuName(),
 					"ShowLocalizedAsset",
 					LOCTEXT("ShowLocalizedAsset", "Show Localized Asset"),
 					LOCTEXT("ShowLocalizedAssetTooltip", "Show the localized asset in the Content Browser."),
-					FNewEditorMenuDelegate::CreateSP(this, &FAssetContextMenu::MakeShowLocalizedAssetSubMenu, LocalizedAssetsState),
+					FNewToolMenuDelegate::CreateSP(this, &FAssetContextMenu::MakeShowLocalizedAssetSubMenu, LocalizedAssetsState),
 					FUIAction(),
 					EUserInterfaceActionType::Button,
 					false,
 					FSlateIcon(FEditorStyle::GetStyleSetName(), "SystemWideCommands.FindInContentBrowser")
 					));
 
-				Section.AddEntry(FEditorMenuEntry::InitSubMenu(
+				Section.AddEntry(FToolMenuEntry::InitSubMenu(
 					Menu->GetMenuName(),
 					"EditLocalizedAsset",
 					LOCTEXT("EditLocalizedAsset", "Edit Localized Asset"),
 					LOCTEXT("EditLocalizedAssetTooltip", "Edit the localized asset."),
-					FNewEditorMenuDelegate::CreateSP(this, &FAssetContextMenu::MakeEditLocalizedAssetSubMenu, LocalizedAssetsState),
+					FNewToolMenuDelegate::CreateSP(this, &FAssetContextMenu::MakeEditLocalizedAssetSubMenu, LocalizedAssetsState),
 					FUIAction(),
 					EUserInterfaceActionType::Button,
 					false,
@@ -1075,9 +1075,9 @@ void FAssetContextMenu::MakeAssetLocalizationSubMenu(UEditorMenu* Menu)
 	}
 }
 
-void FAssetContextMenu::MakeCreateLocalizedAssetSubMenu(UEditorMenu* Menu, TSet<FName> InSelectedSourceAssets, TArray<FLocalizedAssetsState> InLocalizedAssetsState)
+void FAssetContextMenu::MakeCreateLocalizedAssetSubMenu(UToolMenu* Menu, TSet<FName> InSelectedSourceAssets, TArray<FLocalizedAssetsState> InLocalizedAssetsState)
 {
-	FEditorMenuSection& Section = Menu->AddSection("Section");
+	FToolMenuSection& Section = Menu->AddSection("Section");
 
 	for (const FLocalizedAssetsState& LocalizedAssetsStateForCulture : InLocalizedAssetsState)
 	{
@@ -1095,9 +1095,9 @@ void FAssetContextMenu::MakeCreateLocalizedAssetSubMenu(UEditorMenu* Menu, TSet<
 	}
 }
 
-void FAssetContextMenu::MakeShowLocalizedAssetSubMenu(UEditorMenu* Menu, TArray<FLocalizedAssetsState> InLocalizedAssetsState)
+void FAssetContextMenu::MakeShowLocalizedAssetSubMenu(UToolMenu* Menu, TArray<FLocalizedAssetsState> InLocalizedAssetsState)
 {
-	FEditorMenuSection& Section = Menu->AddSection("Section");
+	FToolMenuSection& Section = Menu->AddSection("Section");
 
 	for (const FLocalizedAssetsState& LocalizedAssetsStateForCulture : InLocalizedAssetsState)
 	{
@@ -1114,9 +1114,9 @@ void FAssetContextMenu::MakeShowLocalizedAssetSubMenu(UEditorMenu* Menu, TArray<
 	}
 }
 
-void FAssetContextMenu::MakeEditLocalizedAssetSubMenu(UEditorMenu* Menu, TArray<FLocalizedAssetsState> InLocalizedAssetsState)
+void FAssetContextMenu::MakeEditLocalizedAssetSubMenu(UToolMenu* Menu, TArray<FLocalizedAssetsState> InLocalizedAssetsState)
 {
-	FEditorMenuSection& Section = Menu->AddSection("Section");
+	FToolMenuSection& Section = Menu->AddSection("Section");
 
 	for (const FLocalizedAssetsState& LocalizedAssetsStateForCulture : InLocalizedAssetsState)
 	{
@@ -1198,9 +1198,9 @@ void FAssetContextMenu::ExecuteOpenEditorsForAssets(TArray<FName> InAssets)
 	FAssetEditorManager::Get().OpenEditorsForAssets(InAssets);
 }
 
-bool FAssetContextMenu::AddReferenceMenuOptions(UEditorMenu* Menu)
+bool FAssetContextMenu::AddReferenceMenuOptions(UToolMenu* Menu)
 {
-	FEditorMenuSection& Section = Menu->AddSection("AssetContextReferences", LOCTEXT("ReferencesMenuHeading", "References"));
+	FToolMenuSection& Section = Menu->AddSection("AssetContextReferences", LOCTEXT("ReferencesMenuHeading", "References"));
 	{
 		Section.AddMenuEntry(
 			"CopyReference",
@@ -1214,7 +1214,7 @@ bool FAssetContextMenu::AddReferenceMenuOptions(UEditorMenu* Menu)
 	return true;
 }
 
-bool FAssetContextMenu::AddDocumentationMenuOptions(UEditorMenu* Menu)
+bool FAssetContextMenu::AddDocumentationMenuOptions(UToolMenu* Menu)
 {
 	bool bAddedOption = false;
 
@@ -1252,7 +1252,7 @@ bool FAssetContextMenu::AddDocumentationMenuOptions(UEditorMenu* Menu)
 
 				const FString CodeFileName = FPaths::GetCleanFilename( *ClassHeaderPath );
 
-				FEditorMenuSection& Section = Menu->AddSection( "AssetCode"/*, LOCTEXT("AssetCodeHeading", "C++")*/ );
+				FToolMenuSection& Section = Menu->AddSection( "AssetCode"/*, LOCTEXT("AssetCodeHeading", "C++")*/ );
 				{
 					Section.AddMenuEntry(
 						"GoToCodeForAsset",
@@ -1270,7 +1270,7 @@ bool FAssetContextMenu::AddDocumentationMenuOptions(UEditorMenu* Menu)
 		{
 			bAddedOption = true;
 
-			FEditorMenuSection& Section = Menu->AddSection( "AssetDocumentation"/*, LOCTEXT("AseetDocsHeading", "Documentation")*/ );
+			FToolMenuSection& Section = Menu->AddSection( "AssetDocumentation"/*, LOCTEXT("AseetDocsHeading", "Documentation")*/ );
 			{
 					if (bIsBlueprint)
 					{
@@ -1345,7 +1345,7 @@ bool FAssetContextMenu::AddDocumentationMenuOptions(UEditorMenu* Menu)
 	return bAddedOption;
 }
 
-bool FAssetContextMenu::AddAssetTypeMenuOptions(UEditorMenu* Menu, TArray<UObject*>& SelectedObjects)
+bool FAssetContextMenu::AddAssetTypeMenuOptions(UToolMenu* Menu, TArray<UObject*>& SelectedObjects)
 {
 	bool bAnyTypeOptions = false;
 
@@ -1353,7 +1353,7 @@ bool FAssetContextMenu::AddAssetTypeMenuOptions(UEditorMenu* Menu, TArray<UObjec
 	{
 		// Label "GetAssetActions" section
 		UContentBrowserAssetContextMenuContext* Context = Menu->FindContext<UContentBrowserAssetContextMenuContext>();
-		FEditorMenuSection* Section = Menu->FindSection("GetAssetActions");
+		FToolMenuSection* Section = Menu->FindSection("GetAssetActions");
 		if (Section && Section->Blocks.Num() > 0 && Context)
 		{
 			if (Context->CommonAssetTypeActions.IsValid())
@@ -1372,19 +1372,19 @@ bool FAssetContextMenu::AddAssetTypeMenuOptions(UEditorMenu* Menu, TArray<UObjec
 	return bAnyTypeOptions;
 }
 
-bool FAssetContextMenu::AddSourceControlMenuOptions(UEditorMenu* Menu)
+bool FAssetContextMenu::AddSourceControlMenuOptions(UToolMenu* Menu)
 {
-	FEditorMenuSection& Section = Menu->AddSection("AssetContextSourceControl");
+	FToolMenuSection& Section = Menu->AddSection("AssetContextSourceControl");
 	
 	if ( ISourceControlModule::Get().IsEnabled() )
 	{
 		// SCC sub menu
-		Section.AddEntry(FEditorMenuEntry::InitSubMenu(
+		Section.AddEntry(FToolMenuEntry::InitSubMenu(
 			Menu->GetMenuName(),
 			"SourceControlSubMenu",
 			LOCTEXT("SourceControlSubMenuLabel", "Source Control"),
 			LOCTEXT("SourceControlSubMenuToolTip", "Source control actions."),
-			FNewEditorMenuDelegate::CreateSP(this, &FAssetContextMenu::FillSourceControlSubMenu),
+			FNewToolMenuDelegate::CreateSP(this, &FAssetContextMenu::FillSourceControlSubMenu),
 			FUIAction(
 				FExecuteAction(),
 				FCanExecuteAction::CreateSP( this, &FAssetContextMenu::CanExecuteSourceControlActions )
@@ -1425,9 +1425,9 @@ bool FAssetContextMenu::AddSourceControlMenuOptions(UEditorMenu* Menu)
 	return true;
 }
 
-void FAssetContextMenu::FillSourceControlSubMenu(UEditorMenu* Menu)
+void FAssetContextMenu::FillSourceControlSubMenu(UToolMenu* Menu)
 {
-	FEditorMenuSection& Section = Menu->AddSection("AssetSourceControlActions", LOCTEXT("AssetSourceControlActionsMenuHeading", "Source Control"));
+	FToolMenuSection& Section = Menu->AddSection("AssetSourceControlActions", LOCTEXT("AssetSourceControlActionsMenuHeading", "Source Control"));
 
 	if( CanExecuteSCCMerge() )
 	{
@@ -1555,12 +1555,12 @@ bool FAssetContextMenu::CanExecuteSourceControlActions() const
 	return !bAtLeastOneClassSelected;
 }
 
-bool FAssetContextMenu::AddCollectionMenuOptions(UEditorMenu* Menu)
+bool FAssetContextMenu::AddCollectionMenuOptions(UToolMenu* Menu)
 {
 	class FManageCollectionsContextMenu
 	{
 	public:
-		static void CreateManageCollectionsSubMenu(UEditorMenu* SubMenu, TSharedRef<FCollectionAssetManagement> QuickAssetManagement)
+		static void CreateManageCollectionsSubMenu(UToolMenu* SubMenu, TSharedRef<FCollectionAssetManagement> QuickAssetManagement)
 		{
 			FCollectionManagerModule& CollectionManagerModule = FCollectionManagerModule::GetModule();
 
@@ -1570,7 +1570,7 @@ bool FAssetContextMenu::AddCollectionMenuOptions(UEditorMenu* Menu)
 			CreateManageCollectionsSubMenu(SubMenu, QuickAssetManagement, MoveTemp(AvailableCollections));
 		}
 
-		static void CreateManageCollectionsSubMenu(UEditorMenu* SubMenu, TSharedRef<FCollectionAssetManagement> QuickAssetManagement, TArray<FCollectionNameType> AvailableCollections)
+		static void CreateManageCollectionsSubMenu(UToolMenu* SubMenu, TSharedRef<FCollectionAssetManagement> QuickAssetManagement, TArray<FCollectionNameType> AvailableCollections)
 		{
 			FCollectionManagerModule& CollectionManagerModule = FCollectionManagerModule::GetModule();
 
@@ -1579,7 +1579,7 @@ bool FAssetContextMenu::AddCollectionMenuOptions(UEditorMenu* Menu)
 				return One.Name.LexicalLess(Two.Name);
 			});
 
-			FEditorMenuSection& Section = SubMenu->AddSection("Section");
+			FToolMenuSection& Section = SubMenu->AddSection("Section");
 			for (const FCollectionNameType& AvailableCollection : AvailableCollections)
 			{
 				// Never display system collections
@@ -1601,12 +1601,12 @@ bool FAssetContextMenu::AddCollectionMenuOptions(UEditorMenu* Menu)
 
 				if (AvailableChildCollections.Num() > 0)
 				{
-					Section.AddEntry(FEditorMenuEntry::InitSubMenu(
+					Section.AddEntry(FToolMenuEntry::InitSubMenu(
 						SubMenu->GetMenuName(),
 						NAME_None,
 						FText::FromName(AvailableCollection.Name), 
 						FText::GetEmpty(), 
-						FNewEditorMenuDelegate::CreateStatic(&FManageCollectionsContextMenu::CreateManageCollectionsSubMenu, QuickAssetManagement, AvailableChildCollections),
+						FNewToolMenuDelegate::CreateStatic(&FManageCollectionsContextMenu::CreateManageCollectionsSubMenu, QuickAssetManagement, AvailableChildCollections),
 						FUIAction(
 							FExecuteAction::CreateStatic(&FManageCollectionsContextMenu::OnCollectionClicked, QuickAssetManagement, AvailableCollection),
 							FCanExecuteAction::CreateStatic(&FManageCollectionsContextMenu::IsCollectionEnabled, QuickAssetManagement, AvailableCollection),
@@ -1665,7 +1665,7 @@ bool FAssetContextMenu::AddCollectionMenuOptions(UEditorMenu* Menu)
 
 	FCollectionManagerModule& CollectionManagerModule = FCollectionManagerModule::GetModule();
 
-	FEditorMenuSection& Section = Menu->AddSection("AssetContextCollections", LOCTEXT("AssetCollectionOptionsMenuHeading", "Collections"));
+	FToolMenuSection& Section = Menu->AddSection("AssetContextCollections", LOCTEXT("AssetCollectionOptionsMenuHeading", "Collections"));
 
 	// Show a sub-menu that allows you to quickly add or remove the current asset selection from the available collections
 	if (CollectionManagerModule.Get().HasCollections())
@@ -1673,12 +1673,12 @@ bool FAssetContextMenu::AddCollectionMenuOptions(UEditorMenu* Menu)
 		TSharedRef<FCollectionAssetManagement> QuickAssetManagement = MakeShareable(new FCollectionAssetManagement());
 		QuickAssetManagement->SetCurrentAssets(SelectedAssets);
 
-		Section.AddEntry(FEditorMenuEntry::InitSubMenu(
+		Section.AddEntry(FToolMenuEntry::InitSubMenu(
 			Menu->GetMenuName(),
 			"ManageCollections",
 			LOCTEXT("ManageCollections", "Manage Collections"),
 			LOCTEXT("ManageCollections_ToolTip", "Manage the collections that the selected asset(s) belong to."),
-			FNewEditorMenuDelegate::CreateStatic(&FManageCollectionsContextMenu::CreateManageCollectionsSubMenu, QuickAssetManagement)
+			FNewToolMenuDelegate::CreateStatic(&FManageCollectionsContextMenu::CreateManageCollectionsSubMenu, QuickAssetManagement)
 			));
 
 		bHasAddedItems = true;
@@ -3185,7 +3185,7 @@ void FAssetContextMenu::GetSelectedPackages(TArray<UPackage*>& OutPackages) cons
 	}
 }
 
-void FAssetContextMenu::MakeChunkIDListMenu(UEditorMenu* Menu)
+void FAssetContextMenu::MakeChunkIDListMenu(UToolMenu* Menu)
 {
 	TArray<int32> FoundChunks;
 	TArray< FAssetData > AssetViewSelectedAssets = AssetView.Pin()->GetSelectedAssets();
@@ -3202,7 +3202,7 @@ void FAssetContextMenu::MakeChunkIDListMenu(UEditorMenu* Menu)
 		}
 	}
 
-	FEditorMenuSection& Section = Menu->AddSection("Chunks");
+	FToolMenuSection& Section = Menu->AddSection("Chunks");
 	for (auto ChunkID : FoundChunks)
 	{
 		Section.AddMenuEntry(
diff --git a/Engine/Source/Editor/ContentBrowser/Private/AssetContextMenu.h b/Engine/Source/Editor/ContentBrowser/Private/AssetContextMenu.h
--- a/Engine/Source/Editor/ContentBrowser/Private/AssetContextMenu.h
+++ b/Engine/Source/Editor/ContentBrowser/Private/AssetContextMenu.h
@@ -9,7 +9,7 @@
 #include "ContentBrowserDelegates.h"
 #include "SourcesData.h"
 
-class UEditorMenu;
+class UToolMenu;
 class FUICommandList;
 class SAssetView;
 class SWindow;
@@ -118,34 +118,34 @@ private:
 
 private:
 	/** Adds asset type-specific menu options to a menu builder. Returns true if any options were added. */
-	bool AddAssetTypeMenuOptions(UEditorMenu* Menu, TArray<UObject*>& SelectedObjects);
+	bool AddAssetTypeMenuOptions(UToolMenu* Menu, TArray<UObject*>& SelectedObjects);
 
 	/** Adds asset type-specific menu options to a menu builder. Returns true if any options were added. */
-	bool AddImportedAssetMenuOptions(UEditorMenu* Menu);
+	bool AddImportedAssetMenuOptions(UToolMenu* Menu);
 	
 	/** Adds common menu options to a menu builder. Returns true if any options were added. */
-	bool AddCommonMenuOptions(UEditorMenu* Menu);
+	bool AddCommonMenuOptions(UToolMenu* Menu);
 
 	/** Adds explore menu options to a menu builder. */
-	void AddExploreMenuOptions(UEditorMenu* Menu);
+	void AddExploreMenuOptions(UToolMenu* Menu);
 
 	/** Adds Asset Actions sub-menu to a menu builder. */
-	void MakeAssetActionsSubMenu(UEditorMenu* Menu);
+	void MakeAssetActionsSubMenu(UToolMenu* Menu);
 
 	/** Handler to check to see if "Asset Actions" are allowed */
 	bool CanExecuteAssetActions() const;
 
 	/** Adds Asset Localization sub-menu to a menu builder. */
-	void MakeAssetLocalizationSubMenu(UEditorMenu* Menu);
+	void MakeAssetLocalizationSubMenu(UToolMenu* Menu);
 
 	/** Adds the Create Localized Asset sub-menu to a menu builder. */
-	void MakeCreateLocalizedAssetSubMenu(UEditorMenu* Menu, TSet<FName> InSelectedSourceAssets, TArray<FLocalizedAssetsState> InLocalizedAssetsState);
+	void MakeCreateLocalizedAssetSubMenu(UToolMenu* Menu, TSet<FName> InSelectedSourceAssets, TArray<FLocalizedAssetsState> InLocalizedAssetsState);
 
 	/** Adds the Show Localized Assets sub-menu to a menu builder. */
-	void MakeShowLocalizedAssetSubMenu(UEditorMenu* Menu, TArray<FLocalizedAssetsState> InLocalizedAssetsState);
+	void MakeShowLocalizedAssetSubMenu(UToolMenu* Menu, TArray<FLocalizedAssetsState> InLocalizedAssetsState);
 
 	/** Adds the Edit Localized Assets sub-menu to a menu builder. */
-	void MakeEditLocalizedAssetSubMenu(UEditorMenu* Menu, TArray<FLocalizedAssetsState> InLocalizedAssetsState);
+	void MakeEditLocalizedAssetSubMenu(UToolMenu* Menu, TArray<FLocalizedAssetsState> InLocalizedAssetsState);
 
 	/** Create new localized assets for the given culture */
 	void ExecuteCreateLocalizedAsset(TSet<FName> InSelectedSourceAssets, FLocalizedAssetsState InLocalizedAssetsStateForCulture);
@@ -157,25 +157,25 @@ private:
 	void ExecuteOpenEditorsForAssets(TArray<FName> InAssets);
 
 	/** Adds asset reference menu options to a menu builder. Returns true if any options were added. */
-	bool AddReferenceMenuOptions(UEditorMenu* Menu);
+	bool AddReferenceMenuOptions(UToolMenu* Menu);
 
 	/** Adds asset documentation menu options to a menu builder. Returns true if any options were added. */
-	bool AddDocumentationMenuOptions(UEditorMenu* Menu);
+	bool AddDocumentationMenuOptions(UToolMenu* Menu);
 	
 	/** Adds source control menu options to a menu builder. */
-	bool AddSourceControlMenuOptions(UEditorMenu* Menu);
+	bool AddSourceControlMenuOptions(UToolMenu* Menu);
 
 	/** Fills the source control sub-menu */
-	void FillSourceControlSubMenu(UEditorMenu* Menu);
+	void FillSourceControlSubMenu(UToolMenu* Menu);
 
 	/** Handler to check to see if SCC actions are allowed */
 	bool CanExecuteSourceControlActions() const;
 
 	/** Adds menu options related to working with collections */
-	bool AddCollectionMenuOptions(UEditorMenu* Menu);
+	bool AddCollectionMenuOptions(UToolMenu* Menu);
 
 	/** Creates a sub-menu of Chunk IDs that are are assigned to all selected assets */
-	void MakeChunkIDListMenu(UEditorMenu* Menu);
+	void MakeChunkIDListMenu(UToolMenu* Menu);
 
 	/** Handler for when sync to asset tree is selected */
 	void ExecuteSyncToAssetTree();
diff --git a/Engine/Source/Editor/DetailCustomizations/Private/ActorDetails.cpp b/Engine/Source/Editor/DetailCustomizations/Private/ActorDetails.cpp
--- a/Engine/Source/Editor/DetailCustomizations/Private/ActorDetails.cpp
+++ b/Engine/Source/Editor/DetailCustomizations/Private/ActorDetails.cpp
@@ -18,7 +18,7 @@
 #include "Widgets/Images/SImage.h"
 #include "Widgets/Text/STextBlock.h"
 #include "Widgets/Layout/SBox.h"
-#include "EditorMenuSubsystem.h"
+#include "ToolMenus.h"
 #include "Widgets/Input/SButton.h"
 #include "Widgets/Input/SComboButton.h"
 #include "EditorStyleSet.h"
@@ -399,18 +399,18 @@ bool FActorDetails::AreAnySelectedActorsInLevelScript() const
 /** Util to create a menu for events we can add for the selected actor */
 TSharedRef<SWidget> FActorDetails::MakeEventOptionsWidgetFromSelection()
 {
-	UEditorMenuSubsystem* EditorMenus = UEditorMenuSubsystem::Get();
+	UToolMenus* ToolMenus = UToolMenus::Get();
 	static const FName MenuName("DetailCustomizations.EventOptions");
-	if (!EditorMenus->IsMenuRegistered(MenuName))
+	if (!ToolMenus->IsMenuRegistered(MenuName))
 	{
-		EditorMenus->RegisterMenu(MenuName);
+		ToolMenus->RegisterMenu(MenuName);
 	}
 
-	FEditorMenuContext Context;
-	UEditorMenu* Menu = EditorMenus->GenerateMenu(MenuName, Context);
+	FToolMenuContext Context;
+	UToolMenu* Menu = ToolMenus->GenerateMenu(MenuName, Context);
 	AActor* Actor = SelectedActors[0].Get();
 	FKismetEditorUtilities::AddLevelScriptEventOptionsForActor(Menu, SelectedActors[0], true, true, false);
-	return EditorMenus->GenerateWidget(Menu);
+	return ToolMenus->GenerateWidget(Menu);
 }
 
 
diff --git a/Engine/Source/Editor/HierarchicalLODOutliner/Private/HLODOutliner.cpp b/Engine/Source/Editor/HierarchicalLODOutliner/Private/HLODOutliner.cpp
--- a/Engine/Source/Editor/HierarchicalLODOutliner/Private/HLODOutliner.cpp
+++ b/Engine/Source/Editor/HierarchicalLODOutliner/Private/HLODOutliner.cpp
@@ -12,7 +12,7 @@
 #include "SlateOptMacros.h"
 #include "Framework/MultiBox/MultiBoxExtender.h"
 #include "Framework/MultiBox/MultiBoxBuilder.h"
-#include "EditorMenuSubsystem.h"
+#include "ToolMenus.h"
 #include "Widgets/Input/SButton.h"
 #include "Widgets/Images/SImage.h"
 #include "EditorStyleSet.h"
@@ -1492,18 +1492,18 @@ namespace HLODOutliner
 
 		if (SelectedItems.Num() && bSameType)
 		{
-			UEditorMenuSubsystem* EditorMenus = UEditorMenuSubsystem::Get();
+			UToolMenus* ToolMenus = UToolMenus::Get();
 			static const FName MenuName = "HierarchicalLODOutliner.HLODOutlinerContextMenu";
-			if (!EditorMenus->IsMenuRegistered(MenuName))
+			if (!ToolMenus->IsMenuRegistered(MenuName))
 			{
-				EditorMenus->RegisterMenu(MenuName);
+				ToolMenus->RegisterMenu(MenuName);
 			}
 
 			// Build up the menu for a selection
-			FEditorMenuContext Context;
-			UEditorMenu* Menu = EditorMenus->GenerateMenu(MenuName, Context);
+			FToolMenuContext Context;
+			UToolMenu* Menu = ToolMenus->GenerateMenu(MenuName, Context);
 			TreeView->GetSelectedItems()[0]->GenerateContextMenu(Menu, *this);
-			return EditorMenus->GenerateWidget(Menu);
+			return ToolMenus->GenerateWidget(Menu);
 		}
 
 		return TSharedPtr<SWidget>();
diff --git a/Engine/Source/Editor/HierarchicalLODOutliner/Private/ITreeItem.h b/Engine/Source/Editor/HierarchicalLODOutliner/Private/ITreeItem.h
--- a/Engine/Source/Editor/HierarchicalLODOutliner/Private/ITreeItem.h
+++ b/Engine/Source/Editor/HierarchicalLODOutliner/Private/ITreeItem.h
@@ -6,7 +6,7 @@
 #include "Widgets/SWidget.h"
 #include "Editor/HierarchicalLODOutliner/Private/HLODOutlinerDragDrop.h"
 
-class UEditorMenu;
+class UToolMenu;
 
 namespace HLODOutliner
 {
@@ -100,7 +100,7 @@ namespace HLODOutliner
 		virtual void OnExpansionChanged() {};
 
 		/** Generate a context menu for this item. Only called if *only* this item is selected. */
-		virtual void GenerateContextMenu(UEditorMenu* Menu, SHLODOutliner& Outliner) {};
+		virtual void GenerateContextMenu(UToolMenu* Menu, SHLODOutliner& Outliner) {};
 
 		/** Populate the specified drag/drop payload with any relevant information for this type */
 		virtual void PopulateDragDropPayload(FDragDropPayload& Payload) const = 0;
diff --git a/Engine/Source/Editor/HierarchicalLODOutliner/Private/LODActorItem.cpp b/Engine/Source/Editor/HierarchicalLODOutliner/Private/LODActorItem.cpp
--- a/Engine/Source/Editor/HierarchicalLODOutliner/Private/LODActorItem.cpp
+++ b/Engine/Source/Editor/HierarchicalLODOutliner/Private/LODActorItem.cpp
@@ -3,7 +3,7 @@
 #include "LODActorItem.h"
 #include "Textures/SlateIcon.h"
 #include "Framework/Commands/UIAction.h"
-#include "EditorMenuSubsystem.h"
+#include "ToolMenus.h"
 #include "Modules/ModuleManager.h"
 #include "HLODOutliner.h"
 #include "IHierarchicalLODUtilities.h"
@@ -24,11 +24,11 @@ bool HLODOutliner::FLODActorItem::CanInteract() const
 	return true;
 }
 
-void HLODOutliner::FLODActorItem::GenerateContextMenu(UEditorMenu* Menu, SHLODOutliner& Outliner)
+void HLODOutliner::FLODActorItem::GenerateContextMenu(UToolMenu* Menu, SHLODOutliner& Outliner)
 {
 	auto SharedOutliner = StaticCastSharedRef<SHLODOutliner>(Outliner.AsShared());
 
-	FEditorMenuSection& Section = Menu->AddSection("Section");
+	FToolMenuSection& Section = Menu->AddSection("Section");
 	Section.AddMenuEntry("SelectLODActor", LOCTEXT("SelectLODActor", "Select LOD Actor"), FText(), FSlateIcon(), FUIAction(FExecuteAction::CreateRaw(&Outliner, &SHLODOutliner::SelectLODActor)));
 
 	Section.AddMenuEntry("SelectContainedActors", LOCTEXT("SelectContainedActors", "Select Contained Actors"), FText(), FSlateIcon(), FUIAction(FExecuteAction::CreateRaw(&Outliner, &SHLODOutliner::SelectContainedActors)));
diff --git a/Engine/Source/Editor/HierarchicalLODOutliner/Private/LODActorItem.h b/Engine/Source/Editor/HierarchicalLODOutliner/Private/LODActorItem.h
--- a/Engine/Source/Editor/HierarchicalLODOutliner/Private/LODActorItem.h
+++ b/Engine/Source/Editor/HierarchicalLODOutliner/Private/LODActorItem.h
@@ -9,7 +9,7 @@
 #include "Engine/LODActor.h"
 #include "Editor/HierarchicalLODOutliner/Private/ITreeItem.h"
 
-class UEditorMenu;
+class UToolMenu;
 
 namespace HLODOutliner
 {
@@ -39,7 +39,7 @@ namespace HLODOutliner
 
 		//~ Begin ITreeItem Interface.
 		virtual bool CanInteract() const override;
-		virtual void GenerateContextMenu(UEditorMenu* Menu, SHLODOutliner& Outliner) override;
+		virtual void GenerateContextMenu(UToolMenu* Menu, SHLODOutliner& Outliner) override;
 		virtual FString GetDisplayString() const override;
 		virtual FSlateColor GetTint() const override;
 		virtual FTreeItemID GetID() override;
diff --git a/Engine/Source/Editor/HierarchicalLODOutliner/Private/LODLevelItem.cpp b/Engine/Source/Editor/HierarchicalLODOutliner/Private/LODLevelItem.cpp
--- a/Engine/Source/Editor/HierarchicalLODOutliner/Private/LODLevelItem.cpp
+++ b/Engine/Source/Editor/HierarchicalLODOutliner/Private/LODLevelItem.cpp
@@ -7,7 +7,7 @@
 #include "IHierarchicalLODUtilities.h"
 #include "HierarchicalLODUtilitiesModule.h"
 #include "Engine/Selection.h"
-#include "EditorMenuSubsystem.h"
+#include "ToolMenus.h"
 #include "HLODOutliner.h"
 #include "Editor.h"
 
@@ -25,7 +25,7 @@ bool HLODOutliner::FLODLevelItem::CanInteract() const
 	return true;
 }
 
-void HLODOutliner::FLODLevelItem::GenerateContextMenu(UEditorMenu* Menu, SHLODOutliner& Outliner)
+void HLODOutliner::FLODLevelItem::GenerateContextMenu(UToolMenu* Menu, SHLODOutliner& Outliner)
 {
 	// Collect actors / components available for creating a new HLOD cluster 
 	USelection* CurrentSelection = GEditor->GetSelectedActors();	
@@ -45,7 +45,7 @@ void HLODOutliner::FLODLevelItem::GenerateContextMenu(UEditorMenu* Menu, SHLODOu
 		}
 
 		// Generate context menu with option to create a new cluster from the current viewport selection
-		FEditorMenuSection& Section = Menu->AddSection("Section");
+		FToolMenuSection& Section = Menu->AddSection("Section");
 		Section.AddMenuEntry("CreateClusterFromViewportSelection", LOCTEXT("CreateClusterFromViewportSelection", "Create Cluster from World Selection"), FText(), FSlateIcon(), 
 			FUIAction(FExecuteAction::CreateLambda([&Outliner, ValidActorsToCluster, this]()
 			{
diff --git a/Engine/Source/Editor/HierarchicalLODOutliner/Private/LODLevelItem.h b/Engine/Source/Editor/HierarchicalLODOutliner/Private/LODLevelItem.h
--- a/Engine/Source/Editor/HierarchicalLODOutliner/Private/LODLevelItem.h
+++ b/Engine/Source/Editor/HierarchicalLODOutliner/Private/LODLevelItem.h
@@ -6,7 +6,7 @@
 #include "TreeItemID.h"
 #include "ITreeItem.h"
 
-class UEditorMenu;
+class UToolMenu;
 
 namespace HLODOutliner
 {
@@ -37,7 +37,7 @@ namespace HLODOutliner
 
 		//~ Begin ITreeItem Interface.
 		virtual bool CanInteract() const override;
-		virtual void GenerateContextMenu(UEditorMenu* Menu, SHLODOutliner& Outliner) override;
+		virtual void GenerateContextMenu(UToolMenu* Menu, SHLODOutliner& Outliner) override;
 		virtual FString GetDisplayString() const override;
 		virtual FTreeItemID GetID() override;
 		//~ Begin ITreeItem Interface.
diff --git a/Engine/Source/Editor/HierarchicalLODOutliner/Private/StaticMeshActorItem.cpp b/Engine/Source/Editor/HierarchicalLODOutliner/Private/StaticMeshActorItem.cpp
--- a/Engine/Source/Editor/HierarchicalLODOutliner/Private/StaticMeshActorItem.cpp
+++ b/Engine/Source/Editor/HierarchicalLODOutliner/Private/StaticMeshActorItem.cpp
@@ -4,7 +4,7 @@
 #include "GameFramework/Actor.h"
 #include "Textures/SlateIcon.h"
 #include "Framework/Commands/UIAction.h"
-#include "EditorMenuSubsystem.h"
+#include "ToolMenus.h"
 #include "HLODOutliner.h"
 
 #define LOCTEXT_NAMESPACE "StaticMeshActorItem"
@@ -20,10 +20,10 @@ bool HLODOutliner::FStaticMeshActorItem::CanInteract() const
 	return true;
 }
 
-void HLODOutliner::FStaticMeshActorItem::GenerateContextMenu(UEditorMenu* Menu, SHLODOutliner& Outliner)
+void HLODOutliner::FStaticMeshActorItem::GenerateContextMenu(UToolMenu* Menu, SHLODOutliner& Outliner)
 {
 	auto SharedOutliner = StaticCastSharedRef<SHLODOutliner>(Outliner.AsShared());
-	FEditorMenuSection& Section = Menu->AddSection("Section");
+	FToolMenuSection& Section = Menu->AddSection("Section");
 	Section.AddMenuEntry("RemoveSMActorFromCluster", LOCTEXT("RemoveSMActorFromCluster", "Remove From Cluster"), FText(), FSlateIcon(), FUIAction(FExecuteAction::CreateRaw(&Outliner, &SHLODOutliner::RemoveStaticMeshActorFromCluster)));
 	Section.AddMenuEntry("RemoveSMActorFromCluster", LOCTEXT("ExcludeSMActorFromGeneration", "Exclude From Cluster Generation"), FText(), FSlateIcon(), FUIAction(FExecuteAction::CreateRaw(&Outliner, &SHLODOutliner::ExcludeFromClusterGeneration)));
 }
diff --git a/Engine/Source/Editor/IntroTutorials/Private/IntroTutorials.cpp b/Engine/Source/Editor/IntroTutorials/Private/IntroTutorials.cpp
--- a/Engine/Source/Editor/IntroTutorials/Private/IntroTutorials.cpp
+++ b/Engine/Source/Editor/IntroTutorials/Private/IntroTutorials.cpp
@@ -37,7 +37,7 @@
 #include "AssetToolsModule.h"
 #include "IClassTypeActions.h"
 #include "ClassTypeActions_EditorTutorial.h"
-#include "EditorMenuSubsystem.h"
+#include "ToolMenus.h"
 
 #define LOCTEXT_NAMESPACE "IntroTutorials"
 
@@ -160,7 +160,7 @@ void FIntroTutorials::StartupModule()
 
 void FIntroTutorials::ShutdownModule()
 {
-	UEditorMenuSubsystem::UnregisterOwner(this);
+	UToolMenus::UnregisterOwner(this);
 
 	if (!bDisableTutorials && !IsRunningCommandlet())
 	{
@@ -219,10 +219,10 @@ void FIntroTutorials::ShutdownModule()
 
 void FIntroTutorials::RegisterSummonTutorialsMenuEntries()
 {
-	FEditorMenuOwnerScoped OwnerScoped(this);
-	UEditorMenu* HelpMenu = UEditorMenuSubsystem::Get()->ExtendMenu("LevelEditor.MainMenu.Help");
-	FEditorMenuSection& Section = HelpMenu->AddSection("Tutorials", LOCTEXT("TutorialsLabel", "Tutorials"), FEditorMenuInsert("HelpBrowse", EEditorMenuInsertType::After));
-	Section.AddEntry(FEditorMenuEntry::InitMenuEntry(
+	FToolMenuOwnerScoped OwnerScoped(this);
+	UToolMenu* HelpMenu = UToolMenus::Get()->ExtendMenu("LevelEditor.MainMenu.Help");
+	FToolMenuSection& Section = HelpMenu->AddSection("Tutorials", LOCTEXT("TutorialsLabel", "Tutorials"), FToolMenuInsert("HelpBrowse", EToolMenuInsertType::After));
+	Section.AddEntry(FToolMenuEntry::InitMenuEntry(
 		"Tutorials",
 		LOCTEXT("TutorialsMenuEntryTitle", "Tutorials"),
 		LOCTEXT("TutorialsMenuEntryToolTip", "Opens up introductory tutorials covering the basics of using the Unreal Engine 4 Editor."),
diff --git a/Engine/Source/Editor/Kismet/Private/SSCSEditor.cpp b/Engine/Source/Editor/Kismet/Private/SSCSEditor.cpp
--- a/Engine/Source/Editor/Kismet/Private/SSCSEditor.cpp
+++ b/Engine/Source/Editor/Kismet/Private/SSCSEditor.cpp
@@ -71,7 +71,7 @@
 #include "Algo/Find.h"
 #include "ActorEditorUtils.h"
 
-#include "EditorMenuSubsystem.h"
+#include "ToolMenus.h"
 #include "SSCSEditorMenuContext.h"
 #include "Kismet2/ComponentEditorContextMenuContex.h"
 
@@ -3921,7 +3921,7 @@ void SSCSEditor::GetSelectedItemsForContextMenu(TArray<FComponentEventConstructi
 	}
 }
 
-void SSCSEditor::PopulateContextMenu(UEditorMenu* Menu)
+void SSCSEditor::PopulateContextMenu(UToolMenu* Menu)
 {
 	TArray<FSCSEditorTreeNodePtrType> SelectedItems = SCSTreeWidget->GetSelectedItems();
 
@@ -3965,7 +3965,7 @@ void SSCSEditor::PopulateContextMenu(UEditorMenu* Menu)
 
 					if (EditorMode == EComponentEditorMode::BlueprintSCS)
 					{
-						FEditorMenuSection& BlueprintSCSSection = Menu->AddSection("BlueprintSCS");
+						FToolMenuSection& BlueprintSCSSection = Menu->AddSection("BlueprintSCS");
 						if (SelectedItems.Num() == 1)
 						{
 							BlueprintSCSSection.AddMenuEntry(FGraphEditorCommands::Get().FindReferences);
@@ -3989,7 +3989,7 @@ void SSCSEditor::PopulateContextMenu(UEditorMenu* Menu)
 							// Build an event submenu if we can generate events
 							if( FBlueprintEditorUtils::CanClassGenerateEvents( SelectedClass ))
 							{
-								BlueprintSCSSection.AddEntry(FEditorMenuEntry::InitSubMenu(
+								BlueprintSCSSection.AddEntry(FToolMenuEntry::InitSubMenu(
 									Menu->GetMenuName(),
 									"AddEventSubMenu",
 									LOCTEXT("AddEventSubMenu", "Add Event"), 
@@ -4012,7 +4012,7 @@ void SSCSEditor::PopulateContextMenu(UEditorMenu* Menu)
 
 		if (bOnlyShowPasteOption)
 		{
-			FEditorMenuSection& Section = Menu->AddSection("PasteComponent", LOCTEXT("EditComponentHeading", "Edit") );
+			FToolMenuSection& Section = Menu->AddSection("PasteComponent", LOCTEXT("EditComponentHeading", "Edit") );
 			{
 				Section.AddMenuEntry( FGenericCommands::Get().Paste );
 			}
@@ -4022,11 +4022,11 @@ void SSCSEditor::PopulateContextMenu(UEditorMenu* Menu)
 
 void SSCSEditor::RegisterContextMenu()
 {
-	UEditorMenuSubsystem* EditorMenus = UEditorMenuSubsystem::Get();
-	if (!EditorMenus->IsMenuRegistered("Kismet.SCSEditorContextMenu"))
+	UToolMenus* ToolMenus = UToolMenus::Get();
+	if (!ToolMenus->IsMenuRegistered("Kismet.SCSEditorContextMenu"))
 	{
-		UEditorMenu* Menu = EditorMenus->RegisterMenu("Kismet.SCSEditorContextMenu");
-		Menu->AddDynamicSection("SCSEditorDynamic", FNewEditorMenuDelegate::CreateLambda([](UEditorMenu* InMenu)
+		UToolMenu* Menu = ToolMenus->RegisterMenu("Kismet.SCSEditorContextMenu");
+		Menu->AddDynamicSection("SCSEditorDynamic", FNewToolMenuDelegate::CreateLambda([](UToolMenu* InMenu)
 		{
 			USSCSEditorMenuContext* ContextObject = InMenu->FindContext<USSCSEditorMenuContext>();
 			if (ContextObject && ContextObject->SCSEditor.IsValid())
@@ -4046,8 +4046,8 @@ TSharedPtr< SWidget > SSCSEditor::CreateContextMenu()
 		RegisterContextMenu();
 		USSCSEditorMenuContext* ContextObject = NewObject<USSCSEditorMenuContext>();
 		ContextObject->SCSEditor = SharedThis(this);
-		FEditorMenuContext EditorMenuContext(CommandList, TSharedPtr<FExtender>(), ContextObject);
-		return UEditorMenuSubsystem::Get()->GenerateWidget("Kismet.SCSEditorContextMenu", EditorMenuContext);
+		FToolMenuContext ToolMenuContext(CommandList, TSharedPtr<FExtender>(), ContextObject);
+		return UToolMenus::Get()->GenerateWidget("Kismet.SCSEditorContextMenu", ToolMenuContext);
 	}
 	return TSharedPtr<SWidget>();
 }
diff --git a/Engine/Source/Editor/Kismet/Public/SSCSEditor.h b/Engine/Source/Editor/Kismet/Public/SSCSEditor.h
--- a/Engine/Source/Editor/Kismet/Public/SSCSEditor.h
+++ b/Engine/Source/Editor/Kismet/Public/SSCSEditor.h
@@ -25,7 +25,7 @@
 #include "ScopedTransaction.h"
 
 class FMenuBuilder;
-class UEditorMenu;
+class UToolMenu;
 class FSCSEditorTreeNode;
 class SSCSEditor;
 class UPrimitiveComponent;
@@ -1022,7 +1022,7 @@ protected:
 	void RegisterContextMenu();
 
 	/** Populate context menu on the fly */
-	void PopulateContextMenu(UEditorMenu* InMenu);
+	void PopulateContextMenu(UToolMenu* InMenu);
 
 	/** Called when the level editor requests a component to be renamed. */
 	void OnLevelComponentRequestRename(const UActorComponent* InComponent);
diff --git a/Engine/Source/Editor/LevelEditor/Private/LevelEditorContextMenu.cpp b/Engine/Source/Editor/LevelEditor/Private/LevelEditorContextMenu.cpp
--- a/Engine/Source/Editor/LevelEditor/Private/LevelEditorContextMenu.cpp
+++ b/Engine/Source/Editor/LevelEditor/Private/LevelEditorContextMenu.cpp
@@ -13,7 +13,7 @@
 #include "Framework/Commands/UICommandList.h"
 #include "Framework/MultiBox/MultiBoxExtender.h"
 #include "Framework/MultiBox/MultiBoxBuilder.h"
-#include "EditorMenuSubsystem.h"
+#include "ToolMenus.h"
 #include "Components/ActorComponent.h"
 #include "GameFramework/Actor.h"
 #include "Kismet2/ComponentEditorUtils.h"
@@ -63,14 +63,14 @@ public:
 	 *
 	 * @param MenuBuilder	The menu to add items to
 	 */
-	static void FillSelectActorMenu(UEditorMenu* Menu);
+	static void FillSelectActorMenu(UToolMenu* Menu);
 
 	/**
 	 * Fills in menu options for the actor visibility menu
 	 *
 	 * @param MenuBuilder	The menu to add items to
 	 */
-	static void FillActorVisibilityMenu(UEditorMenu* Menu);
+	static void FillActorVisibilityMenu(UToolMenu* Menu);
 
 	/**
 	 * Fills in menu options for the actor level menu
@@ -79,42 +79,42 @@ public:
 	 * @param bAllInCurrentLevel	true if all selected actors are in the current level
 	 * @param MenuBuilder			The menu to add items to
 	 */
-	static void FillActorLevelMenu(UEditorMenu* Menu);
+	static void FillActorLevelMenu(UToolMenu* Menu);
 
 	/**
 	 * Fills in menu options for the transform menu
 	 *
 	 * @param MenuBuilder	The menu to add items to
 	 */
-	static void FillTransformMenu(UEditorMenu* Menu);
+	static void FillTransformMenu(UToolMenu* Menu);
 
 	/**
 	 * Fills in menu options for the Fill Actor menu
 	 *
 	 * @param MenuBuilder	The menu to add items to
 	 */
-	static void FillActorMenu(UEditorMenu* Menu);
+	static void FillActorMenu(UToolMenu* Menu);
 
 	/**
 	 * Fills in menu options for the snap menu
 	 *
 	 * @param MenuBuilder	The menu to add items to
 	 */
-	static void FillSnapAlignMenu(UEditorMenu* Menu);
+	static void FillSnapAlignMenu(UToolMenu* Menu);
 
 	/**
 	 * Fills in menu options for the pivot menu
 	 *
 	 * @param MenuBuilder	The menu to add items to
 	 */
-	static void FillPivotMenu(UEditorMenu* Menu);
+	static void FillPivotMenu(UToolMenu* Menu);
 	
 	/**
 	 * Fills in menu options for the group menu
 	 *
 	 * @param MenuBuilder	The menu to add items to
 	 */
-	static void FillGroupMenu( UEditorMenu* Menu );
+	static void FillGroupMenu( UToolMenu* Menu );
 
 	/**
 	 * Fills in menu options for the edit menu
@@ -122,7 +122,7 @@ public:
 	 * @param MenuBuilder	The menu to add items to
 	 * @param ContextType	The context for this editor menu
 	 */
-	static void FillEditMenu(UEditorMenu* Menu);
+	static void FillEditMenu(UToolMenu* Menu);
 
 private:
 	/**
@@ -130,7 +130,7 @@ private:
 	 *
 	 * @param MenuBuilder	The menu to add items to
 	 */
-	static void FillMatineeSelectActorMenu(UEditorMenu* Menu);
+	static void FillMatineeSelectActorMenu(UToolMenu* Menu);
 };
 
 FSelectedActorInfo FLevelEditorContextMenuImpl::SelectionInfo;
@@ -142,19 +142,19 @@ struct FLevelScriptEventMenuHelper
 	*
 	* @param MenuBuilder	The menu to add items to
 	*/
-	static void FillLevelBlueprintEventsMenu(UEditorMenu* Menu, const TArray<AActor*>& SelectedActors);
+	static void FillLevelBlueprintEventsMenu(UToolMenu* Menu, const TArray<AActor*>& SelectedActors);
 };
 
 void FLevelEditorContextMenu::RegisterComponentContextMenu()
 {
-	UEditorMenuSubsystem* EditorMenus = UEditorMenuSubsystem::Get();
-	if (EditorMenus->IsMenuRegistered("LevelEditor.ComponentContextMenu"))
+	UToolMenus* ToolMenus = UToolMenus::Get();
+	if (ToolMenus->IsMenuRegistered("LevelEditor.ComponentContextMenu"))
 	{
 		return;
 	}
 
-	UEditorMenu* Menu = EditorMenus->RegisterMenu("LevelEditor.ComponentContextMenu");
-	Menu->AddDynamicSection("ComponentControlDynamic", FNewEditorMenuDelegate::CreateLambda([](UEditorMenu* InMenu)
+	UToolMenu* Menu = ToolMenus->RegisterMenu("LevelEditor.ComponentContextMenu");
+	Menu->AddDynamicSection("ComponentControlDynamic", FNewToolMenuDelegate::CreateLambda([](UToolMenu* InMenu)
 	{
 		ULevelEditorContextMenuContext* LevelEditorContext = InMenu->FindContext<ULevelEditorContextMenuContext>();
 		if (!LevelEditorContext)
@@ -169,7 +169,7 @@ void FLevelEditorContextMenu::RegisterComponentContextMenu()
 		}
 
 		{
-			FEditorMenuSection& Section = InMenu->AddSection("ComponentControl", LOCTEXT("ComponentControlHeading", "Component"));
+			FToolMenuSection& Section = InMenu->AddSection("ComponentControl", LOCTEXT("ComponentControlHeading", "Component"));
 
 			AActor* OwnerActor = GEditor->GetSelectedActors()->GetTop<AActor>();
 			if(OwnerActor)
@@ -199,14 +199,14 @@ void FLevelEditorContextMenu::RegisterComponentContextMenu()
 
 void FLevelEditorContextMenu::RegisterActorContextMenu()
 {
-	UEditorMenuSubsystem* EditorMenus = UEditorMenuSubsystem::Get();
-	if (EditorMenus->IsMenuRegistered("LevelEditor.ActorContextMenu"))
+	UToolMenus* ToolMenus = UToolMenus::Get();
+	if (ToolMenus->IsMenuRegistered("LevelEditor.ActorContextMenu"))
 	{
 		return;
 	}
 
-	UEditorMenu* Menu = EditorMenus->RegisterMenu("LevelEditor.ActorContextMenu");
-	Menu->AddDynamicSection("ActorContextMenuDynamic", FNewEditorMenuDelegate::CreateLambda([](UEditorMenu* InMenu)
+	UToolMenu* Menu = ToolMenus->RegisterMenu("LevelEditor.ActorContextMenu");
+	Menu->AddDynamicSection("ActorContextMenuDynamic", FNewToolMenuDelegate::CreateLambda([](UToolMenu* InMenu)
 	{
 		ULevelEditorContextMenuContext* LevelEditorContext = InMenu->FindContext<ULevelEditorContextMenuContext>();
 		if (!LevelEditorContext || !LevelEditorContext->LevelEditor.IsValid())
@@ -232,7 +232,7 @@ void FLevelEditorContextMenu::RegisterActorContextMenu()
 		if (bCanSyncToContentBrowser || ReferencedAssets.Num() > 0)
 		{
 			{
-				FEditorMenuSection& Section = InMenu->AddSection("ActorAsset", LOCTEXT("AssetHeading", "Asset"));
+				FToolMenuSection& Section = InMenu->AddSection("ActorAsset", LOCTEXT("AssetHeading", "Asset"));
 				if (bCanSyncToContentBrowser)
 				{
 					Section.AddMenuEntry(FGlobalEditorCommonCommands::Get().FindInContentBrowser);
@@ -264,7 +264,7 @@ void FLevelEditorContextMenu::RegisterActorContextMenu()
 
 
 		{
-			FEditorMenuSection& Section = InMenu->AddSection("ActorControl", LOCTEXT("ActorHeading", "Actor"));
+			FToolMenuSection& Section = InMenu->AddSection("ActorControl", LOCTEXT("ActorHeading", "Actor"));
 			Section.AddMenuEntry(FEditorViewportCommands::Get().FocusViewportToSelection);
 
 
@@ -317,7 +317,7 @@ void FLevelEditorContextMenu::RegisterActorContextMenu()
 				{
 					const FString CodeFileName = FPaths::GetCleanFilename(*ClassHeaderPath);
 
-					FEditorMenuSection& Section = InMenu->AddSection("ActorCode", LOCTEXT("ActorCodeHeading", "C++"));
+					FToolMenuSection& Section = InMenu->AddSection("ActorCode", LOCTEXT("ActorCodeHeading", "C++"));
 					{
 						Section.AddMenuEntry(FLevelEditorCommands::Get().GoToCodeForActor,
 							FText::Format(LOCTEXT("GoToCodeForActor", "Open {0}"), FText::FromString(CodeFileName)),
@@ -329,7 +329,7 @@ void FLevelEditorContextMenu::RegisterActorContextMenu()
 			const FString DocumentationLink = FEditorClassUtils::GetDocumentationLink(SelectionInfo.SelectionClass);
 			if (!DocumentationLink.IsEmpty())
 			{
-				FEditorMenuSection& Section = InMenu->AddSection("ActorDocumentation", LOCTEXT("ActorDocsHeading", "Documentation"));
+				FToolMenuSection& Section = InMenu->AddSection("ActorDocumentation", LOCTEXT("ActorDocsHeading", "Documentation"));
 				{
 					Section.AddMenuEntry(FLevelEditorCommands::Get().GoToDocsForActor,
 						LOCTEXT("GoToDocsForActor", "View Documentation"),
@@ -340,46 +340,46 @@ void FLevelEditorContextMenu::RegisterActorContextMenu()
 		}
 
 		{
-			FEditorMenuSection& Section = InMenu->AddSection("ActorSelectVisibilityLevels");
+			FToolMenuSection& Section = InMenu->AddSection("ActorSelectVisibilityLevels");
 
 			// Add a sub-menu for "Select"
-			Section.AddEntry(FEditorMenuEntry::InitSubMenu(
+			Section.AddEntry(FToolMenuEntry::InitSubMenu(
 				InMenu->MenuName,
 				"SelectSubMenu",
 				LOCTEXT("SelectSubMenu", "Select"),
 				LOCTEXT("SelectSubMenu_ToolTip", "Opens the actor selection menu"),
-				FNewEditorMenuDelegate::CreateStatic(&FLevelEditorContextMenuImpl::FillSelectActorMenu)));
+				FNewToolMenuDelegate::CreateStatic(&FLevelEditorContextMenuImpl::FillSelectActorMenu)));
 
-			Section.AddEntry(FEditorMenuEntry::InitSubMenu(
+			Section.AddEntry(FToolMenuEntry::InitSubMenu(
 				InMenu->MenuName,
 				"EditSubMenu",
 				LOCTEXT("EditSubMenu", "Edit"),
 				FText::GetEmpty(),
-				FNewEditorMenuDelegate::CreateStatic(&FLevelEditorContextMenuImpl::FillEditMenu)));
+				FNewToolMenuDelegate::CreateStatic(&FLevelEditorContextMenuImpl::FillEditMenu)));
 
-			Section.AddEntry(FEditorMenuEntry::InitSubMenu(
+			Section.AddEntry(FToolMenuEntry::InitSubMenu(
 				InMenu->MenuName,
 				"VisibilitySubMenu",
 				LOCTEXT("VisibilitySubMenu", "Visibility"),
 				LOCTEXT("VisibilitySubMenu_ToolTip", "Selected actor visibility options"),
-				FNewEditorMenuDelegate::CreateStatic(&FLevelEditorContextMenuImpl::FillActorVisibilityMenu)));
+				FNewToolMenuDelegate::CreateStatic(&FLevelEditorContextMenuImpl::FillActorVisibilityMenu)));
 
 			// Build the menu for grouping actors
 			BuildGroupMenu(InMenu, SelectionInfo);
 
-			Section.AddEntry(FEditorMenuEntry::InitSubMenu(
+			Section.AddEntry(FToolMenuEntry::InitSubMenu(
 				InMenu->MenuName,
 				"LevelSubMenu",
 				LOCTEXT("LevelSubMenu", "Level"),
 				LOCTEXT("LevelSubMenu_ToolTip", "Options for interacting with this actor's level"),
-				FNewEditorMenuDelegate::CreateStatic(&FLevelEditorContextMenuImpl::FillActorLevelMenu)));
+				FNewToolMenuDelegate::CreateStatic(&FLevelEditorContextMenuImpl::FillActorLevelMenu)));
 		}
 
 		if (LevelEditorContext->ContextType == ELevelEditorMenuContext::Viewport)
 		{
 			LevelEditorCreateActorMenu::FillAddReplaceViewportContextMenuSections(InMenu);
 
-			FEditorMenuSection& Section = InMenu->AddSection("OpenMergeActor");
+			FToolMenuSection& Section = InMenu->AddSection("OpenMergeActor");
 			Section.AddMenuEntry(FLevelEditorCommands::Get().OpenMergeActor,
 					LOCTEXT("OpenMergeActor", "Merge Actors"),
 					LOCTEXT("OpenMergeActor_ToolTip", "Click to open the Merge Actor panel"));
@@ -389,7 +389,7 @@ void FLevelEditorContextMenu::RegisterActorContextMenu()
 		{
 			if (SelectionInfo.NumSelected > 0)
 			{
-				FEditorMenuSection& Section = InMenu->AddSection("Simulation", NSLOCTEXT("LevelViewportContextMenu", "SimulationHeading", "Simulation"));
+				FToolMenuSection& Section = InMenu->AddSection("Simulation", NSLOCTEXT("LevelViewportContextMenu", "SimulationHeading", "Simulation"));
 				{
 					Section.AddMenuEntry(FLevelEditorCommands::Get().KeepSimulationChanges);
 				}
@@ -397,7 +397,7 @@ void FLevelEditorContextMenu::RegisterActorContextMenu()
 		}
 
 		{
-			FEditorMenuSection& Section = InMenu->AddSection("LevelViewportAttach");
+			FToolMenuSection& Section = InMenu->AddSection("LevelViewportAttach");
 
 			// Only display the attach menu if we have actors selected
 			if (GEditor->GetSelectedActorCount())
@@ -407,33 +407,33 @@ void FLevelEditorContextMenu::RegisterActorContextMenu()
 					Section.AddMenuEntry(FLevelEditorCommands::Get().DetachFromParent);
 				}
 
-				Section.AddEntry(FEditorMenuEntry::InitSubMenu(
+				Section.AddEntry(FToolMenuEntry::InitSubMenu(
 					InMenu->GetMenuName(),
 					"ActorAttachToSubMenu",
 					LOCTEXT("ActorAttachToSubMenu", "Attach To"),
 					LOCTEXT("ActorAttachToSubMenu_ToolTip", "Attach Actor as child"),
-					FNewEditorMenuDelegate::CreateStatic(&FLevelEditorContextMenuImpl::FillActorMenu)));
+					FNewToolMenuDelegate::CreateStatic(&FLevelEditorContextMenuImpl::FillActorMenu)));
 			}
 
 			// Add a heading for "Movement" if an actor is selected
 			if (GEditor->GetSelectedActorIterator())
 			{
 				// Add a sub-menu for "Transform"
-				Section.AddEntry(FEditorMenuEntry::InitSubMenu(
+				Section.AddEntry(FToolMenuEntry::InitSubMenu(
 					InMenu->GetMenuName(),
 					"TransformSubMenu",
 					LOCTEXT("TransformSubMenu", "Transform"),
 					LOCTEXT("TransformSubMenu_ToolTip", "Actor transform utils"),
-					FNewEditorMenuDelegate::CreateStatic(&FLevelEditorContextMenuImpl::FillTransformMenu)));
+					FNewToolMenuDelegate::CreateStatic(&FLevelEditorContextMenuImpl::FillTransformMenu)));
 			}
 
 			// Add a sub-menu for "Pivot"
-			Section.AddEntry(FEditorMenuEntry::InitSubMenu(
+			Section.AddEntry(FToolMenuEntry::InitSubMenu(
 				InMenu->MenuName,
 				"PivotSubMenu",
 				LOCTEXT("PivotSubMenu", "Pivot"),
 				LOCTEXT("PivotSubMenu_ToolTip", "Actor pivoting utils"),
-				FNewEditorMenuDelegate::CreateStatic(&FLevelEditorContextMenuImpl::FillPivotMenu)));
+				FNewToolMenuDelegate::CreateStatic(&FLevelEditorContextMenuImpl::FillPivotMenu)));
 		}
 
 		FLevelScriptEventMenuHelper::FillLevelBlueprintEventsMenu(InMenu, SelectedActors);
@@ -442,27 +442,27 @@ void FLevelEditorContextMenu::RegisterActorContextMenu()
 
 void FLevelEditorContextMenu::RegisterSceneOutlinerContextMenu()
 {
-	UEditorMenuSubsystem* EditorMenus = UEditorMenuSubsystem::Get();
-	if (EditorMenus->IsMenuRegistered("LevelEditor.SceneOutlinerContextMenu"))
+	UToolMenus* ToolMenus = UToolMenus::Get();
+	if (ToolMenus->IsMenuRegistered("LevelEditor.SceneOutlinerContextMenu"))
 	{
 		return;
 	}
 
-	UEditorMenu* Menu = EditorMenus->RegisterMenu("LevelEditor.SceneOutlinerContextMenu");
-	Menu->AddDynamicSection("SelectVisibilityLevels", FNewEditorMenuDelegate::CreateLambda([](UEditorMenu* InMenu)
+	UToolMenu* Menu = ToolMenus->RegisterMenu("LevelEditor.SceneOutlinerContextMenu");
+	Menu->AddDynamicSection("SelectVisibilityLevels", FNewToolMenuDelegate::CreateLambda([](UToolMenu* InMenu)
 	{
 		if (ULevelEditorContextMenuContext* LevelEditorContext = InMenu->FindContext<ULevelEditorContextMenuContext>())
 		{
 			TWeakPtr<ISceneOutliner> SceneOutlinerPtr = LevelEditorContext->LevelEditor.Pin()->GetSceneOutliner();
 			if (SceneOutlinerPtr.IsValid())
 			{
-				FEditorMenuSection& Section = InMenu->AddSection("SelectVisibilityLevels");
-				Section.AddEntry(FEditorMenuEntry::InitSubMenu(
+				FToolMenuSection& Section = InMenu->AddSection("SelectVisibilityLevels");
+				Section.AddEntry(FToolMenuEntry::InitSubMenu(
 					InMenu->MenuName,
 					"EditSubMenu",
 					LOCTEXT("EditSubMenu", "Edit"),
 					FText::GetEmpty(),
-					FNewEditorMenuDelegate::CreateStatic(&FLevelEditorContextMenuImpl::FillEditMenu)
+					FNewToolMenuDelegate::CreateStatic(&FLevelEditorContextMenuImpl::FillEditMenu)
 				));
 			}
 		}
@@ -487,7 +487,7 @@ FName FLevelEditorContextMenu::GetContextMenuName(ELevelEditorMenuContext Contex
 	return NAME_None;
 }
 
-FName FLevelEditorContextMenu::InitMenuContext(FEditorMenuContext& Context, TWeakPtr<SLevelEditor> LevelEditor, ELevelEditorMenuContext ContextType)
+FName FLevelEditorContextMenu::InitMenuContext(FToolMenuContext& Context, TWeakPtr<SLevelEditor> LevelEditor, ELevelEditorMenuContext ContextType)
 {
 	RegisterComponentContextMenu();
 	RegisterActorContextMenu();
@@ -531,24 +531,24 @@ FName FLevelEditorContextMenu::InitMenuContext(FEditorMenuContext& Context, TWea
 	return GetContextMenuName(ContextType);
 }
 
-UEditorMenu* FLevelEditorContextMenu::GenerateMenu(TWeakPtr<SLevelEditor> LevelEditor, ELevelEditorMenuContext ContextType, TSharedPtr<FExtender> Extender)
+UToolMenu* FLevelEditorContextMenu::GenerateMenu(TWeakPtr<SLevelEditor> LevelEditor, ELevelEditorMenuContext ContextType, TSharedPtr<FExtender> Extender)
 {
-	FEditorMenuContext Context;
+	FToolMenuContext Context;
 	if (Extender.IsValid())
 	{
 		Context.AddExtender(Extender);
 	}
 
 	FName ContextMenuName = InitMenuContext(Context, LevelEditor, ContextType);
-	return UEditorMenuSubsystem::Get()->GenerateMenu(ContextMenuName, Context);
+	return UToolMenus::Get()->GenerateMenu(ContextMenuName, Context);
 }
 
 // NOTE: We intentionally receive a WEAK pointer here because we want to be callable by a delegate whose
 //       payload contains a weak reference to a level editor instance
 TSharedPtr< SWidget > FLevelEditorContextMenu::BuildMenuWidget(TWeakPtr< SLevelEditor > LevelEditor, ELevelEditorMenuContext ContextType, TSharedPtr<FExtender> Extender)
 {
-	UEditorMenu* Menu = GenerateMenu(LevelEditor, ContextType, Extender);
-	return UEditorMenuSubsystem::Get()->GenerateWidget(Menu);
+	UToolMenu* Menu = GenerateMenu(LevelEditor, ContextType, Extender);
+	return UToolMenus::Get()->GenerateWidget(Menu);
 }
 
 namespace EViewOptionType
@@ -703,11 +703,11 @@ FSlateColor InvertOnHover( const TWeakPtr< SWidget > WidgetPtr )
 	return FSlateColor::UseForeground();
 }
 
-void FLevelEditorContextMenu::BuildGroupMenu(UEditorMenu* Menu, const FSelectedActorInfo& SelectedActorInfo)
+void FLevelEditorContextMenu::BuildGroupMenu(UToolMenu* Menu, const FSelectedActorInfo& SelectedActorInfo)
 {
 	if( UActorGroupingUtils::IsGroupingActive() )
 	{
-		FEditorMenuSection& Section = Menu->AddSection("GroupMenu");
+		FToolMenuSection& Section = Menu->AddSection("GroupMenu");
 
 		// Whether or not we added a grouping sub-menu
 		bool bNeedGroupSubMenu = SelectedActorInfo.bHaveSelectedLockedGroup || SelectedActorInfo.bHaveSelectedUnlockedGroup;
@@ -729,23 +729,23 @@ void FLevelEditorContextMenu::BuildGroupMenu(UEditorMenu* Menu, const FSelectedA
 		
 		if( bNeedGroupSubMenu )
 		{
-			Section.AddEntry(FEditorMenuEntry::InitSubMenu(
+			Section.AddEntry(FToolMenuEntry::InitSubMenu(
 				Menu->GetMenuName(),
 				"GroupMenu",
 				LOCTEXT("GroupMenu", "Groups"),
 				LOCTEXT("GroupMenu_ToolTip", "Opens the actor grouping menu"),
-				FNewEditorMenuDelegate::CreateStatic( &FLevelEditorContextMenuImpl::FillGroupMenu)));
+				FNewToolMenuDelegate::CreateStatic( &FLevelEditorContextMenuImpl::FillGroupMenu)));
 		}
 	}
 }
 
-void FLevelEditorContextMenuImpl::FillSelectActorMenu(UEditorMenu* Menu)
+void FLevelEditorContextMenuImpl::FillSelectActorMenu(UToolMenu* Menu)
 {
 	FText SelectAllActorStr = FText::Format( LOCTEXT("SelectActorsOfSameClass", "Select All {0}(s)"), FText::FromString( SelectionInfo.SelectionStr ) );
 	int32 NumSelectedSurfaces = AssetSelectionUtils::GetNumSelectedSurfaces( SelectionInfo.SharedWorld );
 
 	{
-		FEditorMenuSection& Section = Menu->AddSection("SelectActorGeneral", LOCTEXT("SelectAnyHeading", "General"));
+		FToolMenuSection& Section = Menu->AddSection("SelectActorGeneral", LOCTEXT("SelectAnyHeading", "General"));
 		Section.AddMenuEntry( FGenericCommands::Get().SelectAll, TAttribute<FText>(), LOCTEXT("SelectAll_ToolTip", "Selects all actors") );
 		Section.AddMenuEntry( FLevelEditorCommands::Get().SelectNone );
 		Section.AddMenuEntry( FLevelEditorCommands::Get().InvertSelection );
@@ -754,19 +754,19 @@ void FLevelEditorContextMenuImpl::FillSelectActorMenu(UEditorMenu* Menu)
 	if( !SelectionInfo.bHaveBrush && SelectionInfo.bAllSelectedActorsOfSameType && SelectionInfo.SelectionStr.Len() != 0 )
 	{
 		// These menu options appear if only if all the actors are the same type and we aren't selecting brush
-		FEditorMenuSection& Section = Menu->AddSection("SelectAllActorsOfSameClass");
+		FToolMenuSection& Section = Menu->AddSection("SelectAllActorsOfSameClass");
 		Section.AddMenuEntry(FLevelEditorCommands::Get().SelectAllActorsOfSameClass, SelectAllActorStr);
 	}
 
 	{
-		FEditorMenuSection& Section = Menu->AddSection("SelectActorHierarchy", LOCTEXT("SelectHierarchyHeading", "Hierarchy"));
+		FToolMenuSection& Section = Menu->AddSection("SelectActorHierarchy", LOCTEXT("SelectHierarchyHeading", "Hierarchy"));
 		Section.AddMenuEntry( FLevelEditorCommands::Get().SelectImmediateChildren );
 		Section.AddMenuEntry( FLevelEditorCommands::Get().SelectAllDescendants );
 	}
 
 	// Add brush commands when we have a brush or any surfaces selected
 	{
-		FEditorMenuSection& Section = Menu->AddSection("SelectBSP", LOCTEXT("SelectBSPHeading", "BSP"));
+		FToolMenuSection& Section = Menu->AddSection("SelectBSP", LOCTEXT("SelectBSPHeading", "BSP"));
 		if( SelectionInfo.bHaveBrush || NumSelectedSurfaces > 0 )
 		{
 			if( SelectionInfo.bAllSelectedAreBrushes )
@@ -784,7 +784,7 @@ void FLevelEditorContextMenuImpl::FillSelectActorMenu(UEditorMenu* Menu)
 	{
 		// If any actors are selected add lights selection options
 		{
-			FEditorMenuSection& Section = Menu->AddSection("SelectLights", LOCTEXT("SelectLightHeading", "Lights"));
+			FToolMenuSection& Section = Menu->AddSection("SelectLights", LOCTEXT("SelectLightHeading", "Lights"));
 			Section.AddMenuEntry(FLevelEditorCommands::Get().SelectRelevantLights);
 
 			if ( SelectionInfo.bHaveLight )
@@ -798,7 +798,7 @@ void FLevelEditorContextMenuImpl::FillSelectActorMenu(UEditorMenu* Menu)
 		{
 			// if any static meshes are selected allow selecting actors using the same mesh
 			{
-				FEditorMenuSection& Section = Menu->AddSection("SelectMeshes", LOCTEXT("SelectStaticMeshHeading", "Static Meshes"));
+				FToolMenuSection& Section = Menu->AddSection("SelectMeshes", LOCTEXT("SelectStaticMeshHeading", "Static Meshes"));
 				Section.AddMenuEntry(FLevelEditorCommands::Get().SelectStaticMeshesOfSameClass, LOCTEXT("SelectStaticMeshesOfSameClass_Menu", "Select Matching (Selected Classes)"));
 				Section.AddMenuEntry(FLevelEditorCommands::Get().SelectStaticMeshesAllClasses, LOCTEXT("SelectStaticMeshesAllClasses_Menu", "Select Matching (All Classes)"));
 			}
@@ -806,7 +806,7 @@ void FLevelEditorContextMenuImpl::FillSelectActorMenu(UEditorMenu* Menu)
 			if (SelectionInfo.NumSelected == 1)
 			{
 				{
-					FEditorMenuSection& Section = Menu->AddSection("SelectHLODCluster", LOCTEXT("SelectHLODClusterHeading", "Hierachical LODs"));
+					FToolMenuSection& Section = Menu->AddSection("SelectHLODCluster", LOCTEXT("SelectHLODClusterHeading", "Hierachical LODs"));
 					Section.AddMenuEntry(FLevelEditorCommands::Get().SelectOwningHierarchicalLODCluster, LOCTEXT("SelectOwningHierarchicalLODCluster_Menu", "Select Owning HierarchicalLODCluster"));
 				}
 			}			
@@ -816,7 +816,7 @@ void FLevelEditorContextMenuImpl::FillSelectActorMenu(UEditorMenu* Menu)
 		{
 			// if any skeletal meshes are selected allow selecting actors using the same mesh
 			{
-				FEditorMenuSection& Section = Menu->AddSection("SelectSkeletalMeshes", LOCTEXT("SelectSkeletalMeshHeading", "Skeletal Meshes"));
+				FToolMenuSection& Section = Menu->AddSection("SelectSkeletalMeshes", LOCTEXT("SelectSkeletalMeshHeading", "Skeletal Meshes"));
 				Section.AddMenuEntry(FLevelEditorCommands::Get().SelectSkeletalMeshesOfSameClass);
 				Section.AddMenuEntry(FLevelEditorCommands::Get().SelectSkeletalMeshesAllClasses);
 			}
@@ -825,7 +825,7 @@ void FLevelEditorContextMenuImpl::FillSelectActorMenu(UEditorMenu* Menu)
 		if( SelectionInfo.bHaveEmitter )
 		{
 			{
-				FEditorMenuSection& Section = Menu->AddSection("SelectEmitters", LOCTEXT("SelectEmitterHeading", "Emitters"));
+				FToolMenuSection& Section = Menu->AddSection("SelectEmitters", LOCTEXT("SelectEmitterHeading", "Emitters"));
 				Section.AddMenuEntry(FLevelEditorCommands::Get().SelectMatchingEmitter);
 			}
 		}
@@ -833,7 +833,7 @@ void FLevelEditorContextMenuImpl::FillSelectActorMenu(UEditorMenu* Menu)
 
 	if( SelectionInfo.bHaveBrush || SelectionInfo.NumSelected > 0 )
 	{
-		FEditorMenuSection& Section = Menu->AddSection("SelectMaterial", LOCTEXT("SelectMaterialHeading", "Materials"));
+		FToolMenuSection& Section = Menu->AddSection("SelectMaterial", LOCTEXT("SelectMaterialHeading", "Materials"));
 		{
 			Section.AddMenuEntry(FLevelEditorCommands::Get().SelectAllWithSameMaterial);
 		}
@@ -842,7 +842,7 @@ void FLevelEditorContextMenuImpl::FillSelectActorMenu(UEditorMenu* Menu)
 	// Add geometry collection commands
 	if (FModuleManager::Get().IsModuleLoaded("GeometryCollectionEditor"))
 	{
-		FEditorMenuSection& Section = Menu->AddSection("SelectBones", LOCTEXT("GeometryCollectionHeading", "Geometry Collection"));
+		FToolMenuSection& Section = Menu->AddSection("SelectBones", LOCTEXT("GeometryCollectionHeading", "Geometry Collection"));
 		{
 			Section.AddMenuEntry(FLevelEditorCommands::Get().GeometryCollectionSelectAllGeometry);
 			Section.AddMenuEntry(FLevelEditorCommands::Get().GeometryCollectionSelectNone);
@@ -854,10 +854,10 @@ void FLevelEditorContextMenuImpl::FillSelectActorMenu(UEditorMenu* Menu)
 	FillMatineeSelectActorMenu( Menu );
 }
 
-void FLevelEditorContextMenuImpl::FillMatineeSelectActorMenu(UEditorMenu* Menu)
+void FLevelEditorContextMenuImpl::FillMatineeSelectActorMenu(UToolMenu* Menu)
 {
 	{
-		FEditorMenuSection& Section = Menu->AddSection("SelectMatinee", LOCTEXT("SelectMatineeHeading", "Matinee"));
+		FToolMenuSection& Section = Menu->AddSection("SelectMatinee", LOCTEXT("SelectMatineeHeading", "Matinee"));
 		// show list of Matinee Actors that controls this actor
 		// this is ugly but we don't have good way of knowing which Matinee actor controls me
 		// in the future this can be cached to TMap somewhere and use that list
@@ -938,10 +938,10 @@ void FLevelEditorContextMenuImpl::FillMatineeSelectActorMenu(UEditorMenu* Menu)
 	}
 }
 
-void FLevelEditorContextMenuImpl::FillActorVisibilityMenu(UEditorMenu* Menu)
+void FLevelEditorContextMenuImpl::FillActorVisibilityMenu(UToolMenu* Menu)
 {
 	{
-		FEditorMenuSection& Section = Menu->AddSection("VisibilitySelected");
+		FToolMenuSection& Section = Menu->AddSection("VisibilitySelected");
 		// Show 'Show Selected' only if the selection has any hidden actors
 		if ( SelectionInfo.bHaveHidden )
 		{
@@ -951,23 +951,23 @@ void FLevelEditorContextMenuImpl::FillActorVisibilityMenu(UEditorMenu* Menu)
 	}
 
 	{
-		FEditorMenuSection& Section = Menu->AddSection("VisibilityAll");
+		FToolMenuSection& Section = Menu->AddSection("VisibilityAll");
 		Section.AddMenuEntry(FLevelEditorCommands::Get().ShowSelectedOnly);
 		Section.AddMenuEntry(FLevelEditorCommands::Get().ShowAll);
 	}
 
 	{
-		FEditorMenuSection& Section = Menu->AddSection("VisibilityStartup");
+		FToolMenuSection& Section = Menu->AddSection("VisibilityStartup");
 		Section.AddMenuEntry(FLevelEditorCommands::Get().ShowAllStartup);
 		Section.AddMenuEntry(FLevelEditorCommands::Get().ShowSelectedStartup);
 		Section.AddMenuEntry(FLevelEditorCommands::Get().HideSelectedStartup);
 	}
 }
 
-void FLevelEditorContextMenuImpl::FillActorLevelMenu(UEditorMenu* Menu)
+void FLevelEditorContextMenuImpl::FillActorLevelMenu(UToolMenu* Menu)
 {
 	{
-		FEditorMenuSection& Section = Menu->AddSection("ActorLevel", LOCTEXT("ActorLevel", "Actor Level"));
+		FToolMenuSection& Section = Menu->AddSection("ActorLevel", LOCTEXT("ActorLevel", "Actor Level"));
 		if( SelectionInfo.SharedLevel && SelectionInfo.SharedWorld && SelectionInfo.SharedWorld->GetCurrentLevel() != SelectionInfo.SharedLevel )
 		{
 			// All actors are in the same level and that level is not the current level 
@@ -987,12 +987,12 @@ void FLevelEditorContextMenuImpl::FillActorLevelMenu(UEditorMenu* Menu)
 	}
 
 	{
-		FEditorMenuSection& Section = Menu->AddSection("LevelBlueprint", LOCTEXT("LevelBlueprint", "Level Blueprint"));
+		FToolMenuSection& Section = Menu->AddSection("LevelBlueprint", LOCTEXT("LevelBlueprint", "Level Blueprint"));
 		Section.AddMenuEntry(FLevelEditorCommands::Get().FindActorInLevelScript);
 	}
 
 	{
-		FEditorMenuSection& Section = Menu->AddSection("LevelBrowser", LOCTEXT("LevelBrowser", "Level Browser"));
+		FToolMenuSection& Section = Menu->AddSection("LevelBrowser", LOCTEXT("LevelBrowser", "Level Browser"));
 		Section.AddMenuEntry(FLevelEditorCommands::Get().FindLevelsInLevelBrowser);
 		Section.AddMenuEntry(FLevelEditorCommands::Get().AddLevelsToSelection);
 		Section.AddMenuEntry(FLevelEditorCommands::Get().RemoveLevelsFromSelection);
@@ -1000,28 +1000,28 @@ void FLevelEditorContextMenuImpl::FillActorLevelMenu(UEditorMenu* Menu)
 }
 
 
-void FLevelEditorContextMenuImpl::FillTransformMenu(UEditorMenu* Menu)
+void FLevelEditorContextMenuImpl::FillTransformMenu(UToolMenu* Menu)
 {
 	if ( FLevelEditorActionCallbacks::ActorSelected_CanExecute() )
 	{
 		{
-			FEditorMenuSection& Section = Menu->AddSection("TransformSnapAlign");
-			Section.AddEntry(FEditorMenuEntry::InitSubMenu(
+			FToolMenuSection& Section = Menu->AddSection("TransformSnapAlign");
+			Section.AddEntry(FToolMenuEntry::InitSubMenu(
 				Menu->GetMenuName(),
 				"SnapAlignSubMenu",
 				LOCTEXT("SnapAlignSubMenu", "Snap/Align"), 
 				LOCTEXT("SnapAlignSubMenu_ToolTip", "Actor snap/align utils"),
-				FNewEditorMenuDelegate::CreateStatic(&FLevelEditorContextMenuImpl::FillSnapAlignMenu)));
+				FNewToolMenuDelegate::CreateStatic(&FLevelEditorContextMenuImpl::FillSnapAlignMenu)));
 		}
 
 		{
-			FEditorMenuSection& Section = Menu->AddSection("DeltaTransformToActors");
+			FToolMenuSection& Section = Menu->AddSection("DeltaTransformToActors");
 			Section.AddMenuEntry(FLevelEditorCommands::Get().DeltaTransformToActors);
 		}
 	}
 
 	{
-		FEditorMenuSection& Section = Menu->AddSection("MirrorLock");
+		FToolMenuSection& Section = Menu->AddSection("MirrorLock");
 		Section.AddMenuEntry(FLevelEditorCommands::Get().MirrorActorX);
 		Section.AddMenuEntry(FLevelEditorCommands::Get().MirrorActorY);
 		Section.AddMenuEntry(FLevelEditorCommands::Get().MirrorActorZ);
@@ -1029,7 +1029,7 @@ void FLevelEditorContextMenuImpl::FillTransformMenu(UEditorMenu* Menu)
 	}
 }
 
-void FLevelEditorContextMenuImpl::FillActorMenu(UEditorMenu* Menu)
+void FLevelEditorContextMenuImpl::FillActorMenu(UToolMenu* Menu)
 {
 	struct Local
 	{
@@ -1051,7 +1051,7 @@ void FLevelEditorContextMenuImpl::FillActorMenu(UEditorMenu* Menu)
 		InitOptions.Filters->AddFilterPredicate( SceneOutliner::FActorFilterPredicate::CreateStatic( &FLevelEditorActionCallbacks::IsAttachableActor) );
 	}		
 
-	FEditorMenuSection& Section = Menu->AddSection("Actor");
+	FToolMenuSection& Section = Menu->AddSection("Actor");
 	if(SelectionInfo.bHaveAttachedActor)
 	{
 		Section.AddMenuEntry(FLevelEditorCommands::Get().DetachFromParent, LOCTEXT("None", "None"));
@@ -1102,12 +1102,12 @@ void FLevelEditorContextMenuImpl::FillActorMenu(UEditorMenu* Menu)
 			]
 		];
 
-	Section.AddEntry(FEditorMenuEntry::InitWidget("PickParentActor", MenuWidget, FText::GetEmpty(), false));
+	Section.AddEntry(FToolMenuEntry::InitWidget("PickParentActor", MenuWidget, FText::GetEmpty(), false));
 }
 
-void FLevelEditorContextMenuImpl::FillSnapAlignMenu(UEditorMenu* Menu)
+void FLevelEditorContextMenuImpl::FillSnapAlignMenu(UToolMenu* Menu)
 {
-	FEditorMenuSection& Section = Menu->AddSection("SnapAlign");
+	FToolMenuSection& Section = Menu->AddSection("SnapAlign");
 	Section.AddMenuEntry( FLevelEditorCommands::Get().SnapOriginToGrid );
 	Section.AddMenuEntry( FLevelEditorCommands::Get().SnapOriginToGridPerActor );
 	Section.AddMenuEntry( FLevelEditorCommands::Get().AlignOriginToGrid );
@@ -1166,10 +1166,10 @@ void FLevelEditorContextMenuImpl::FillSnapAlignMenu(UEditorMenu* Menu)
 */
 }
 
-void FLevelEditorContextMenuImpl::FillPivotMenu( UEditorMenu* Menu )
+void FLevelEditorContextMenuImpl::FillPivotMenu( UToolMenu* Menu )
 {
 	{
-		FEditorMenuSection& Section = Menu->AddSection("SaveResetPivot");
+		FToolMenuSection& Section = Menu->AddSection("SaveResetPivot");
 		Section.AddMenuEntry(FLevelEditorCommands::Get().SavePivotToPrePivot);
 		Section.AddMenuEntry(FLevelEditorCommands::Get().ResetPrePivot);
 		Section.AddMenuEntry(FLevelEditorCommands::Get().MovePivotHere);
@@ -1177,14 +1177,14 @@ void FLevelEditorContextMenuImpl::FillPivotMenu( UEditorMenu* Menu )
 	}
 
 	{
-		FEditorMenuSection& Section = Menu->AddSection("MovePivot");
+		FToolMenuSection& Section = Menu->AddSection("MovePivot");
 		Section.AddMenuEntry(FLevelEditorCommands::Get().MovePivotToCenter);
 	}
 }
 
-void FLevelEditorContextMenuImpl::FillGroupMenu( UEditorMenu* Menu )
+void FLevelEditorContextMenuImpl::FillGroupMenu( UToolMenu* Menu )
 {
-	FEditorMenuSection& Section = Menu->AddSection("Group");
+	FToolMenuSection& Section = Menu->AddSection("Group");
 
 	if( SelectionInfo.NumSelectedUngroupedActors > 1 )
 	{
@@ -1228,9 +1228,9 @@ void FLevelEditorContextMenuImpl::FillGroupMenu( UEditorMenu* Menu )
 	}
 }
 
-void FLevelEditorContextMenuImpl::FillEditMenu( UEditorMenu* Menu )
+void FLevelEditorContextMenuImpl::FillEditMenu( UToolMenu* Menu )
 {
-	FEditorMenuSection& Section = Menu->AddSection("Section");
+	FToolMenuSection& Section = Menu->AddSection("Section");
 
 	Section.AddMenuEntry( FGenericCommands::Get().Cut );
 	Section.AddMenuEntry( FGenericCommands::Get().Copy );
@@ -1248,7 +1248,7 @@ void FLevelEditorContextMenuImpl::FillEditMenu( UEditorMenu* Menu )
 	Section.AddMenuEntry( FGenericCommands::Get().Rename );
 }
 
-void FLevelScriptEventMenuHelper::FillLevelBlueprintEventsMenu(UEditorMenu* Menu, const TArray<AActor*>& SelectedActors)
+void FLevelScriptEventMenuHelper::FillLevelBlueprintEventsMenu(UToolMenu* Menu, const TArray<AActor*>& SelectedActors)
 {
 	AActor* SelectedActor = (1 == SelectedActors.Num()) ? SelectedActors[0] : NULL;
 
@@ -1261,16 +1261,16 @@ void FLevelScriptEventMenuHelper::FillLevelBlueprintEventsMenu(UEditorMenu* Menu
 		{
 			TWeakObjectPtr<AActor> ActorPtr(SelectedActor);
 
-			FEditorMenuSection& Section = Menu->AddSection("LevelBlueprintEvents", LOCTEXT("LevelBlueprintEvents", "Level Blueprint Events"));
+			FToolMenuSection& Section = Menu->AddSection("LevelBlueprintEvents", LOCTEXT("LevelBlueprintEvents", "Level Blueprint Events"));
 
 			if (bAnyEventExists)
 			{
-				Section.AddEntry(FEditorMenuEntry::InitSubMenu(
+				Section.AddEntry(FToolMenuEntry::InitSubMenu(
 					Menu->GetMenuName(),
 					"JumpEventSubMenu",
 					LOCTEXT("JumpEventSubMenu", "Jump to Event"),
 					FText::GetEmpty(),
-					FNewEditorMenuDelegate::CreateStatic(&FKismetEditorUtilities::AddLevelScriptEventOptionsForActor
+					FNewToolMenuDelegate::CreateStatic(&FKismetEditorUtilities::AddLevelScriptEventOptionsForActor
 					, ActorPtr
 					, true
 					, false
@@ -1279,12 +1279,12 @@ void FLevelScriptEventMenuHelper::FillLevelBlueprintEventsMenu(UEditorMenu* Menu
 
 			if (bAnyEventCanBeAdded)
 			{
-				Section.AddEntry(FEditorMenuEntry::InitSubMenu(
+				Section.AddEntry(FToolMenuEntry::InitSubMenu(
 					Menu->GetMenuName(),
 					"AddEventSubMenu",
 					LOCTEXT("AddEventSubMenu", "Add Event"),
 					FText::GetEmpty(),
-					FNewEditorMenuDelegate::CreateStatic(&FKismetEditorUtilities::AddLevelScriptEventOptionsForActor
+					FNewToolMenuDelegate::CreateStatic(&FKismetEditorUtilities::AddLevelScriptEventOptionsForActor
 					, ActorPtr
 					, false
 					, true
diff --git a/Engine/Source/Editor/LevelEditor/Private/LevelEditorContextMenu.h b/Engine/Source/Editor/LevelEditor/Private/LevelEditorContextMenu.h
--- a/Engine/Source/Editor/LevelEditor/Private/LevelEditorContextMenu.h
+++ b/Engine/Source/Editor/LevelEditor/Private/LevelEditorContextMenu.h
@@ -7,8 +7,8 @@
 #include "LevelEditorMenuContext.h"
 
 class FExtender;
-class UEditorMenu;
-struct FEditorMenuContext;
+class UToolMenu;
+struct FToolMenuContext;
 class SLevelEditor;
 class SWidget;
 
@@ -51,10 +51,10 @@ public:
 	 * @param	ContextType		The context we should use to specialize this menu
 	 * @param	Extender		Allows extension of this menu based on context.
 	 */
-	static UEditorMenu* GenerateMenu(TWeakPtr< SLevelEditor > LevelEditor, ELevelEditorMenuContext ContextType, TSharedPtr<FExtender> Extender = TSharedPtr<FExtender>());
+	static UToolMenu* GenerateMenu(TWeakPtr< SLevelEditor > LevelEditor, ELevelEditorMenuContext ContextType, TSharedPtr<FExtender> Extender = TSharedPtr<FExtender>());
 
 	/* Adds required information to Context for build menu based on current selection */
-	static FName InitMenuContext(FEditorMenuContext& Context, TWeakPtr<SLevelEditor> LevelEditor, ELevelEditorMenuContext ContextType);
+	static FName InitMenuContext(FToolMenuContext& Context, TWeakPtr<SLevelEditor> LevelEditor, ELevelEditorMenuContext ContextType);
 
 	/* Returns name of menu to display based on current selection */
 	static FName GetContextMenuName(ELevelEditorMenuContext ContextType);
@@ -71,5 +71,5 @@ private:
 	 * @param Menu		The menu to add items to.
 	 * @param SelectedActorInfo	Information about the selected actors.
 	 */
-	static void BuildGroupMenu(UEditorMenu* Menu, const struct FSelectedActorInfo& SelectedActorInfo);
+	static void BuildGroupMenu(UToolMenu* Menu, const struct FSelectedActorInfo& SelectedActorInfo);
 };
diff --git a/Engine/Source/Editor/LevelEditor/Private/LevelEditorCreateActorMenu.cpp b/Engine/Source/Editor/LevelEditor/Private/LevelEditorCreateActorMenu.cpp
--- a/Engine/Source/Editor/LevelEditor/Private/LevelEditorCreateActorMenu.cpp
+++ b/Engine/Source/Editor/LevelEditor/Private/LevelEditorCreateActorMenu.cpp
@@ -9,7 +9,7 @@
 #include "Widgets/SBoxPanel.h"
 #include "Textures/SlateIcon.h"
 #include "Framework/Commands/UIAction.h"
-#include "EditorMenuSubsystem.h"
+#include "ToolMenus.h"
 #include "Widgets/Images/SImage.h"
 #include "Widgets/Text/STextBlock.h"
 #include "Widgets/Layout/SBox.h"
@@ -280,9 +280,9 @@ static void GetContentBrowserSelectionFactoryMenuEntries( FAssetData& TargetAsse
  * @param	AssetMenuOptions	Menu options to use
  * @param	CreateMode			The creation mode to use
  */
-static void FillAssetAddReplaceActorMenu(UEditorMenu* Menu, const FAssetData Asset, const TArray< FActorFactoryAssetProxy::FMenuItem > AssetMenuOptions, EActorCreateMode::Type CreateMode)
+static void FillAssetAddReplaceActorMenu(UToolMenu* Menu, const FAssetData Asset, const TArray< FActorFactoryAssetProxy::FMenuItem > AssetMenuOptions, EActorCreateMode::Type CreateMode)
 {
-	FEditorMenuSection& Section = Menu->AddSection("Section");
+	FToolMenuSection& Section = Menu->AddSection("Section");
 	for( int32 ItemIndex = 0; ItemIndex < AssetMenuOptions.Num(); ++ItemIndex )
 	{
 		const FActorFactoryAssetProxy::FMenuItem& MenuItem = AssetMenuOptions[ItemIndex];
@@ -319,7 +319,7 @@ static void FillAssetAddReplaceActorMenu(UEditorMenu* Menu, const FAssetData Ass
  * @param	CreateMode			The creation mode to use
  * @param	LabelOverride		The lable to use, if any.
  */
-static void BuildSingleAssetAddReplaceActorMenu(FEditorMenuSection& Section, const FAssetData& Asset, const TArray< FActorFactoryAssetProxy::FMenuItem >& AssetMenuOptions, EActorCreateMode::Type CreateMode, const FText& LabelOverride = FText::GetEmpty())
+static void BuildSingleAssetAddReplaceActorMenu(FToolMenuSection& Section, const FAssetData& Asset, const TArray< FActorFactoryAssetProxy::FMenuItem >& AssetMenuOptions, EActorCreateMode::Type CreateMode, const FText& LabelOverride = FText::GetEmpty())
 {
 	if ( !Asset.IsValid() || AssetMenuOptions.Num() == 0 )
 	{
@@ -340,18 +340,18 @@ static void BuildSingleAssetAddReplaceActorMenu(FEditorMenuSection& Section, con
 			Action = FUIAction( FExecuteAction::CreateStatic( &FLevelEditorActionCallbacks::AddActor_Clicked, MenuItem.FactoryToUse,  MenuItem.AssetData, CreateMode == EActorCreateMode::Placement ) );
 		}
 
-		Section.AddEntry(FEditorMenuEntry::InitMenuEntry(NAME_None, Action, SNew(SAssetMenuEntry, Asset, AssetMenuOptions).LabelOverride(LabelOverride)));
+		Section.AddEntry(FToolMenuEntry::InitMenuEntry(NAME_None, Action, SNew(SAssetMenuEntry, Asset, AssetMenuOptions).LabelOverride(LabelOverride)));
 	}
 	else
 	{
-		FEditorMenuEntry Entry = FEditorMenuEntry::InitWidget(NAME_None, SNew( SAssetMenuEntry, Asset, AssetMenuOptions ).LabelOverride( LabelOverride ), FText());
+		FToolMenuEntry Entry = FToolMenuEntry::InitWidget(NAME_None, SNew( SAssetMenuEntry, Asset, AssetMenuOptions ).LabelOverride( LabelOverride ), FText());
 		Entry.SubMenuData.bIsSubMenu = true;
-		Entry.SubMenuData.ConstructMenu = FNewEditorMenuDelegate::CreateStatic(&FillAssetAddReplaceActorMenu, Asset, AssetMenuOptions, CreateMode);
+		Entry.SubMenuData.ConstructMenu = FNewToolMenuDelegate::CreateStatic(&FillAssetAddReplaceActorMenu, Asset, AssetMenuOptions, CreateMode);
 		Section.AddEntry(Entry);
 	}
 }
 
-void LevelEditorCreateActorMenu::FillAddReplaceViewportContextMenuSections(UEditorMenu* Menu)
+void LevelEditorCreateActorMenu::FillAddReplaceViewportContextMenuSections(UToolMenu* Menu)
 {
 	FAssetData TargetAssetData;
 	TArray< FActorFactoryAssetProxy::FMenuItem > AssetMenuOptions;
@@ -360,22 +360,22 @@ void LevelEditorCreateActorMenu::FillAddReplaceViewportContextMenuSections(UEdit
 	if ( AssetMenuOptions.Num() == 0 )
 	{
 		{
-			FEditorMenuSection& Section = Menu->AddSection("ActorType");
-			Section.AddEntry(FEditorMenuEntry::InitSubMenu(
+			FToolMenuSection& Section = Menu->AddSection("ActorType");
+			Section.AddEntry(FToolMenuEntry::InitSubMenu(
 				Menu->GetMenuName(),
 				"LevelViewportContextMenu",
 				NSLOCTEXT("LevelViewportContextMenu", "AddActorHeading", "Place Actor") , 
 				NSLOCTEXT("LevelViewportContextMenu", "AddActorMenu_ToolTip", "Templates for adding a new actor to the world"),
-				FNewEditorMenuDelegate::CreateStatic(&LevelEditorCreateActorMenu::FillAddReplaceActorMenu, EActorCreateMode::Add)));
+				FNewToolMenuDelegate::CreateStatic(&LevelEditorCreateActorMenu::FillAddReplaceActorMenu, EActorCreateMode::Add)));
 
 			if ( CanReplaceActors() )
 			{
-				Section.AddEntry(FEditorMenuEntry::InitSubMenu(
+				Section.AddEntry(FToolMenuEntry::InitSubMenu(
 					Menu->GetMenuName(),
 					"LevelViewportContextMenu",
 					NSLOCTEXT("LevelViewportContextMenu", "ReplaceActorHeading", "Replace Selected Actors with") , 
 					NSLOCTEXT("LevelViewportContextMenu", "ReplaceActorMenu_ToolTip", "Templates for replacing selected with new actors in the world"),
-					FNewEditorMenuDelegate::CreateStatic(&LevelEditorCreateActorMenu::FillAddReplaceActorMenu, EActorCreateMode::Replace)));
+					FNewToolMenuDelegate::CreateStatic(&LevelEditorCreateActorMenu::FillAddReplaceActorMenu, EActorCreateMode::Replace)));
 			}
 		}
 	}
@@ -387,40 +387,40 @@ void LevelEditorCreateActorMenu::FillAddReplaceViewportContextMenuSections(UEdit
 		}
 
 		{
-			FEditorMenuSection& Section = Menu->AddSection("AddActor", NSLOCTEXT("LevelViewportContextMenu", "AddActorHeading", "Place Actor"));
+			FToolMenuSection& Section = Menu->AddSection("AddActor", NSLOCTEXT("LevelViewportContextMenu", "AddActorHeading", "Place Actor"));
 			FUIAction Action( FExecuteAction::CreateStatic( &FLevelEditorActionCallbacks::AddActor_Clicked, AssetMenuOptions[0].FactoryToUse,  AssetMenuOptions[0].AssetData, false ) );
 			TSharedRef<SWidget> Widget = SNew(SAssetMenuEntry, TargetAssetData, AssetMenuOptions);
-			Section.AddEntry(FEditorMenuEntry::InitSubMenu(
+			Section.AddEntry(FToolMenuEntry::InitSubMenu(
 				Menu->GetMenuName(),
 				"AddActor",
 				Action,
 				Widget,
-				FNewEditorMenuDelegate::CreateStatic(&LevelEditorCreateActorMenu::FillAddReplaceActorMenu, EActorCreateMode::Add)
+				FNewToolMenuDelegate::CreateStatic(&LevelEditorCreateActorMenu::FillAddReplaceActorMenu, EActorCreateMode::Add)
 			));
 		}
 
 		if ( CanReplaceActors() )
 		{
 			{
-				FEditorMenuSection& Section = Menu->AddSection("ReplaceActor", NSLOCTEXT("LevelViewportContextMenu", "ReplaceActorHeading", "Replace Selected Actors with"));
+				FToolMenuSection& Section = Menu->AddSection("ReplaceActor", NSLOCTEXT("LevelViewportContextMenu", "ReplaceActorHeading", "Replace Selected Actors with"));
 				FUIAction Action( FExecuteAction::CreateStatic( &FLevelEditorActionCallbacks::ReplaceActors_Clicked, AssetMenuOptions[0].FactoryToUse,  AssetMenuOptions[0].AssetData ) );
 				TSharedRef<SWidget> Widget = SNew(SAssetMenuEntry, TargetAssetData, AssetMenuOptions);
-				Section.AddEntry(FEditorMenuEntry::InitSubMenu(
+				Section.AddEntry(FToolMenuEntry::InitSubMenu(
 					Menu->GetMenuName(),
 					"ReplaceActor",
 					Action,
 					Widget,
-					FNewEditorMenuDelegate::CreateStatic(&LevelEditorCreateActorMenu::FillAddReplaceActorMenu, EActorCreateMode::Replace)
+					FNewToolMenuDelegate::CreateStatic(&LevelEditorCreateActorMenu::FillAddReplaceActorMenu, EActorCreateMode::Replace)
 				));
 			}
 		}
 	}
 }
 
-void LevelEditorCreateActorMenu::FillAddReplaceActorMenu(UEditorMenu* Menu, EActorCreateMode::Type CreateMode)
+void LevelEditorCreateActorMenu::FillAddReplaceActorMenu(UToolMenu* Menu, EActorCreateMode::Type CreateMode)
 {
 	{
-		FEditorMenuSection& Section = Menu->AddSection("ContentBrowserActor", NSLOCTEXT("LevelViewportContextMenu", "AssetSelectionSection", "Selection"));
+		FToolMenuSection& Section = Menu->AddSection("ContentBrowserActor", NSLOCTEXT("LevelViewportContextMenu", "AssetSelectionSection", "Selection"));
 		FAssetData TargetAssetData;
 		TArray< FActorFactoryAssetProxy::FMenuItem > AssetMenuOptions;
 		GetContentBrowserSelectionFactoryMenuEntries( /*OUT*/TargetAssetData, /*OUT*/AssetMenuOptions );
@@ -429,7 +429,7 @@ void LevelEditorCreateActorMenu::FillAddReplaceActorMenu(UEditorMenu* Menu, EAct
 	}
 
 	{
-		FEditorMenuSection& Section = Menu->AddSection("RecentlyPlaced", NSLOCTEXT("LevelViewportContextMenu", "RecentlyPlacedSection", "Recently Placed"));
+		FToolMenuSection& Section = Menu->AddSection("RecentlyPlaced", NSLOCTEXT("LevelViewportContextMenu", "RecentlyPlacedSection", "Recently Placed"));
 		if ( IPlacementModeModule::IsAvailable() )
 		{
 			FAssetRegistryModule& AssetRegistryModule = FModuleManager::LoadModuleChecked<FAssetRegistryModule>(TEXT("AssetRegistry"));
@@ -464,7 +464,7 @@ void LevelEditorCreateActorMenu::FillAddReplaceActorMenu(UEditorMenu* Menu, EAct
 	}
 
 	{
-		FEditorMenuSection& Section = Menu->AddSection("Lights", NSLOCTEXT("LevelViewportContextMenu", "LightsSection", "Lights"));
+		FToolMenuSection& Section = Menu->AddSection("Lights", NSLOCTEXT("LevelViewportContextMenu", "LightsSection", "Lights"));
 		TArray< FActorFactoryAssetProxy::FMenuItem > AssetMenuOptions;
 
 		{
@@ -501,7 +501,7 @@ void LevelEditorCreateActorMenu::FillAddReplaceActorMenu(UEditorMenu* Menu, EAct
 	}
 
 	{
-		FEditorMenuSection& Section = Menu->AddSection("Primitives", NSLOCTEXT("LevelViewportContextMenu", "PrimitivesSection", "Primitives"));
+		FToolMenuSection& Section = Menu->AddSection("Primitives", NSLOCTEXT("LevelViewportContextMenu", "PrimitivesSection", "Primitives"));
 		TArray< FActorFactoryAssetProxy::FMenuItem > AssetMenuOptions;
 		{
 			AssetMenuOptions.Empty();
@@ -554,7 +554,7 @@ void LevelEditorCreateActorMenu::FillAddReplaceActorMenu(UEditorMenu* Menu, EAct
 	}
 
 	{
-		FEditorMenuSection& Section = Menu->AddSection("Custom", NSLOCTEXT("LevelViewportContextMenu", "CustomSection", "Custom Actors"));
+		FToolMenuSection& Section = Menu->AddSection("Custom", NSLOCTEXT("LevelViewportContextMenu", "CustomSection", "Custom Actors"));
 		TArray< FActorFactoryAssetProxy::FMenuItem > AssetMenuOptions;
 		FText UnusedErrorMessage;
 		const FAssetData NoAssetData;
diff --git a/Engine/Source/Editor/LevelEditor/Private/LevelEditorCreateActorMenu.h b/Engine/Source/Editor/LevelEditor/Private/LevelEditorCreateActorMenu.h
--- a/Engine/Source/Editor/LevelEditor/Private/LevelEditorCreateActorMenu.h
+++ b/Engine/Source/Editor/LevelEditor/Private/LevelEditorCreateActorMenu.h
@@ -4,7 +4,7 @@
 
 #include "CoreMinimal.h"
 
-class UEditorMenu;
+class UToolMenu;
 
 /**
  * The mode to use when creating an actor
@@ -30,11 +30,11 @@ namespace LevelEditorCreateActorMenu
 	 * Fill the context menu section(s) for adding or replacing an actor in the viewport
 	 * @param	MenuBuilder		The menu builder used to generate the context menu
 	 */
-	void FillAddReplaceViewportContextMenuSections(UEditorMenu* Menu);
+	void FillAddReplaceViewportContextMenuSections(UToolMenu* Menu);
 
 	/**
 	 * Fill the context menu for adding or replacing an actor. Used for in-viewport and level editor toolbar menus.
 	 * @param	MenuBuilder		The menu builder used to generate the context menu
 	 */
-	void FillAddReplaceActorMenu(UEditorMenu* Menu, EActorCreateMode::Type CreateMode);
+	void FillAddReplaceActorMenu(UToolMenu* Menu, EActorCreateMode::Type CreateMode);
 };
