Matinee: Remove from UI

Removed access to matinee from the following:
- Cinematics->Add Matinee [Legacy]
- List of matinee actors in the cinematics menu
- Open Matinee button in the details panel

diff --git a/Engine/Plugins/MovieScene/LevelSequenceEditor/Source/LevelSequenceEditor/Private/LevelSequenceEditorModule.cpp b/Engine/Plugins/MovieScene/LevelSequenceEditor/Source/LevelSequenceEditor/Private/LevelSequenceEditorModule.cpp
--- a/Engine/Plugins/MovieScene/LevelSequenceEditor/Source/LevelSequenceEditor/Private/LevelSequenceEditorModule.cpp
+++ b/Engine/Plugins/MovieScene/LevelSequenceEditor/Source/LevelSequenceEditor/Private/LevelSequenceEditorModule.cpp
@@ -159,10 +159,10 @@ protected:
 
 		// Create and register the level editor toolbar menu extension
 		CinematicsMenuExtender = MakeShareable(new FExtender);
-		CinematicsMenuExtender->AddMenuExtension("LevelEditorNewMatinee", EExtensionHook::First, CommandList, FMenuExtensionDelegate::CreateStatic([](FMenuBuilder& MenuBuilder) {
+		CinematicsMenuExtender->AddMenuExtension("LevelEditorNewCinematics", EExtensionHook::First, CommandList, FMenuExtensionDelegate::CreateStatic([](FMenuBuilder& MenuBuilder) {
 			MenuBuilder.AddMenuEntry(FLevelSequenceEditorCommands::Get().CreateNewLevelSequenceInLevel);
 		}));
-		CinematicsMenuExtender->AddMenuExtension("LevelEditorNewMatinee", EExtensionHook::First, CommandList, FMenuExtensionDelegate::CreateStatic([](FMenuBuilder& MenuBuilder) {
+		CinematicsMenuExtender->AddMenuExtension("LevelEditorNewCinematics", EExtensionHook::First, CommandList, FMenuExtensionDelegate::CreateStatic([](FMenuBuilder& MenuBuilder) {
 			MenuBuilder.AddMenuEntry(FLevelSequenceEditorCommands::Get().CreateNewMasterSequenceInLevel);
 		}));
 
diff --git a/Engine/Source/Editor/DetailCustomizations/Private/DetailCustomizations.cpp b/Engine/Source/Editor/DetailCustomizations/Private/DetailCustomizations.cpp
--- a/Engine/Source/Editor/DetailCustomizations/Private/DetailCustomizations.cpp
+++ b/Engine/Source/Editor/DetailCustomizations/Private/DetailCustomizations.cpp
@@ -16,7 +16,6 @@
 #include "SkeletalMeshComponentDetails.h"
 #include "SplineComponentDetails.h"
 #include "MeshComponentDetails.h"
-#include "MatineeActorDetails.h"
 #include "LevelSequenceActorDetails.h"
 #include "ReflectionCaptureDetails.h"
 #include "SkyLightComponentDetails.h"
@@ -297,7 +296,6 @@ void FDetailCustomizationsModule::RegisterObjectCustomizations()
 	RegisterCustomClassLayout("DirectionalLightComponent", FOnGetDetailCustomizationInstance::CreateStatic(&FDirectionalLightComponentDetails::MakeInstance));
 	RegisterCustomClassLayout("StaticMeshActor", FOnGetDetailCustomizationInstance::CreateStatic(&FStaticMeshActorDetails::MakeInstance));
 	RegisterCustomClassLayout("MeshComponent", FOnGetDetailCustomizationInstance::CreateStatic(&FMeshComponentDetails::MakeInstance));
-	RegisterCustomClassLayout("MatineeActor", FOnGetDetailCustomizationInstance::CreateStatic(&FMatineeActorDetails::MakeInstance));
 	RegisterCustomClassLayout("LevelSequenceActor", FOnGetDetailCustomizationInstance::CreateStatic(&FLevelSequenceActorDetails::MakeInstance));
 	RegisterCustomClassLayout("ReflectionCapture", FOnGetDetailCustomizationInstance::CreateStatic(&FReflectionCaptureDetails::MakeInstance));
 	RegisterCustomClassLayout("SceneCaptureComponent", FOnGetDetailCustomizationInstance::CreateStatic(&FSceneCaptureDetails::MakeInstance));
diff --git a/Engine/Source/Editor/DetailCustomizations/Private/MatineeActorDetails.cpp b/Engine/Source/Editor/DetailCustomizations/Private/MatineeActorDetails.cpp
deleted file mode 100644
diff --git a/Engine/Source/Editor/DetailCustomizations/Private/MatineeActorDetails.h b/Engine/Source/Editor/DetailCustomizations/Private/MatineeActorDetails.h
deleted file mode 100644
diff --git a/Engine/Source/Editor/LevelEditor/Private/LevelEditor.cpp b/Engine/Source/Editor/LevelEditor/Private/LevelEditor.cpp
--- a/Engine/Source/Editor/LevelEditor/Private/LevelEditor.cpp
+++ b/Engine/Source/Editor/LevelEditor/Private/LevelEditor.cpp
@@ -1741,10 +1741,6 @@ void FLevelEditorModule::BindGlobalLevelEditorCommands()
 		FCanExecuteAction(),
 		FIsActionChecked::CreateStatic( &FLevelEditorActionCallbacks::IsViewportUIHidden ) 
 		);
-	ActionList.MapAction(
-		Commands.AddMatinee,
-		FExecuteAction::CreateStatic( &FLevelEditorActionCallbacks::OnAddMatinee )
-		);
 
 	ActionList.MapAction( 
 		Commands.MaterialQualityLevel_Low, 
diff --git a/Engine/Source/Editor/LevelEditor/Private/LevelEditorActions.cpp b/Engine/Source/Editor/LevelEditor/Private/LevelEditorActions.cpp
--- a/Engine/Source/Editor/LevelEditor/Private/LevelEditorActions.cpp
+++ b/Engine/Source/Editor/LevelEditor/Private/LevelEditorActions.cpp
@@ -2707,27 +2707,6 @@ void FLevelEditorActionCallbacks::OnAddVolume( UClass* VolumeClass )
 	GEditor->RedrawAllViewports();
 }
 
-void FLevelEditorActionCallbacks::OnAddMatinee()
-{
-	// Warn the user prior to creating our actor
-	if ( GEditor->ShouldOpenMatinee( NULL ) )
-	{
-		// Spawn a matinee actor at the origin, and either move infront of the camera or focus camera on it (depending on the viewport) and open for edit
-		UActorFactory* ActorFactory = GEditor->FindActorFactoryForActorClass( AMatineeActor::StaticClass() );
-		check( ActorFactory );
-		AMatineeActor* MatineeActor = CastChecked<AMatineeActor>( FLevelEditorActionCallbacks::AddActor( ActorFactory, FAssetData(), &FTransform::Identity ) );
-		if( GCurrentLevelEditingViewportClient->IsPerspective() )
-		{
-			GEditor->MoveActorInFrontOfCamera( *MatineeActor, GCurrentLevelEditingViewportClient->GetViewLocation(), GCurrentLevelEditingViewportClient->GetViewRotation().Vector() );
-		}
-		else
-		{
-			GEditor->MoveViewportCamerasToActor( *MatineeActor, false );
-		}
-		GEditor->OpenMatinee( MatineeActor, false );
-	}
-}
-
 void FLevelEditorActionCallbacks::SelectActorsInLayers()
 {
 	// Iterate over selected actors and make a list of all layers the selected actors belong to.
@@ -3357,8 +3336,6 @@ void FLevelEditorCommands::RegisterCommands()
 	UI_COMMAND( WorldProperties, "World Settings", "Displays the world settings", EUserInterfaceActionType::Button, FInputChord() );
 	UI_COMMAND( OpenContentBrowser, "Open Content Browser", "Opens the Content Browser", EUserInterfaceActionType::Button, FInputChord(EModifierKey::Control|EModifierKey::Shift, EKeys::F) );
 	UI_COMMAND( OpenMarketplace, "Open Marketplace", "Opens the Marketplace", EUserInterfaceActionType::Button, FInputChord() );
-	UI_COMMAND( AddMatinee, "Add Matinee [Legacy]", "Creates a new matinee actor to edit", EUserInterfaceActionType::Button, FInputChord() );
-	UI_COMMAND( EditMatinee, "Edit Matinee", "Selects a Matinee to edit", EUserInterfaceActionType::Button, FInputChord() );
 
 	UI_COMMAND( ToggleVR, "Toggle VR", "Toggles VR (Virtual Reality) mode", EUserInterfaceActionType::ToggleButton, FInputChord( EModifierKey::Alt, EKeys::V ) );
 
diff --git a/Engine/Source/Editor/LevelEditor/Private/LevelEditorToolBar.cpp b/Engine/Source/Editor/LevelEditor/Private/LevelEditorToolBar.cpp
--- a/Engine/Source/Editor/LevelEditor/Private/LevelEditorToolBar.cpp
+++ b/Engine/Source/Editor/LevelEditor/Private/LevelEditorToolBar.cpp
@@ -36,7 +36,6 @@
 #include "SScalabilitySettings.h"
 #include "IContentBrowserSingleton.h"
 #include "ContentBrowserModule.h"
-#include "Matinee/MatineeActor.h"
 #include "LevelSequenceActor.h"
 #include "Engine/LevelScriptBlueprint.h"
 #include "ISettingsCategory.h"
@@ -1430,7 +1429,7 @@ void FLevelEditorToolBar::RegisterLevelEditorToolBar( const TSharedRef<FUIComman
 			FUIAction(),
 			FOnGetContent::CreateStatic( &FLevelEditorToolBar::GenerateCinematicsMenuContent, InCommandList, TWeakPtr<SLevelEditor>( InLevelEditor ) ),
 			LOCTEXT( "EditCinematics_Label", "Cinematics" ),
-			LOCTEXT( "EditCinematics_Tooltip", "Displays a list of Matinee and Level Sequence objects to open in their respective editors"),
+			LOCTEXT( "EditCinematics_Tooltip", "Displays a list of Level Sequence objects to open in their respective editors"),
 			FSlateIcon(FEditorStyle::GetStyleSetName(), "LevelEditor.EditMatinee") 
 			));
 
@@ -2295,10 +2294,7 @@ void FLevelEditorToolBar::RegisterCinematicsMenu()
 	UToolMenu* Menu = UToolMenus::Get()->RegisterMenu("LevelEditor.LevelEditorToolBar.Cinematics");
 	Menu->bShouldCloseWindowAfterMenuSelection = true;
 
-	{
-		FToolMenuSection& Section = Menu->AddSection("LevelEditorNewMatinee", LOCTEXT("MatineeMenuCombo_NewHeading", "New"));
-		Section.AddMenuEntry( FLevelEditorCommands::Get().AddMatinee, TAttribute<FText>(), TAttribute<FText>(), FSlateIcon(FEditorStyle::GetStyleSetName(), "ClassIcon.MatineeActor") );
-	}
+	Menu->AddSection("LevelEditorNewCinematics", LOCTEXT("CinematicsMenuCombo_NewHeading", "New"));
 
 	//Add a heading to separate the existing cinematics from the 'Add New Cinematic Actor' button
 	FToolMenuSection& ExistingCinematicSection = Menu->AddSection("LevelEditorExistingCinematic", LOCTEXT("CinematicMenuCombo_ExistingHeading", "Edit Existing Cinematic"));
@@ -2311,7 +2307,7 @@ void FLevelEditorToolBar::RegisterCinematicsMenu()
 		}
 
 		UWorld* World = FoundContext->LevelEditor.IsValid() ? FoundContext->LevelEditor.Pin()->GetWorld() : nullptr;
-		const bool bHasAnyCinematicsActors = !!TActorIterator<AMatineeActor>(World) || !!TActorIterator<ALevelSequenceActor>(World);
+		const bool bHasAnyCinematicsActors = !!TActorIterator<ALevelSequenceActor>(World);
 		if (!bHasAnyCinematicsActors)
 		{
 			return;
@@ -2319,7 +2315,7 @@ void FLevelEditorToolBar::RegisterCinematicsMenu()
 
 		using namespace SceneOutliner;
 
-		// We can't build a list of Matinees and LevelSequenceActors while the current World is a PIE world.
+		// We can't build a list of LevelSequenceActors while the current World is a PIE world.
 		FInitializationOptions InitOptions;
 		{
 			InitOptions.Mode = ESceneOutlinerMode::ActorPicker;
@@ -2333,9 +2329,9 @@ void FLevelEditorToolBar::RegisterCinematicsMenu()
 			InitOptions.ColumnMap.Add(FBuiltInColumnTypes::Label(), FColumnInfo(EColumnVisibility::Visible, 0));
 			InitOptions.ColumnMap.Add(FBuiltInColumnTypes::ActorInfo(), FColumnInfo(EColumnVisibility::Visible, 10));
 
-			// Only display Matinee and MovieScene actors
+			// Only display MovieScene actors
 			auto ActorFilter = [](const AActor* Actor) {
-				return Actor->IsA(AMatineeActor::StaticClass()) || Actor->IsA(ALevelSequenceActor::StaticClass());
+				return Actor->IsA(ALevelSequenceActor::StaticClass());
 			};
 			InitOptions.Filters->AddFilterPredicate(FActorFilterPredicate::CreateLambda(ActorFilter));
 		}
@@ -2366,12 +2362,7 @@ void FLevelEditorToolBar::OnCinematicsActorPicked( AActor* Actor )
 	FSlateApplication::Get().DismissAllMenus();
 
 	// Make sure we dismiss the menus before we open this
-	if (AMatineeActor* MatineeActor = Cast<AMatineeActor>(Actor))
-	{
-		// Open Matinee for editing!
-		GEditor->OpenMatinee( MatineeActor );
-	}
-	else if (ALevelSequenceActor* LevelSequenceActor = Cast<ALevelSequenceActor>(Actor))
+	if (ALevelSequenceActor* LevelSequenceActor = Cast<ALevelSequenceActor>(Actor))
 	{
 		FScopedSlowTask SlowTask(1.f, NSLOCTEXT("LevelToolBarCinematicsMenu", "LoadSequenceSlowTask", "Loading Level Sequence..."));
 		SlowTask.MakeDialog();
diff --git a/Engine/Source/Editor/LevelEditor/Public/LevelEditorActions.h b/Engine/Source/Editor/LevelEditor/Public/LevelEditorActions.h
--- a/Engine/Source/Editor/LevelEditor/Public/LevelEditorActions.h
+++ b/Engine/Source/Editor/LevelEditor/Public/LevelEditorActions.h
@@ -545,7 +545,6 @@ public:
 	TSharedPtr< FUICommandInfo > WorldProperties;
 	TSharedPtr< FUICommandInfo > OpenContentBrowser;
 	TSharedPtr< FUICommandInfo > OpenMarketplace;
-	TSharedPtr< FUICommandInfo > EditMatinee;
 	TSharedPtr< FUICommandInfo > ToggleVR;
 
 	/**
@@ -584,8 +583,6 @@ public:
 
 	TSharedPtr< FUICommandInfo > ToggleHideViewportUI;
 
-	TSharedPtr< FUICommandInfo > AddMatinee;
-
 	TSharedPtr< FUICommandInfo > MaterialQualityLevel_Low;
 	TSharedPtr< FUICommandInfo > MaterialQualityLevel_Medium;
 	TSharedPtr< FUICommandInfo > MaterialQualityLevel_High;
@@ -1248,8 +1245,6 @@ public:
 
 	static void OnAddVolume( UClass* VolumeClass );
 
-	static void OnAddMatinee();
-
 	static void SelectActorsInLayers();
 
 	static void SetWidgetMode( FWidget::EWidgetMode WidgetMode );
