Matinee: Added a cvar Matinee.AllowMatineeActors to enable accessing the matinee actors from the cinematics list. Defaults to off.

diff --git a/Engine/Source/Editor/LevelEditor/Private/LevelEditorToolBar.cpp b/Engine/Source/Editor/LevelEditor/Private/LevelEditorToolBar.cpp
--- a/Engine/Source/Editor/LevelEditor/Private/LevelEditorToolBar.cpp
+++ b/Engine/Source/Editor/LevelEditor/Private/LevelEditorToolBar.cpp
@@ -36,6 +36,7 @@
 #include "SScalabilitySettings.h"
 #include "IContentBrowserSingleton.h"
 #include "ContentBrowserModule.h"
+#include "Matinee/MatineeActor.h"
 #include "LevelSequenceActor.h"
 #include "Engine/LevelScriptBlueprint.h"
 #include "ISettingsCategory.h"
@@ -61,6 +62,11 @@
 #include "ToolMenus.h"
 #include "Subsystems/AssetEditorSubsystem.h"
 
+static TAutoConsoleVariable<int32> CVarAllowMatineeActors(
+	TEXT("Matinee.AllowMatineeActors"),
+	0,
+	TEXT("Toggles whether matinee actors should appear in the cinematics menu so that they can be edited."));
+
 namespace LevelEditorActionHelpers
 {
 	/** Filters out any classes for the Class Picker when creating or selecting classes in the Blueprints dropdown */
@@ -2306,8 +2312,10 @@ void FLevelEditorToolBar::RegisterCinematicsMenu()
 			return;
 		}
 
+		const int32 bAllowMatineeActors = CVarAllowMatineeActors->GetInt();
+
 		UWorld* World = FoundContext->LevelEditor.IsValid() ? FoundContext->LevelEditor.Pin()->GetWorld() : nullptr;
-		const bool bHasAnyCinematicsActors = !!TActorIterator<ALevelSequenceActor>(World);
+		const bool bHasAnyCinematicsActors = (bAllowMatineeActors && !!TActorIterator<AMatineeActor>(World)) || !!TActorIterator<ALevelSequenceActor>(World);
 		if (!bHasAnyCinematicsActors)
 		{
 			return;
@@ -2315,7 +2323,7 @@ void FLevelEditorToolBar::RegisterCinematicsMenu()
 
 		using namespace SceneOutliner;
 
-		// We can't build a list of LevelSequenceActors while the current World is a PIE world.
+		// We can't build a list of Matinees and LevelSequenceActors while the current World is a PIE world.
 		FInitializationOptions InitOptions;
 		{
 			InitOptions.Mode = ESceneOutlinerMode::ActorPicker;
@@ -2329,9 +2337,9 @@ void FLevelEditorToolBar::RegisterCinematicsMenu()
 			InitOptions.ColumnMap.Add(FBuiltInColumnTypes::Label(), FColumnInfo(EColumnVisibility::Visible, 0));
 			InitOptions.ColumnMap.Add(FBuiltInColumnTypes::ActorInfo(), FColumnInfo(EColumnVisibility::Visible, 10));
 
-			// Only display MovieScene actors
-			auto ActorFilter = [](const AActor* Actor) {
-				return Actor->IsA(ALevelSequenceActor::StaticClass());
+			// Only display Matinee and MovieScene actors
+			auto ActorFilter = [&](const AActor* Actor) {
+				return (bAllowMatineeActors && Actor->IsA(AMatineeActor::StaticClass())) || Actor->IsA(ALevelSequenceActor::StaticClass());
 			};
 			InitOptions.Filters->AddFilterPredicate(FActorFilterPredicate::CreateLambda(ActorFilter));
 		}
@@ -2362,7 +2370,12 @@ void FLevelEditorToolBar::OnCinematicsActorPicked( AActor* Actor )
 	FSlateApplication::Get().DismissAllMenus();
 
 	// Make sure we dismiss the menus before we open this
-	if (ALevelSequenceActor* LevelSequenceActor = Cast<ALevelSequenceActor>(Actor))
+	if (AMatineeActor* MatineeActor = Cast<AMatineeActor>(Actor))
+	{
+		// Open Matinee for editing!
+		GEditor->OpenMatinee( MatineeActor );
+	}
+	else if (ALevelSequenceActor* LevelSequenceActor = Cast<ALevelSequenceActor>(Actor))
 	{
 		FScopedSlowTask SlowTask(1.f, NSLOCTEXT("LevelToolBarCinematicsMenu", "LoadSequenceSlowTask", "Loading Level Sequence..."));
 		SlowTask.MakeDialog();
