Only check the blueprint's class flags when determining if it is placeable

diff --git a/Engine/Source/Editor/LevelEditor/Private/LevelEditorCreateActorMenu.cpp b/Engine/Source/Editor/LevelEditor/Private/LevelEditorCreateActorMenu.cpp
--- a/Engine/Source/Editor/LevelEditor/Private/LevelEditorCreateActorMenu.cpp
+++ b/Engine/Source/Editor/LevelEditor/Private/LevelEditorCreateActorMenu.cpp
@@ -238,22 +238,11 @@ static void GetContentBrowserSelectionFactoryMenuEntries( FAssetData& TargetAsse
 
 		bPlaceable = AssetSelectionUtils::IsClassPlaceable( Class );
 	}
-	else if ( TargetAssetData.GetClass() == UBlueprint::StaticClass() )
+	else if ( TargetAssetData.GetClass()->IsChildOf<UBlueprint>() )
 	{
 		// For blueprints, attempt to determine placeability from its tag information
 
 		FString TagValue;
-
-		if ( TargetAssetData.GetTagValue( FBlueprintTags::NativeParentClassPath, TagValue ) && !TagValue.IsEmpty() )
-		{
-			// If the native parent class can't be placed, neither can the blueprint
-
-			UObject* Outer = nullptr;
-			ResolveName( Outer, TagValue, false, false );
-			UClass* NativeParentClass = FindObject<UClass>( ANY_PACKAGE, *TagValue );
-
-			bPlaceable = AssetSelectionUtils::IsClassPlaceable( NativeParentClass );
-		}
 		
 		if ( bPlaceable && TargetAssetData.GetTagValue( FBlueprintTags::ClassFlags, TagValue ) && !TagValue.IsEmpty() )
 		{
