Child blueprints don't need to check if their parents are abstract when being placed, just if they are placeable and not deprecated.

diff --git a/Engine/Source/Editor/LevelEditor/Private/LevelEditorCreateActorMenu.cpp b/Engine/Source/Editor/LevelEditor/Private/LevelEditorCreateActorMenu.cpp
--- a/Engine/Source/Editor/LevelEditor/Private/LevelEditorCreateActorMenu.cpp
+++ b/Engine/Source/Editor/LevelEditor/Private/LevelEditorCreateActorMenu.cpp
@@ -238,11 +238,22 @@ static void GetContentBrowserSelectionFactoryMenuEntries( FAssetData& TargetAsse
 
 		bPlaceable = AssetSelectionUtils::IsClassPlaceable( Class );
 	}
-	else if ( TargetAssetData.GetClass()->IsChildOf<UBlueprint>() )
+	else if (TargetAssetData.GetClass()->IsChildOf<UBlueprint>())
 	{
 		// For blueprints, attempt to determine placeability from its tag information
 
 		FString TagValue;
+
+		if ( TargetAssetData.GetTagValue( FBlueprintTags::NativeParentClassPath, TagValue ) && !TagValue.IsEmpty() )
+		{
+			// If the native parent class can't be placed, neither can the blueprint
+
+			UObject* Outer = nullptr;
+			ResolveName( Outer, TagValue, false, false );
+			UClass* NativeParentClass = FindObject<UClass>( ANY_PACKAGE, *TagValue );
+
+			bPlaceable = AssetSelectionUtils::IsChildBlueprintPlaceable( NativeParentClass );
+		}
 		
 		if ( bPlaceable && TargetAssetData.GetTagValue( FBlueprintTags::ClassFlags, TagValue ) && !TagValue.IsEmpty() )
 		{
diff --git a/Engine/Source/Editor/UnrealEd/Private/AssetSelection.cpp b/Engine/Source/Editor/UnrealEd/Private/AssetSelection.cpp
--- a/Engine/Source/Editor/UnrealEd/Private/AssetSelection.cpp
+++ b/Engine/Source/Editor/UnrealEd/Private/AssetSelection.cpp
@@ -61,6 +61,16 @@ namespace AssetSelectionUtils
 			&& Class->IsChildOf( AActor::StaticClass() );
 		return bIsAddable;
 	}
+
+	// Blueprints handle their own Abstract flag, but we should check for deprecation and NotPlaceable here
+	bool IsChildBlueprintPlaceable(const UClass* Class)
+	{
+		const bool bIsAddable =
+			Class
+			&& !(Class->HasAnyClassFlags(CLASS_NotPlaceable | CLASS_Deprecated))
+			&& Class->IsChildOf(AActor::StaticClass());
+		return bIsAddable;
+	}
 	
 	void GetSelectedAssets( TArray<FAssetData>& OutSelectedAssets )
 	{
diff --git a/Engine/Source/Editor/UnrealEd/Public/AssetSelection.h b/Engine/Source/Editor/UnrealEd/Public/AssetSelection.h
--- a/Engine/Source/Editor/UnrealEd/Public/AssetSelection.h
+++ b/Engine/Source/Editor/UnrealEd/Public/AssetSelection.h
@@ -171,6 +171,13 @@ namespace AssetSelectionUtils
 	 * @return true if the class is placeable
 	 */
 	UNREALED_API bool IsClassPlaceable(const UClass* Class);
+	/**
+	 * Checks if a Blueprint based on a class type can be placed in a level
+	 *
+	 * @param Class	The class to check
+	 * @return true if the Blueprint based on this class is placeable
+	 */
+	UNREALED_API bool IsChildBlueprintPlaceable(const UClass* Class);
 	/**
 	 * Gets the selected assets in the content browser
 	 *
