Clear names of dynamic list menu entries to prevent customization
#rnx

diff --git a/Engine/Source/Editor/LevelEditor/Private/LevelEditorMenu.cpp b/Engine/Source/Editor/LevelEditor/Private/LevelEditorMenu.cpp
--- a/Engine/Source/Editor/LevelEditor/Private/LevelEditorMenu.cpp
+++ b/Engine/Source/Editor/LevelEditor/Private/LevelEditorMenu.cpp
@@ -112,7 +112,7 @@ void FLevelEditorMenu::RegisterLevelEditorMenus()
 							const FText ToolTip = FText::Format(LOCTEXT("FavoriteLevelToolTip", "Opens favorite level: {0}"), FText::FromString(CurFavorite));
 							const FText Label = FText::FromString(FPaths::GetBaseFilename(CurFavorite));
 
-							InMenu->FindOrAddSection("Favorite").AddMenuEntry(OpenFavoriteFile, Label, ToolTip);
+							InMenu->FindOrAddSection("Favorite").AddMenuEntry(OpenFavoriteFile, Label, ToolTip).Name = NAME_None;
 						}
 					}
 				};
@@ -156,7 +156,7 @@ void FLevelEditorMenu::RegisterLevelEditorMenus()
 							const FText ToolTip = FText::Format( LOCTEXT( "RecentLevelToolTip", "Opens recent level: {0}" ), FText::FromString( CurRecent ) );
 							const FText Label = FText::FromString( FPaths::GetBaseFilename( CurRecent ) );
 
-							InMenu->FindOrAddSection("Recent").AddMenuEntry( OpenRecentFile, Label, ToolTip );
+							InMenu->FindOrAddSection("Recent").AddMenuEntry( OpenRecentFile, Label, ToolTip ).Name = NAME_None;
 						}
 					}
 				};
diff --git a/Engine/Source/Editor/MainFrame/Private/Menus/LayoutsMenu.cpp b/Engine/Source/Editor/MainFrame/Private/Menus/LayoutsMenu.cpp
--- a/Engine/Source/Editor/MainFrame/Private/Menus/LayoutsMenu.cpp
+++ b/Engine/Source/Editor/MainFrame/Private/Menus/LayoutsMenu.cpp
@@ -241,7 +241,7 @@ void DisplayLayoutsInternal(FToolMenuSection& InSection, const TArray<TSharedPtr
 			// If no localization name, then display the file name
 			const FText DisplayName = (!LayoutName.IsEmpty() ? LayoutName : GetDisplayTextInternal(InLayoutIniFileNames[LayoutIndex]));
 			const FText Tooltip = GetTooltipTextInternal(DisplayName, LayoutFilePath, LayoutDescription, LayoutIndex);
-			InSection.AddMenuEntry(InXLayoutCommands[LayoutIndex], DisplayName, Tooltip);
+			InSection.AddMenuEntry(InXLayoutCommands[LayoutIndex], DisplayName, Tooltip).Name = NAME_None;
 		}
 	}
 }
diff --git a/Engine/Source/Editor/MainFrame/Private/Menus/RecentProjectsMenu.h b/Engine/Source/Editor/MainFrame/Private/Menus/RecentProjectsMenu.h
--- a/Engine/Source/Editor/MainFrame/Private/Menus/RecentProjectsMenu.h
+++ b/Engine/Source/Editor/MainFrame/Private/Menus/RecentProjectsMenu.h
@@ -39,7 +39,7 @@ public:
 			FToolMenuSection& Section = Menu->FindOrAddSection("Recent");
 			const FText DisplayName = FText::FromString( FPaths::GetBaseFilename(*ProjectName) );
 			const FText Tooltip = FText::FromString( IFileManager::Get().ConvertToAbsolutePathForExternalAppForRead(*ProjectName) );
-			Section.AddMenuEntry( FMainFrameCommands::Get().SwitchProjectCommands[ ProjectIndex ], DisplayName, Tooltip );
+			Section.AddMenuEntry( FMainFrameCommands::Get().SwitchProjectCommands[ ProjectIndex ], DisplayName, Tooltip ).Name = NAME_None;
 		}
 	}
 };
