Pushes many parts of viewport layout out of level editor up t o UnrealEd.  #rb Laurent.Barnes  #fyi chris.gagnon

diff --git a/Engine/Config/BaseEditor.ini b/Engine/Config/BaseEditor.ini
--- a/Engine/Config/BaseEditor.ini
+++ b/Engine/Config/BaseEditor.ini
@@ -530,3 +530,6 @@ bValidateOnSave=true
 [/Script/UnrealEd.DiffAssetRegistriesCommandlet]
 P4EngineBasePath="/Engine/Content/"
 P4EngineAssetPath="/Engine/"
+
+[CoreRedirects]
++ClassRedirects=(OldName="IViewportLayoutEntity", NewName="IAssetViewportLayoutEntity")
diff --git a/Engine/Plugins/MovieScene/LevelSequenceEditor/Source/LevelSequenceEditor/Private/CinematicViewport/CinematicViewportLayoutEntity.h b/Engine/Plugins/MovieScene/LevelSequenceEditor/Source/LevelSequenceEditor/Private/CinematicViewport/CinematicViewportLayoutEntity.h
--- a/Engine/Plugins/MovieScene/LevelSequenceEditor/Source/LevelSequenceEditor/Private/CinematicViewport/CinematicViewportLayoutEntity.h
+++ b/Engine/Plugins/MovieScene/LevelSequenceEditor/Source/LevelSequenceEditor/Private/CinematicViewport/CinematicViewportLayoutEntity.h
@@ -12,7 +12,7 @@
 #include "Editor/UnrealEdEngine.h"
 #include "UnrealEdGlobals.h"
 
-class FCinematicViewportLayoutEntity : public IViewportLayoutEntity
+class FCinematicViewportLayoutEntity : public ILevelViewportLayoutEntity
 {
 public:
 	FCinematicViewportLayoutEntity(const FViewportConstructionArgs& Args)
diff --git a/Engine/Source/Editor/LevelEditor/Private/LevelEditor.cpp b/Engine/Source/Editor/LevelEditor/Private/LevelEditor.cpp
--- a/Engine/Source/Editor/LevelEditor/Private/LevelEditor.cpp
+++ b/Engine/Source/Editor/LevelEditor/Private/LevelEditor.cpp
@@ -576,7 +576,7 @@ void FLevelEditorModule::RemoveStatusBarItem(FName InStatusBarIdentifier)
 	BroadcastNotificationBarChanged();
 }
 
-TSharedRef<IViewportLayoutEntity> FLevelEditorModule::FactoryViewport(FName InTypeName, const FViewportConstructionArgs& ConstructionArgs) const
+TSharedRef<ILevelViewportLayoutEntity> FLevelEditorModule::FactoryViewport(FName InTypeName, const FViewportConstructionArgs& ConstructionArgs) const
 {
 	const FViewportTypeDefinition* Definition = CustomViewports.Find(InTypeName);
 	if (Definition)
diff --git a/Engine/Source/Editor/LevelEditor/Private/LevelViewportLayout.cpp b/Engine/Source/Editor/LevelEditor/Private/LevelViewportLayout.cpp
--- a/Engine/Source/Editor/LevelEditor/Private/LevelViewportLayout.cpp
+++ b/Engine/Source/Editor/LevelEditor/Private/LevelViewportLayout.cpp
@@ -309,7 +309,7 @@ void FLevelViewportLayout::RequestMaximizeViewport( FName ViewportToMaximize, co
 
 void FLevelViewportLayout::MaximizeViewport( FName ViewportToMaximize, const bool bWantMaximize, const bool bWantImmersive, const bool bAllowAnimation )
 {
-	TSharedPtr<IViewportLayoutEntity> Entity = Viewports.FindRef(ViewportToMaximize);
+	TSharedPtr<ILevelViewportLayoutEntity> Entity = StaticCastSharedPtr<ILevelViewportLayoutEntity>(Viewports.FindRef(ViewportToMaximize));
 
 	// Should never get into a situation where the viewport is being maximized and there is already a maximized viewport. 
 	// I.E Maximized viewport is NULL which means this is a new maximize or MaximizeViewport is equal to the passed in one which means this is a restore of the current maximized viewport
@@ -430,7 +430,7 @@ void FLevelViewportLayout::MaximizeViewport( FName ViewportToMaximize, const boo
 				// Store the maximized viewport
 				MaximizedViewport = ViewportToMaximize;
 
-				TSharedPtr<IViewportLayoutEntity> MaximizedEntity = Viewports.FindRef(MaximizedViewport);
+				TSharedPtr<ILevelViewportLayoutEntity> MaximizedEntity = StaticCastSharedPtr<ILevelViewportLayoutEntity>(Viewports.FindRef(MaximizedViewport));
 				if (MaximizedEntity.IsValid())
 				{
 					// Replace our viewport with a dummy widget in it's place during the maximize transition.  We can't
@@ -580,7 +580,7 @@ void FLevelViewportLayout::TakeHighResScreenShot()
 {
 	if (bIsImmersive || bIsMaximized)
 	{
-		TSharedPtr<IViewportLayoutEntity> MaximizedViewportEntity = Viewports.FindRef(MaximizedViewport);
+		TSharedPtr<ILevelViewportLayoutEntity> MaximizedViewportEntity = StaticCastSharedPtr<ILevelViewportLayoutEntity>(Viewports.FindRef(MaximizedViewport));
 		check(MaximizedViewportEntity.IsValid());
 
 		MaximizedViewportEntity->TakeHighResScreenShot();
@@ -589,7 +589,8 @@ void FLevelViewportLayout::TakeHighResScreenShot()
 	{
 		for (auto& Elem : Viewports)
 		{
-			TSharedPtr<IViewportLayoutEntity> ViewportEntity = Elem.Value; 
+			TSharedPtr<ILevelViewportLayoutEntity> ViewportEntity = StaticCastSharedPtr<ILevelViewportLayoutEntity>(Elem.Value);
+
 			if (ViewportEntity.IsValid())
 			{
 				ViewportEntity->TakeHighResScreenShot();
@@ -640,7 +641,7 @@ void FLevelViewportLayout::FinishMaximizeTransition()
 {
 	if( bIsTransitioning )
 	{
-		TSharedPtr<IViewportLayoutEntity> MaximizedViewportEntity = Viewports.FindRef(MaximizedViewport);
+		TSharedPtr<ILevelViewportLayoutEntity> MaximizedViewportEntity = StaticCastSharedPtr<ILevelViewportLayoutEntity>(Viewports.FindRef(MaximizedViewport));
 		check(MaximizedViewportEntity.IsValid());
 
 		// The transition animation is complete, allow the engine to tick normally
diff --git a/Engine/Source/Editor/LevelEditor/Private/LevelViewportLayout2x2.cpp b/Engine/Source/Editor/LevelEditor/Private/LevelViewportLayout2x2.cpp
--- a/Engine/Source/Editor/LevelEditor/Private/LevelViewportLayout2x2.cpp
+++ b/Engine/Source/Editor/LevelEditor/Private/LevelViewportLayout2x2.cpp
@@ -65,25 +65,25 @@ TSharedRef<SWidget> FLevelViewportLayout2x2::MakeViewportLayout(const FString& L
 	Args.bRealtime = false;
 	Args.ConfigKey = *TopLeftKey;
 	Args.ViewportType = LVT_OrthoYZ;
-	TSharedPtr< IViewportLayoutEntity > ViewportTL = LevelEditor.FactoryViewport(*TopLeftType, Args);
+	TSharedPtr< ILevelViewportLayoutEntity > ViewportTL = LevelEditor.FactoryViewport(*TopLeftType, Args);
 
 	// Persp viewport
 	Args.bRealtime = !FPlatformMisc::IsRemoteSession();
 	Args.ConfigKey = *BottomLeftKey;
 	Args.ViewportType = LVT_Perspective;
-	TSharedPtr< IViewportLayoutEntity > ViewportBL = LevelEditor.FactoryViewport(*BottomLeftType, Args);
+	TSharedPtr< ILevelViewportLayoutEntity > ViewportBL = LevelEditor.FactoryViewport(*BottomLeftType, Args);
 
 	// Front viewport
 	Args.bRealtime = false;
 	Args.ConfigKey = *TopRightKey;
 	Args.ViewportType = LVT_OrthoXZ;
-	TSharedPtr< IViewportLayoutEntity > ViewportTR = LevelEditor.FactoryViewport(*TopRightType, Args);
+	TSharedPtr< ILevelViewportLayoutEntity > ViewportTR = LevelEditor.FactoryViewport(*TopRightType, Args);
 
 	// Top Viewport
 	Args.bRealtime = false;
 	Args.ConfigKey = *BottomRightKey;
 	Args.ViewportType = LVT_OrthoXY;
-	TSharedPtr< IViewportLayoutEntity > ViewportBR = LevelEditor.FactoryViewport(*BottomRightType, Args);
+	TSharedPtr< ILevelViewportLayoutEntity > ViewportBR = LevelEditor.FactoryViewport(*BottomRightType, Args);
 
 	Viewports.Add( *TopLeftKey, ViewportTL );
 	Viewports.Add( *BottomLeftKey, ViewportBL );
diff --git a/Engine/Source/Editor/LevelEditor/Private/LevelViewportLayoutEntity.h b/Engine/Source/Editor/LevelEditor/Private/LevelViewportLayoutEntity.h
--- a/Engine/Source/Editor/LevelEditor/Private/LevelViewportLayoutEntity.h
+++ b/Engine/Source/Editor/LevelEditor/Private/LevelViewportLayoutEntity.h
@@ -9,7 +9,7 @@
 class FLevelEditorViewportClient;
 class SLevelViewport;
 
-class LEVELEDITOR_API FLevelViewportLayoutEntity : public IViewportLayoutEntity
+class LEVELEDITOR_API FLevelViewportLayoutEntity : public ILevelViewportLayoutEntity
 {
 public:
 
diff --git a/Engine/Source/Editor/LevelEditor/Private/LevelViewportLayoutFourPanes.cpp b/Engine/Source/Editor/LevelEditor/Private/LevelViewportLayoutFourPanes.cpp
--- a/Engine/Source/Editor/LevelEditor/Private/LevelViewportLayoutFourPanes.cpp
+++ b/Engine/Source/Editor/LevelEditor/Private/LevelViewportLayoutFourPanes.cpp
@@ -72,24 +72,24 @@ TSharedRef<SWidget> FLevelViewportLayoutFourPanes::MakeViewportLayout(const FStr
 	Args.bRealtime = !FPlatformMisc::IsRemoteSession();
 	Args.ConfigKey = *ViewportKey0;
 	Args.ViewportType = LVT_Perspective;
-	TSharedRef<IViewportLayoutEntity> Viewport0 = LevelEditor.FactoryViewport(*ViewportType0, Args);
+	TSharedRef<ILevelViewportLayoutEntity> Viewport0 = LevelEditor.FactoryViewport(*ViewportType0, Args);
 
 	Args.bRealtime = false;
 	Args.ConfigKey = *ViewportKey1;
 	Args.ViewportType = LVT_OrthoXY;
-	TSharedRef<IViewportLayoutEntity> Viewport1 = LevelEditor.FactoryViewport(*ViewportType1, Args);
+	TSharedRef<ILevelViewportLayoutEntity> Viewport1 = LevelEditor.FactoryViewport(*ViewportType1, Args);
 
 	// Front viewport
 	Args.bRealtime = false;
 	Args.ConfigKey = *ViewportKey2;
 	Args.ViewportType = LVT_OrthoXZ;
-	TSharedRef<IViewportLayoutEntity> Viewport2 = LevelEditor.FactoryViewport(*ViewportType2, Args);
+	TSharedRef<ILevelViewportLayoutEntity> Viewport2 = LevelEditor.FactoryViewport(*ViewportType2, Args);
 
 	// Top Viewport
 	Args.bRealtime = false;
 	Args.ConfigKey = *ViewportKey3;
 	Args.ViewportType = LVT_OrthoYZ;
-	TSharedRef<IViewportLayoutEntity> Viewport3 = LevelEditor.FactoryViewport(*ViewportType2, Args);
+	TSharedRef<ILevelViewportLayoutEntity> Viewport3 = LevelEditor.FactoryViewport(*ViewportType2, Args);
 
 	Viewports.Add( *ViewportKey0, Viewport0 );
 	Viewports.Add( *ViewportKey1, Viewport1 );
@@ -169,7 +169,7 @@ void FLevelViewportLayoutFourPanes::ReplaceWidget( TSharedRef< SWidget > Source,
 // FLevelViewportLayoutFourPanesLeft /////////////////////////////
 
 TSharedRef<SWidget> FLevelViewportLayoutFourPanesLeft::MakeFourPanelWidget(
-	TMap<FName, TSharedPtr< IViewportLayoutEntity >>& ViewportWidgets,
+	TMap<FName, TSharedPtr< IAssetViewportLayoutEntity >>& ViewportWidgets,
 	TSharedRef<SWidget> Viewport0, TSharedRef<SWidget> Viewport1, TSharedRef<SWidget> Viewport2, TSharedRef<SWidget> Viewport3,
 	float PrimarySplitterPercentage, float SecondarySplitterPercentage0, float SecondarySplitterPercentage1)
 {
@@ -210,7 +210,7 @@ TSharedRef<SWidget> FLevelViewportLayoutFourPanesLeft::MakeFourPanelWidget(
 // FLevelViewportLayoutFourPanesRight /////////////////////////////
 
 TSharedRef<SWidget> FLevelViewportLayoutFourPanesRight::MakeFourPanelWidget(
-	TMap<FName, TSharedPtr< IViewportLayoutEntity >>& ViewportWidgets,
+	TMap<FName, TSharedPtr< IAssetViewportLayoutEntity >>& ViewportWidgets,
 	TSharedRef<SWidget> Viewport0, TSharedRef<SWidget> Viewport1, TSharedRef<SWidget> Viewport2, TSharedRef<SWidget> Viewport3,
 	float PrimarySplitterPercentage, float SecondarySplitterPercentage0, float SecondarySplitterPercentage1)
 {
@@ -251,7 +251,7 @@ TSharedRef<SWidget> FLevelViewportLayoutFourPanesRight::MakeFourPanelWidget(
 // FLevelViewportLayoutFourPanesTop /////////////////////////////
 
 TSharedRef<SWidget> FLevelViewportLayoutFourPanesTop::MakeFourPanelWidget(
-	TMap<FName, TSharedPtr< IViewportLayoutEntity >>& ViewportWidgets,
+	TMap<FName, TSharedPtr< IAssetViewportLayoutEntity >>& ViewportWidgets,
 	TSharedRef<SWidget> Viewport0, TSharedRef<SWidget> Viewport1, TSharedRef<SWidget> Viewport2, TSharedRef<SWidget> Viewport3,
 	float PrimarySplitterPercentage, float SecondarySplitterPercentage0, float SecondarySplitterPercentage1)
 {
@@ -292,7 +292,7 @@ TSharedRef<SWidget> FLevelViewportLayoutFourPanesTop::MakeFourPanelWidget(
 // FLevelViewportLayoutFourPanesBottom /////////////////////////////
 
 TSharedRef<SWidget> FLevelViewportLayoutFourPanesBottom::MakeFourPanelWidget(
-	TMap<FName, TSharedPtr< IViewportLayoutEntity >>& ViewportWidgets,
+	TMap<FName, TSharedPtr< IAssetViewportLayoutEntity >>& ViewportWidgets,
 	TSharedRef<SWidget> Viewport0, TSharedRef<SWidget> Viewport1, TSharedRef<SWidget> Viewport2, TSharedRef<SWidget> Viewport3,
 	float PrimarySplitterPercentage, float SecondarySplitterPercentage0, float SecondarySplitterPercentage1)
 {
diff --git a/Engine/Source/Editor/LevelEditor/Private/LevelViewportLayoutFourPanes.h b/Engine/Source/Editor/LevelEditor/Private/LevelViewportLayoutFourPanes.h
--- a/Engine/Source/Editor/LevelEditor/Private/LevelViewportLayoutFourPanes.h
+++ b/Engine/Source/Editor/LevelEditor/Private/LevelViewportLayoutFourPanes.h
@@ -22,7 +22,7 @@ protected:
 	virtual TSharedRef<SWidget> MakeViewportLayout(const FString& LayoutString) override;
 
 	virtual TSharedRef<SWidget> MakeFourPanelWidget(
-		TMap<FName, TSharedPtr< IViewportLayoutEntity >>& ViewportWidgets,
+		TMap<FName, TSharedPtr< IAssetViewportLayoutEntity >>& ViewportWidgets,
 		TSharedRef<SWidget> Viewport0, TSharedRef<SWidget> Viewport1, TSharedRef<SWidget> Viewport2, TSharedRef<SWidget> Viewport3,
 		float PrimarySplitterPercentage, float SecondarySplitterPercentage0, float SecondarySplitterPercentage1) = 0;
 
@@ -45,7 +45,7 @@ public:
 	virtual const FName& GetLayoutTypeName() const override { return LevelViewportConfigurationNames::FourPanesLeft; }
 
 	virtual TSharedRef<SWidget> MakeFourPanelWidget(
-		TMap<FName, TSharedPtr< IViewportLayoutEntity >>& ViewportWidgets,
+		TMap<FName, TSharedPtr< IAssetViewportLayoutEntity >>& ViewportWidgets,
 		TSharedRef<SWidget> Viewport0, TSharedRef<SWidget> Viewport1, TSharedRef<SWidget> Viewport2, TSharedRef<SWidget> Viewport3,
 		float PrimarySplitterPercentage, float SecondarySplitterPercentage0, float SecondarySplitterPercentage1) override;
 };
@@ -59,7 +59,7 @@ public:
 	virtual const FName& GetLayoutTypeName() const override { return LevelViewportConfigurationNames::FourPanesRight; }
 
 	virtual TSharedRef<SWidget> MakeFourPanelWidget(
-		TMap<FName, TSharedPtr< IViewportLayoutEntity >>& ViewportWidgets,
+		TMap<FName, TSharedPtr< IAssetViewportLayoutEntity >>& ViewportWidgets,
 		TSharedRef<SWidget> Viewport0, TSharedRef<SWidget> Viewport1, TSharedRef<SWidget> Viewport2, TSharedRef<SWidget> Viewport3,
 		float PrimarySplitterPercentage, float SecondarySplitterPercentage0, float SecondarySplitterPercentage1) override;
 };
@@ -73,7 +73,7 @@ public:
 	virtual const FName& GetLayoutTypeName() const override { return LevelViewportConfigurationNames::FourPanesTop; }
 
 	virtual TSharedRef<SWidget> MakeFourPanelWidget(
-		TMap<FName, TSharedPtr< IViewportLayoutEntity >>& ViewportWidgets,
+		TMap<FName, TSharedPtr< IAssetViewportLayoutEntity >>& ViewportWidgets,
 		TSharedRef<SWidget> Viewport0, TSharedRef<SWidget> Viewport1, TSharedRef<SWidget> Viewport2, TSharedRef<SWidget> Viewport3,
 		float PrimarySplitterPercentage, float SecondarySplitterPercentage0, float SecondarySplitterPercentage1) override;
 };
@@ -87,7 +87,7 @@ public:
 	virtual const FName& GetLayoutTypeName() const override { return LevelViewportConfigurationNames::FourPanesBottom; }
 
 	virtual TSharedRef<SWidget> MakeFourPanelWidget(
-		TMap<FName, TSharedPtr< IViewportLayoutEntity >>& ViewportWidgets,
+		TMap<FName, TSharedPtr< IAssetViewportLayoutEntity >>& ViewportWidgets,
 		TSharedRef<SWidget> Viewport0, TSharedRef<SWidget> Viewport1, TSharedRef<SWidget> Viewport2, TSharedRef<SWidget> Viewport3,
 		float PrimarySplitterPercentage, float SecondarySplitterPercentage0, float SecondarySplitterPercentage1) override;
 };
diff --git a/Engine/Source/Editor/LevelEditor/Private/LevelViewportLayoutOnePane.cpp b/Engine/Source/Editor/LevelEditor/Private/LevelViewportLayoutOnePane.cpp
--- a/Engine/Source/Editor/LevelEditor/Private/LevelViewportLayoutOnePane.cpp
+++ b/Engine/Source/Editor/LevelEditor/Private/LevelViewportLayoutOnePane.cpp
@@ -55,7 +55,7 @@ TSharedRef<SWidget> FLevelViewportLayoutOnePane::MakeViewportLayout(const FStrin
 	Args.bRealtime = !FPlatformMisc::IsRemoteSession();
 	Args.ConfigKey = *ViewportKey;
 	Args.ViewportType = LVT_Perspective;
-	TSharedRef<IViewportLayoutEntity> Viewport = LevelEditor.FactoryViewport(*ViewportType, Args);
+	TSharedRef<ILevelViewportLayoutEntity> Viewport = LevelEditor.FactoryViewport(*ViewportType, Args);
 
 	ViewportBox =
 		SNew(SHorizontalBox)
diff --git a/Engine/Source/Editor/LevelEditor/Private/LevelViewportLayoutThreePanes.cpp b/Engine/Source/Editor/LevelEditor/Private/LevelViewportLayoutThreePanes.cpp
--- a/Engine/Source/Editor/LevelEditor/Private/LevelViewportLayoutThreePanes.cpp
+++ b/Engine/Source/Editor/LevelEditor/Private/LevelViewportLayoutThreePanes.cpp
@@ -67,18 +67,18 @@ TSharedRef<SWidget> FLevelViewportLayoutThreePanes::MakeViewportLayout(const FSt
 	Args.bRealtime = !FPlatformMisc::IsRemoteSession();
 	Args.ConfigKey = *ViewportKey0;
 	Args.ViewportType = LVT_Perspective;
-	TSharedRef<IViewportLayoutEntity> Viewport0 = LevelEditor.FactoryViewport(*ViewportType0, Args);
+	TSharedRef<ILevelViewportLayoutEntity> Viewport0 = LevelEditor.FactoryViewport(*ViewportType0, Args);
 
 	Args.bRealtime = false;
 	Args.ConfigKey = *ViewportKey1;
 	Args.ViewportType = LVT_OrthoXY;
-	TSharedRef<IViewportLayoutEntity> Viewport1 = LevelEditor.FactoryViewport(*ViewportType1, Args);
+	TSharedRef<ILevelViewportLayoutEntity> Viewport1 = LevelEditor.FactoryViewport(*ViewportType1, Args);
 
 	// Front viewport
 	Args.bRealtime = false;
 	Args.ConfigKey = *ViewportKey2;
 	Args.ViewportType = LVT_OrthoXZ;
-	TSharedRef<IViewportLayoutEntity> Viewport2 = LevelEditor.FactoryViewport(*ViewportType2, Args);
+	TSharedRef<ILevelViewportLayoutEntity> Viewport2 = LevelEditor.FactoryViewport(*ViewportType2, Args);
 
 	Viewports.Add( *ViewportKey0, Viewport0 );
 	Viewports.Add( *ViewportKey1, Viewport1 );
@@ -152,7 +152,7 @@ void FLevelViewportLayoutThreePanes::ReplaceWidget( TSharedRef< SWidget > Source
 // FLevelViewportLayoutThreePanesLeft /////////////////////////////
 
 TSharedRef<SWidget> FLevelViewportLayoutThreePanesLeft::MakeThreePanelWidget(
-	TMap< FName, TSharedPtr< IViewportLayoutEntity >>& ViewportWidgets,
+	TMap< FName, TSharedPtr< IAssetViewportLayoutEntity >>& ViewportWidgets,
 	const TSharedRef<SWidget>& Viewport0, const TSharedRef<SWidget>& Viewport1, const TSharedRef<SWidget>& Viewport2,
 	float PrimarySplitterPercentage, float SecondarySplitterPercentage)
 {
@@ -188,7 +188,7 @@ TSharedRef<SWidget> FLevelViewportLayoutThreePanesLeft::MakeThreePanelWidget(
 // FLevelViewportLayoutThreePanesRight /////////////////////////////
 
 TSharedRef<SWidget> FLevelViewportLayoutThreePanesRight::MakeThreePanelWidget(
-	TMap< FName, TSharedPtr< IViewportLayoutEntity >>& ViewportWidgets,
+	TMap< FName, TSharedPtr< IAssetViewportLayoutEntity >>& ViewportWidgets,
 	const TSharedRef<SWidget>& Viewport0, const TSharedRef<SWidget>& Viewport1, const TSharedRef<SWidget>& Viewport2,
 	float PrimarySplitterPercentage, float SecondarySplitterPercentage)
 {
@@ -224,7 +224,7 @@ TSharedRef<SWidget> FLevelViewportLayoutThreePanesRight::MakeThreePanelWidget(
 // FLevelViewportLayoutThreePanesTop /////////////////////////////
 
 TSharedRef<SWidget> FLevelViewportLayoutThreePanesTop::MakeThreePanelWidget(
-	TMap< FName, TSharedPtr< IViewportLayoutEntity >>& ViewportWidgets,
+	TMap< FName, TSharedPtr< IAssetViewportLayoutEntity >>& ViewportWidgets,
 	const TSharedRef<SWidget>& Viewport0, const TSharedRef<SWidget>& Viewport1, const TSharedRef<SWidget>& Viewport2,
 	float PrimarySplitterPercentage, float SecondarySplitterPercentage)
 {
@@ -260,7 +260,7 @@ TSharedRef<SWidget> FLevelViewportLayoutThreePanesTop::MakeThreePanelWidget(
 // FLevelViewportLayoutThreePanesBottom /////////////////////////////
 
 TSharedRef<SWidget> FLevelViewportLayoutThreePanesBottom::MakeThreePanelWidget(
-	TMap< FName, TSharedPtr< IViewportLayoutEntity >>& ViewportWidgets,
+	TMap< FName, TSharedPtr< IAssetViewportLayoutEntity >>& ViewportWidgets,
 	const TSharedRef<SWidget>& Viewport0, const TSharedRef<SWidget>& Viewport1, const TSharedRef<SWidget>& Viewport2,
 	float PrimarySplitterPercentage, float SecondarySplitterPercentage)
 {
diff --git a/Engine/Source/Editor/LevelEditor/Private/LevelViewportLayoutThreePanes.h b/Engine/Source/Editor/LevelEditor/Private/LevelViewportLayoutThreePanes.h
--- a/Engine/Source/Editor/LevelEditor/Private/LevelViewportLayoutThreePanes.h
+++ b/Engine/Source/Editor/LevelEditor/Private/LevelViewportLayoutThreePanes.h
@@ -22,7 +22,7 @@ protected:
 	virtual TSharedRef<SWidget> MakeViewportLayout(const FString& LayoutString) override;
 
 	virtual TSharedRef<SWidget> MakeThreePanelWidget(
-		TMap< FName, TSharedPtr< IViewportLayoutEntity >>& ViewportWidgets,
+		TMap< FName, TSharedPtr< IAssetViewportLayoutEntity >>& ViewportWidgets,
 		const TSharedRef<SWidget>& ViewportKey0, const TSharedRef<SWidget>& ViewportKey1, const TSharedRef<SWidget>& ViewportKey2,
 		float PrimarySplitterPercentage, float SecondarySplitterPercentage) = 0;
 
@@ -44,7 +44,7 @@ public:
 	virtual const FName& GetLayoutTypeName() const override { return LevelViewportConfigurationNames::ThreePanesLeft; }
 
 	virtual TSharedRef<SWidget> MakeThreePanelWidget(
-		TMap< FName, TSharedPtr< IViewportLayoutEntity >>& ViewportWidgets,
+		TMap< FName, TSharedPtr< IAssetViewportLayoutEntity >>& ViewportWidgets,
 		const TSharedRef<SWidget>& ViewportKey0, const TSharedRef<SWidget>& ViewportKey1, const TSharedRef<SWidget>& ViewportKey2,
 		float PrimarySplitterPercentage, float SecondarySplitterPercentage) override;
 };
@@ -58,7 +58,7 @@ public:
 	virtual const FName& GetLayoutTypeName() const override { return LevelViewportConfigurationNames::ThreePanesRight; }
 
 	virtual TSharedRef<SWidget> MakeThreePanelWidget(
-		TMap< FName, TSharedPtr< IViewportLayoutEntity >>& ViewportWidgets,
+		TMap< FName, TSharedPtr< IAssetViewportLayoutEntity >>& ViewportWidgets,
 		const TSharedRef<SWidget>& ViewportKey0, const TSharedRef<SWidget>& ViewportKey1, const TSharedRef<SWidget>& ViewportKey2,
 		float PrimarySplitterPercentage, float SecondarySplitterPercentage) override;
 };
@@ -72,7 +72,7 @@ public:
 	virtual const FName& GetLayoutTypeName() const override { return LevelViewportConfigurationNames::ThreePanesTop; }
 
 	virtual TSharedRef<SWidget> MakeThreePanelWidget(
-		TMap< FName, TSharedPtr< IViewportLayoutEntity >>& ViewportWidgets,
+		TMap< FName, TSharedPtr< IAssetViewportLayoutEntity >>& ViewportWidgets,
 		const TSharedRef<SWidget>& ViewportKey0, const TSharedRef<SWidget>& ViewportKey1, const TSharedRef<SWidget>& ViewportKey2,
 		float PrimarySplitterPercentage, float SecondarySplitterPercentage) override;
 };
@@ -86,7 +86,7 @@ public:
 	virtual const FName& GetLayoutTypeName() const override { return LevelViewportConfigurationNames::ThreePanesBottom; }
 
 	virtual TSharedRef<SWidget> MakeThreePanelWidget(
-		TMap< FName, TSharedPtr< IViewportLayoutEntity >>& ViewportWidgets,
+		TMap< FName, TSharedPtr< IAssetViewportLayoutEntity >>& ViewportWidgets,
 		const TSharedRef<SWidget>& ViewportKey0, const TSharedRef<SWidget>& ViewportKey1, const TSharedRef<SWidget>& ViewportKey2,
 		float PrimarySplitterPercentage, float SecondarySplitterPercentage) override;
 };
diff --git a/Engine/Source/Editor/LevelEditor/Private/LevelViewportLayoutTwoPanes.cpp b/Engine/Source/Editor/LevelEditor/Private/LevelViewportLayoutTwoPanes.cpp
--- a/Engine/Source/Editor/LevelEditor/Private/LevelViewportLayoutTwoPanes.cpp
+++ b/Engine/Source/Editor/LevelEditor/Private/LevelViewportLayoutTwoPanes.cpp
@@ -55,12 +55,12 @@ TSharedRef<SWidget> TLevelViewportLayoutTwoPanes<TOrientation>::MakeViewportLayo
 	Args.bRealtime = false;
 	Args.ConfigKey = *ViewportKey0;
 	Args.ViewportType = LVT_OrthoXY;
-	TSharedRef<IViewportLayoutEntity> Viewport0 = LevelEditor.FactoryViewport(*ViewportType0, Args);
+	TSharedRef<ILevelViewportLayoutEntity> Viewport0 = LevelEditor.FactoryViewport(*ViewportType0, Args);
 
 	Args.bRealtime = !FPlatformMisc::IsRemoteSession();
 	Args.ConfigKey = *ViewportKey1;
 	Args.ViewportType = LVT_Perspective;
-	TSharedRef<IViewportLayoutEntity> Viewport1 = LevelEditor.FactoryViewport(*ViewportType1, Args);
+	TSharedRef<ILevelViewportLayoutEntity> Viewport1 = LevelEditor.FactoryViewport(*ViewportType1, Args);
 
 	Viewports.Add( *ViewportKey0, Viewport0 );
 	Viewports.Add( *ViewportKey1, Viewport1 );
diff --git a/Engine/Source/Editor/LevelEditor/Private/LevelViewportTabContent.cpp b/Engine/Source/Editor/LevelEditor/Private/LevelViewportTabContent.cpp
--- a/Engine/Source/Editor/LevelEditor/Private/LevelViewportTabContent.cpp
+++ b/Engine/Source/Editor/LevelEditor/Private/LevelViewportTabContent.cpp
@@ -11,6 +11,7 @@
 #include "LevelViewportLayoutThreePanes.h"
 #include "LevelViewportLayoutFourPanes.h"
 #include "Widgets/Docking/SDockTab.h"
+#include "LevelViewportLayout.h"
 
 
 // FLevelViewportTabContent ///////////////////////////
@@ -61,62 +62,64 @@ void FLevelViewportTabContent::Initialize(TSharedPtr<ILevelEditor> InParentLevel
 	SetViewportConfiguration(LayoutType);
 }
 
+
 bool FLevelViewportTabContent::IsVisible() const
 {
-	if (ActiveLevelViewportLayout.IsValid())
+	if (ActiveViewportLayout.IsValid())
 	{
-		return ActiveLevelViewportLayout->IsVisible();
+		TSharedPtr<FLevelViewportLayout> ViewportLayout = StaticCastSharedPtr<FLevelViewportLayout>(ActiveViewportLayout);
+		return ViewportLayout->IsVisible();
 	}
 	return false;
 }
 
-const TMap< FName, TSharedPtr< IViewportLayoutEntity > >* FLevelViewportTabContent::GetViewports() const
+const TMap< FName, TSharedPtr< IAssetViewportLayoutEntity > >* FLevelViewportTabContent::GetViewports() const
 {
-	if (ActiveLevelViewportLayout.IsValid())
+	if (ActiveViewportLayout.IsValid())
 	{
-		return &ActiveLevelViewportLayout->GetViewports();
+		return &ActiveViewportLayout->GetViewports();
 	}
-	return NULL;
+	return nullptr;
 }
 
 void FLevelViewportTabContent::SetViewportConfiguration(const FName& ConfigurationName)
 {
-	bool bSwitchingLayouts = ActiveLevelViewportLayout.IsValid();
+	bool bSwitchingLayouts = ActiveViewportLayout.IsValid();
 
 	if (bSwitchingLayouts)
 	{
 		SaveConfig();
-		ActiveLevelViewportLayout.Reset();
+		ActiveViewportLayout.Reset();
 	}
 
-	ActiveLevelViewportLayout = ConstructViewportLayoutByTypeName(ConfigurationName, bSwitchingLayouts);
-	check (ActiveLevelViewportLayout.IsValid());
+	ActiveViewportLayout = ConstructViewportLayoutByTypeName(ConfigurationName, bSwitchingLayouts);
+	check (ActiveViewportLayout.IsValid());
 
 	UpdateViewportTabWidget();
 }
 
 void FLevelViewportTabContent::SaveConfig() const
 {
-	if (ActiveLevelViewportLayout.IsValid())
+	if (ActiveViewportLayout.IsValid())
 	{
 		if (!LayoutString.IsEmpty())
 		{
-			FString LayoutTypeString = ActiveLevelViewportLayout->GetLayoutTypeName().ToString();
+			FString LayoutTypeString = ActiveViewportLayout->GetLayoutTypeName().ToString();
 
 			const FString& IniSection = FLayoutSaveRestore::GetAdditionalLayoutConfigIni();
 			GConfig->SetString(*IniSection, *(LayoutString + TEXT(".LayoutType")), *LayoutTypeString, GEditorPerProjectIni);
 		}
 
-		ActiveLevelViewportLayout->SaveLayoutString(LayoutString);
+		ActiveViewportLayout->SaveLayoutString(LayoutString);
 	}
 }
 
 void FLevelViewportTabContent::RefreshViewportConfiguration()
 {
-	check(ActiveLevelViewportLayout.IsValid());
+	check(ActiveViewportLayout.IsValid());
 
-	FName ConfigurationName = ActiveLevelViewportLayout->GetLayoutTypeName();
-	for (auto& Pair : ActiveLevelViewportLayout->GetViewports())
+	FName ConfigurationName = ActiveViewportLayout->GetLayoutTypeName();
+	for (auto& Pair : ActiveViewportLayout->GetViewports())
 	{
 		if (Pair.Value->AsWidget()->HasFocusedDescendants())
 		{
@@ -125,11 +128,11 @@ void FLevelViewportTabContent::RefreshViewportConfiguration()
 		}
 	}
 
-	ActiveLevelViewportLayout.Reset();
+	ActiveViewportLayout.Reset();
 
 	bool bSwitchingLayouts = false;
-	ActiveLevelViewportLayout = ConstructViewportLayoutByTypeName(ConfigurationName, bSwitchingLayouts);
-	check(ActiveLevelViewportLayout.IsValid());
+	ActiveViewportLayout = ConstructViewportLayoutByTypeName(ConfigurationName, bSwitchingLayouts);
+	check(ActiveViewportLayout.IsValid());
 
 	UpdateViewportTabWidget();
 }
@@ -137,9 +140,9 @@ void FLevelViewportTabContent::RefreshViewportConfiguration()
 
 bool FLevelViewportTabContent::IsViewportConfigurationSet(const FName& ConfigurationName) const
 {
-	if (ActiveLevelViewportLayout.IsValid())
+	if (ActiveViewportLayout.IsValid())
 	{
-		return ActiveLevelViewportLayout->GetLayoutTypeName() == ConfigurationName;
+		return ActiveViewportLayout->GetLayoutTypeName() == ConfigurationName;
 	}
 	return false;
 }
@@ -153,14 +156,15 @@ bool FLevelViewportTabContent::BelongsToTab(TSharedRef<class SDockTab> InParentT
 void FLevelViewportTabContent::UpdateViewportTabWidget()
 {
 	TSharedPtr<SDockTab> ParentTabPinned = ParentTab.Pin();
-	if (ParentTabPinned.IsValid() && ActiveLevelViewportLayout.IsValid())
+	if (ParentTabPinned.IsValid() && ActiveViewportLayout.IsValid())
 	{
-		TSharedRef<SWidget> LayoutWidget = ActiveLevelViewportLayout->BuildViewportLayout(ParentTabPinned, SharedThis(this), LayoutString, ParentLevelEditor);
+
+		TSharedRef<SWidget> LayoutWidget = StaticCastSharedPtr<FLevelViewportLayout>(ActiveViewportLayout)->BuildViewportLayout(ParentTabPinned, SharedThis(this), LayoutString, ParentLevelEditor);
 		ParentTabPinned->SetContent(LayoutWidget);
 
 		if (PreviouslyFocusedViewport.IsSet())
 		{
-			TSharedPtr<IViewportLayoutEntity> ViewportToFocus = ActiveLevelViewportLayout->GetViewports().FindRef(PreviouslyFocusedViewport.GetValue());
+			TSharedPtr<IAssetViewportLayoutEntity> ViewportToFocus = ActiveViewportLayout->GetViewports().FindRef(PreviouslyFocusedViewport.GetValue());
 			if (ViewportToFocus.IsValid())
 			{
 				ViewportToFocus->SetKeyboardFocus();
diff --git a/Engine/Source/Editor/LevelEditor/Private/SLevelEditor.cpp b/Engine/Source/Editor/LevelEditor/Private/SLevelEditor.cpp
--- a/Engine/Source/Editor/LevelEditor/Private/SLevelEditor.cpp
+++ b/Engine/Source/Editor/LevelEditor/Private/SLevelEditor.cpp
@@ -344,14 +344,15 @@ bool SLevelEditor::HasActivePlayInEditorViewport() const
 		if (ViewportTab.IsValid())
 		{
 			// Get all the viewports in the layout
-			const TMap< FName, TSharedPtr< IViewportLayoutEntity > >* LevelViewports = ViewportTab.Pin()->GetViewports();
+			const TMap< FName, TSharedPtr< IAssetViewportLayoutEntity > >* LevelViewports = ViewportTab.Pin()->GetViewports();
 
 			if (LevelViewports != NULL)
 			{
 				// Search for a viewport with a pie session
 				for (auto& Pair : *LevelViewports)
 				{
-					if (Pair.Value->IsPlayInEditorViewportActive())
+					TSharedPtr<ILevelViewportLayoutEntity> ViewportEntity = StaticCastSharedPtr<ILevelViewportLayoutEntity>(Pair.Value);
+					if (ViewportEntity.IsValid() && ViewportEntity->IsPlayInEditorViewportActive())
 					{
 						return true;
 					}
@@ -392,13 +393,13 @@ TSharedPtr<SLevelViewport> SLevelEditor::GetActiveViewport()
 			// Only check the viewports in the tab if its visible
 			if( ViewportTab->IsVisible() )
 			{
-				const TMap< FName, TSharedPtr< IViewportLayoutEntity > >* LevelViewports = ViewportTab->GetViewports();
+				const TMap< FName, TSharedPtr< IAssetViewportLayoutEntity > >* LevelViewports = ViewportTab->GetViewports();
 
-				if (LevelViewports != NULL)
+				if (LevelViewports != nullptr)
 				{
 					for(auto& Pair : *LevelViewports)
 					{
-						TSharedPtr<SLevelViewport> Viewport = Pair.Value->AsLevelViewport();
+						TSharedPtr<SLevelViewport> Viewport = StaticCastSharedPtr<ILevelViewportLayoutEntity>(Pair.Value)->AsLevelViewport();
 						if( Viewport.IsValid() && Viewport->IsInForegroundTab() )
 						{
 							if( &Viewport->GetLevelViewportClient() == GCurrentLevelEditingViewportClient )
@@ -969,13 +970,14 @@ void SLevelEditor::OnViewportTabClosed(TSharedRef<SDockTab> ClosedTab)
 
 void SLevelEditor::SaveViewportTabInfo(TSharedRef<const FLevelViewportTabContent> ViewportTabContent)
 {
-	const TMap<FName, TSharedPtr<IViewportLayoutEntity>>* const Viewports = ViewportTabContent->GetViewports();
+	const TMap<FName, TSharedPtr<IAssetViewportLayoutEntity>>* const Viewports = ViewportTabContent->GetViewports();
 	if(Viewports)
 	{
 		const FString& LayoutId = ViewportTabContent->GetLayoutString();
 		for (auto& Pair : *Viewports)
 		{
-			TSharedPtr<SLevelViewport> Viewport = Pair.Value->AsLevelViewport();
+			TSharedPtr<SLevelViewport> Viewport = StaticCastSharedPtr<ILevelViewportLayoutEntity>(Pair.Value)->AsLevelViewport();
+
 			if( !Viewport.IsValid() )
 			{
 				continue;
@@ -997,13 +999,13 @@ void SLevelEditor::SaveViewportTabInfo(TSharedRef<const FLevelViewportTabContent
 
 void SLevelEditor::RestoreViewportTabInfo(TSharedRef<FLevelViewportTabContent> ViewportTabContent) const
 {
-	const TMap<FName, TSharedPtr<IViewportLayoutEntity>>* const Viewports = ViewportTabContent->GetViewports();
+	const TMap<FName, TSharedPtr<IAssetViewportLayoutEntity>>* const Viewports = ViewportTabContent->GetViewports();
 	if(Viewports)
 	{
 		const FString& LayoutId = ViewportTabContent->GetLayoutString();
 		for (auto& Pair : *Viewports)
 		{
-			TSharedPtr<SLevelViewport> Viewport = Pair.Value->AsLevelViewport();
+			TSharedPtr<SLevelViewport> Viewport = StaticCastSharedPtr<ILevelViewportLayoutEntity>(Pair.Value)->AsLevelViewport();
 			if( !Viewport.IsValid() )
 			{
 				continue;
@@ -1509,13 +1511,13 @@ TArray< TSharedPtr< ILevelViewport > > SLevelEditor::GetViewports() const
 		
 		if (ViewportTab.IsValid())
 		{
-			const TMap< FName, TSharedPtr< IViewportLayoutEntity > >* LevelViewports = ViewportTab->GetViewports();
+			const TMap< FName, TSharedPtr< IAssetViewportLayoutEntity > >* LevelViewports = ViewportTab->GetViewports();
 
 			if (LevelViewports != NULL)
 			{
 				for (auto& Pair : *LevelViewports)
 				{
-					TSharedPtr<SLevelViewport> Viewport = Pair.Value->AsLevelViewport();
+					TSharedPtr<SLevelViewport> Viewport = StaticCastSharedPtr<ILevelViewportLayoutEntity>(Pair.Value)->AsLevelViewport();
 					if( Viewport.IsValid() )
 					{
 						OutViewports.Add(Viewport);
diff --git a/Engine/Source/Editor/LevelEditor/Private/SLevelViewport.cpp b/Engine/Source/Editor/LevelEditor/Private/SLevelViewport.cpp
--- a/Engine/Source/Editor/LevelEditor/Private/SLevelViewport.cpp
+++ b/Engine/Source/Editor/LevelEditor/Private/SLevelViewport.cpp
@@ -3752,7 +3752,7 @@ FName SLevelViewport::GetViewportTypeWithinLayout() const
 	TSharedPtr<FLevelViewportLayout> LayoutPinned = ParentLayout.Pin();
 	if (LayoutPinned.IsValid() && !ConfigKey.IsNone())
 	{
-		TSharedPtr<IViewportLayoutEntity> Entity = LayoutPinned->GetViewports().FindRef(ConfigKey);
+		TSharedPtr<ILevelViewportLayoutEntity> Entity = StaticCastSharedPtr<ILevelViewportLayoutEntity>(LayoutPinned->GetViewports().FindRef(ConfigKey));
 		if (Entity.IsValid())
 		{
 			return Entity->GetType();
diff --git a/Engine/Source/Editor/LevelEditor/Public/LevelEditor.h b/Engine/Source/Editor/LevelEditor/Public/LevelEditor.h
--- a/Engine/Source/Editor/LevelEditor/Public/LevelEditor.h
+++ b/Engine/Source/Editor/LevelEditor/Public/LevelEditor.h
@@ -15,11 +15,10 @@
 
 class AActor;
 class ILevelViewport;
-class IViewportLayoutEntity;
 class SLevelEditor;
 class UAnimSequence;
 class USkeletalMeshComponent;
-struct FViewportConstructionArgs;
+
 enum class EMapChangeType : uint8;
 
 extern const FName LevelEditorApp;
@@ -329,7 +328,7 @@ public:
 	}
 
 	/** Create an instance of a custom viewport from the specified viewport type name */
-	TSharedRef<IViewportLayoutEntity> FactoryViewport(FName InTypeName, const FViewportConstructionArgs& ConstructionArgs) const;
+	TSharedRef<ILevelViewportLayoutEntity> FactoryViewport(FName InTypeName, const FViewportConstructionArgs& ConstructionArgs) const;
 
 private:
 	/**
diff --git a/Engine/Source/Editor/LevelEditor/Public/LevelViewportLayout.h b/Engine/Source/Editor/LevelEditor/Public/LevelViewportLayout.h
--- a/Engine/Source/Editor/LevelEditor/Public/LevelViewportLayout.h
+++ b/Engine/Source/Editor/LevelEditor/Public/LevelViewportLayout.h
@@ -11,7 +11,7 @@
 #include "Widgets/SWidget.h"
 #include "Editor/UnrealEdTypes.h"
 #include "Application/ThrottleManager.h"
-#include "TickableEditorObject.h"
+#include "AssetViewportLayout.h"
 
 class FLevelEditorViewportClient;
 class FLevelViewportLayout;
@@ -43,15 +43,13 @@ struct FViewportConstructionArgs
 	TAttribute<bool> IsEnabled;
 };
 
+
 /** Interface that defines an entity within a viewport layout */
-class IViewportLayoutEntity : public TSharedFromThis<IViewportLayoutEntity>
+class ILevelViewportLayoutEntity : public IAssetViewportLayoutEntity
 {
 public:
 	/** Virtual destruction */
-	virtual ~IViewportLayoutEntity() {};
-
-	/** Return a widget that is represents this entity */
-	virtual TSharedRef<SWidget> AsWidget() const = 0;
+	virtual ~ILevelViewportLayoutEntity() {};
 
 	/** Optionally return this entity as an SLevelViewport. Legacy function for interop with code that used to deal directly with SLevelViewports */
 	virtual TSharedPtr<SLevelViewport> AsLevelViewport() const { return nullptr; }
@@ -64,28 +62,13 @@ public:
 
 	/** Register this viewport layout entity as a game viewport, if it's currently PIE-ing */
 	virtual void RegisterGameViewportIfPIE() = 0;
-
-	/** Set keyboard focus to this viewport entity */
-	virtual void SetKeyboardFocus() = 0;
-
-	/** Called when the parent layout is being destroyed */
-	virtual void OnLayoutDestroyed() = 0;
-
-	/** Called to save this item's settings in the specified config section */
-	virtual void SaveConfig(const FString& ConfigSection) = 0;
-
-	/** Get the type of this viewport as a name */
-	virtual FName GetType() const = 0;
-
-	/** Take a high res screen shot of viewport entity */
-	virtual void TakeHighResScreenShot() const = 0;
 };
 
 /**
  * Base class for level viewport layout configurations
  * Handles maximizing and restoring well as visibility of specific viewports.
  */
-class LEVELEDITOR_API FLevelViewportLayout : public TSharedFromThis<FLevelViewportLayout>, public FTickableEditorObject
+class LEVELEDITOR_API FLevelViewportLayout : public TSharedFromThis<FLevelViewportLayout>, public FAssetViewportLayout, public FTickableEditorObject
 {
 public:
 	/**
@@ -155,17 +138,7 @@ public:
 	 */
 	bool IsViewportImmersive( FName InViewport ) const;
 
-	/**
-	 * @return All the viewports in this configuration
-	 */
-	const TMap< FName, TSharedPtr<IViewportLayoutEntity> >& GetViewports() const { return Viewports; }
-
-	/**
-	 * Saves viewport layout information between editor sessions
-	 */
-	virtual void SaveLayoutString(const FString& LayoutString) const = 0;
 
-	virtual const FName& GetLayoutTypeName() const = 0;
 
 
 	/** FTickableEditorObject interface */
@@ -269,8 +242,6 @@ protected:
 	/** Curve for animating from a "restored" state to a maximized state */
 	FCurveSequence MaximizeAnimation;
 
-	/** List of all of the viewports in this layout, keyed on their config key */
-	TMap< FName, TSharedPtr< IViewportLayoutEntity > > Viewports;
 	
 	/** The parent tab where this layout resides */
 	TWeakPtr< SDockTab > ParentTab;
diff --git a/Engine/Source/Editor/LevelEditor/Public/LevelViewportTabContent.h b/Engine/Source/Editor/LevelEditor/Public/LevelViewportTabContent.h
--- a/Engine/Source/Editor/LevelEditor/Public/LevelViewportTabContent.h
+++ b/Engine/Source/Editor/LevelEditor/Public/LevelViewportTabContent.h
@@ -4,6 +4,7 @@
 
 #include "CoreMinimal.h"
 #include "LevelViewportLayout.h"
+#include "ViewportTabContent.h"
 
 class ILevelEditor;
 
@@ -11,7 +12,7 @@ class ILevelEditor;
  * Represents the content in a viewport tab in the level editor.
  * Each SDockTab holding viewports in the level editor contains and owns one of these.
  */
-class LEVELEDITOR_API FLevelViewportTabContent : public TSharedFromThis<FLevelViewportTabContent>
+class LEVELEDITOR_API FLevelViewportTabContent : public FViewportTabContent
 {
 public:
 	/** Starts the tab content object and creates the initial layout based on the layout string */
@@ -21,7 +22,7 @@ public:
 	bool IsVisible() const;
 
 	/** Returns an array of viewports inside this tab */
-	const TMap< FName, TSharedPtr< IViewportLayoutEntity > >* GetViewports() const;
+	const TMap< FName, TSharedPtr< IAssetViewportLayoutEntity > >* GetViewports() const;
 
 	/**
 	 * Sets the current layout by changing the contained layout object
@@ -72,12 +73,6 @@ private:
 	static TSharedPtr< class FLevelViewportLayout > ConstructViewportLayoutByTypeName(const FName& TypeName, bool bSwitchingLayouts);
 
 private:
-	TWeakPtr<class SDockTab> ParentTab;
 	TWeakPtr<ILevelEditor> ParentLevelEditor;
-	FString LayoutString;
 
-	TOptional<FName> PreviouslyFocusedViewport;
-
-	/** Current layout */
-	TSharedPtr< class FLevelViewportLayout > ActiveLevelViewportLayout;
 };
diff --git a/Engine/Source/Editor/LevelEditor/Public/ViewportTypeDefinition.h b/Engine/Source/Editor/LevelEditor/Public/ViewportTypeDefinition.h
--- a/Engine/Source/Editor/LevelEditor/Public/ViewportTypeDefinition.h
+++ b/Engine/Source/Editor/LevelEditor/Public/ViewportTypeDefinition.h
@@ -5,18 +5,18 @@
 #include "CoreMinimal.h"
 
 class FUICommandInfo;
-class IViewportLayoutEntity;
+class ILevelViewportLayoutEntity;
 struct FViewportConstructionArgs;
 
 /** Definition of a custom viewport */
 struct FViewportTypeDefinition
 {
-	typedef TFunction<TSharedRef<IViewportLayoutEntity>(const FViewportConstructionArgs&)> FFactoryFunctionType;
+	typedef TFunction<TSharedRef<ILevelViewportLayoutEntity>(const FViewportConstructionArgs&)> FFactoryFunctionType;
 
 	template<typename T>
 	static FViewportTypeDefinition FromType(const TSharedPtr<FUICommandInfo>& ActivationCommand)
 	{
-		return FViewportTypeDefinition([](const FViewportConstructionArgs& Args) -> TSharedRef<IViewportLayoutEntity> {
+		return FViewportTypeDefinition([](const FViewportConstructionArgs& Args) -> TSharedRef<ILevelViewportLayoutEntity> {
 			return MakeShareable(new T(Args));
 		}, ActivationCommand);
 	}
diff --git a/Engine/Source/Editor/UnrealEd/Private/ViewportTabContent.cpp b/Engine/Source/Editor/UnrealEd/Private/ViewportTabContent.cpp
new file mode 100644
--- /dev/null
+++ b/Engine/Source/Editor/UnrealEd/Private/ViewportTabContent.cpp
@@ -0,0 +1,16 @@
+// Copyright 1998-2019 Epic Games, Inc. All Rights Reserved.
+
+#include "ViewportTabContent.h"
+#include "Misc/ConfigCacheIni.h"
+#include "Framework/Docking/LayoutService.h"
+#include "Modules/ModuleManager.h"
+#include "Widgets/Docking/SDockTab.h"
+
+
+// FViewportTabContent ///////////////////////////
+
+bool FViewportTabContent::BelongsToTab(TSharedRef<class SDockTab> InParentTab) const
+{
+	TSharedPtr<SDockTab> ParentTabPinned = ParentTab.Pin();
+	return ParentTabPinned == InParentTab;
+}
diff --git a/Engine/Source/Editor/UnrealEd/Public/AssetViewportLayout.h b/Engine/Source/Editor/UnrealEd/Public/AssetViewportLayout.h
new file mode 100644
--- /dev/null
+++ b/Engine/Source/Editor/UnrealEd/Public/AssetViewportLayout.h
@@ -0,0 +1,56 @@
+// Copyright 1998-2019 Epic Games, Inc. All Rights Reserved.
+
+#pragma once
+
+#include "CoreMinimal.h"
+
+
+class IAssetViewportLayoutEntity : public TSharedFromThis<IAssetViewportLayoutEntity>
+{
+public:
+	/** Virtual destruction */
+	virtual ~IAssetViewportLayoutEntity() {};
+
+	/** Return a widget that is represents this entity */
+	virtual TSharedRef<SWidget> AsWidget() const = 0;
+
+	/** Set keyboard focus to this viewport entity */
+	virtual void SetKeyboardFocus() = 0;
+
+	/** Called when the parent layout is being destroyed */
+	virtual void OnLayoutDestroyed() = 0;
+
+	/** Called to save this item's settings in the specified config section */
+	virtual void SaveConfig(const FString& ConfigSection) = 0;
+
+	/** Get the type of this viewport as a name */
+	virtual FName GetType() const = 0;
+
+	/** Take a high res screen shot of viewport entity */
+	virtual void TakeHighResScreenShot() const = 0;
+};
+
+class FAssetViewportLayout 
+{
+public:
+	/** Virtual destruction */
+	virtual ~FAssetViewportLayout() {};
+
+	/**
+	* @return All the viewports in this configuration
+	*/
+	const TMap< FName, TSharedPtr<IAssetViewportLayoutEntity> >& GetViewports() const { return Viewports; }
+
+	/**
+	* Saves viewport layout information between editor sessions
+	*/
+	virtual void SaveLayoutString(const FString& LayoutString) const = 0;
+
+	virtual const FName& GetLayoutTypeName() const = 0;
+
+protected:
+	/** List of all of the viewports in this layout, keyed on their config key */
+	TMap< FName, TSharedPtr< IAssetViewportLayoutEntity > > Viewports;
+
+
+};
diff --git a/Engine/Source/Editor/UnrealEd/Public/ViewportTabContent.h b/Engine/Source/Editor/UnrealEd/Public/ViewportTabContent.h
new file mode 100644
--- /dev/null
+++ b/Engine/Source/Editor/UnrealEd/Public/ViewportTabContent.h
@@ -0,0 +1,34 @@
+// Copyright 1998-2019 Epic Games, Inc. All Rights Reserved.
+
+#pragma once
+
+#include "CoreMinimal.h"
+
+/**
+ * Represents the content in a viewport tab in an editor.
+ * Each SDockTab holding viewports in an editor contains and owns one of these.
+ */
+class UNREALED_API FViewportTabContent : public TSharedFromThis<FViewportTabContent>
+{
+public:
+	/** Returns whether the tab is currently shown */
+	bool IsVisible() const;
+
+	/** @return True if this viewport belongs to the tab given */
+	bool BelongsToTab(TSharedRef<class SDockTab> InParentTab) const;
+
+	/** @return The string used to identify the layout of this viewport tab */
+	const FString& GetLayoutString() const
+	{
+		return LayoutString;
+	}
+
+protected:
+	TWeakPtr<class SDockTab> ParentTab;
+	FString LayoutString;
+
+	/** Current layout */
+	TSharedPtr< class FAssetViewportLayout > ActiveViewportLayout;
+
+	TOptional<FName> PreviouslyFocusedViewport;
+};
