Add missing header.  #rb trivial

diff --git a/Engine/Source/Editor/LevelEditor/Public/LevelViewportLayout.h b/Engine/Source/Editor/LevelEditor/Public/LevelViewportLayout.h
--- a/Engine/Source/Editor/LevelEditor/Public/LevelViewportLayout.h
+++ b/Engine/Source/Editor/LevelEditor/Public/LevelViewportLayout.h
@@ -12,6 +12,7 @@
 #include "Editor/UnrealEdTypes.h"
 #include "Application/ThrottleManager.h"
 #include "AssetViewportLayout.h"
+#include "TickableEditorObject.h"
 
 class FLevelEditorViewportClient;
 class FLevelViewportLayout;
