Fix missing source control error status icon on level editor toolbar

diff --git a/Engine/Source/Editor/LevelEditor/Private/LevelEditorToolBar.cpp b/Engine/Source/Editor/LevelEditor/Private/LevelEditorToolBar.cpp
--- a/Engine/Source/Editor/LevelEditor/Private/LevelEditorToolBar.cpp
+++ b/Engine/Source/Editor/LevelEditor/Private/LevelEditorToolBar.cpp
@@ -1246,7 +1246,7 @@ void FLevelEditorToolBar::RegisterLevelEditorToolBar( const TSharedRef<FUIComman
 						{
 							if (!SourceControlModule.GetProvider().IsAvailable())
 							{
-								return FSlateIcon(FEditorStyle::GetStyleSetName(), "LevelEditor.SourceControl.Error");
+								return FSlateIcon(FEditorStyle::GetStyleSetName(), "LevelEditor.SourceControl.Problem");
 							}
 							else
 							{
