Rename/move file(s)
Moves AssetLayout to EditorLayout

diff --git a/Engine/Config/BaseEditor.ini b/Engine/Config/BaseEditor.ini
--- a/Engine/Config/BaseEditor.ini
+++ b/Engine/Config/BaseEditor.ini
@@ -536,4 +536,4 @@ P4EngineBasePath="/Engine/Content/"
 P4EngineAssetPath="/Engine/"
 
 [CoreRedirects]
-+ClassRedirects=(OldName="IViewportLayoutEntity", NewName="IAssetViewportLayoutEntity")
++ClassRedirects=(OldName="IViewportLayoutEntity", NewName="IEditorViewportLayoutEntity")
diff --git a/Engine/Source/Editor/LevelEditor/Private/LevelViewportLayoutFourPanes.cpp b/Engine/Source/Editor/LevelEditor/Private/LevelViewportLayoutFourPanes.cpp
--- a/Engine/Source/Editor/LevelEditor/Private/LevelViewportLayoutFourPanes.cpp
+++ b/Engine/Source/Editor/LevelEditor/Private/LevelViewportLayoutFourPanes.cpp
@@ -169,7 +169,7 @@ void FLevelViewportLayoutFourPanes::ReplaceWidget( TSharedRef< SWidget > Source,
 // FLevelViewportLayoutFourPanesLeft /////////////////////////////
 
 TSharedRef<SWidget> FLevelViewportLayoutFourPanesLeft::MakeFourPanelWidget(
-	TMap<FName, TSharedPtr< IAssetViewportLayoutEntity >>& ViewportWidgets,
+	TMap<FName, TSharedPtr< IEditorViewportLayoutEntity >>& ViewportWidgets,
 	TSharedRef<SWidget> Viewport0, TSharedRef<SWidget> Viewport1, TSharedRef<SWidget> Viewport2, TSharedRef<SWidget> Viewport3,
 	float PrimarySplitterPercentage, float SecondarySplitterPercentage0, float SecondarySplitterPercentage1)
 {
@@ -210,7 +210,7 @@ TSharedRef<SWidget> FLevelViewportLayoutFourPanesLeft::MakeFourPanelWidget(
 // FLevelViewportLayoutFourPanesRight /////////////////////////////
 
 TSharedRef<SWidget> FLevelViewportLayoutFourPanesRight::MakeFourPanelWidget(
-	TMap<FName, TSharedPtr< IAssetViewportLayoutEntity >>& ViewportWidgets,
+	TMap<FName, TSharedPtr< IEditorViewportLayoutEntity >>& ViewportWidgets,
 	TSharedRef<SWidget> Viewport0, TSharedRef<SWidget> Viewport1, TSharedRef<SWidget> Viewport2, TSharedRef<SWidget> Viewport3,
 	float PrimarySplitterPercentage, float SecondarySplitterPercentage0, float SecondarySplitterPercentage1)
 {
@@ -251,7 +251,7 @@ TSharedRef<SWidget> FLevelViewportLayoutFourPanesRight::MakeFourPanelWidget(
 // FLevelViewportLayoutFourPanesTop /////////////////////////////
 
 TSharedRef<SWidget> FLevelViewportLayoutFourPanesTop::MakeFourPanelWidget(
-	TMap<FName, TSharedPtr< IAssetViewportLayoutEntity >>& ViewportWidgets,
+	TMap<FName, TSharedPtr< IEditorViewportLayoutEntity >>& ViewportWidgets,
 	TSharedRef<SWidget> Viewport0, TSharedRef<SWidget> Viewport1, TSharedRef<SWidget> Viewport2, TSharedRef<SWidget> Viewport3,
 	float PrimarySplitterPercentage, float SecondarySplitterPercentage0, float SecondarySplitterPercentage1)
 {
@@ -292,7 +292,7 @@ TSharedRef<SWidget> FLevelViewportLayoutFourPanesTop::MakeFourPanelWidget(
 // FLevelViewportLayoutFourPanesBottom /////////////////////////////
 
 TSharedRef<SWidget> FLevelViewportLayoutFourPanesBottom::MakeFourPanelWidget(
-	TMap<FName, TSharedPtr< IAssetViewportLayoutEntity >>& ViewportWidgets,
+	TMap<FName, TSharedPtr< IEditorViewportLayoutEntity >>& ViewportWidgets,
 	TSharedRef<SWidget> Viewport0, TSharedRef<SWidget> Viewport1, TSharedRef<SWidget> Viewport2, TSharedRef<SWidget> Viewport3,
 	float PrimarySplitterPercentage, float SecondarySplitterPercentage0, float SecondarySplitterPercentage1)
 {
diff --git a/Engine/Source/Editor/LevelEditor/Private/LevelViewportLayoutFourPanes.h b/Engine/Source/Editor/LevelEditor/Private/LevelViewportLayoutFourPanes.h
--- a/Engine/Source/Editor/LevelEditor/Private/LevelViewportLayoutFourPanes.h
+++ b/Engine/Source/Editor/LevelEditor/Private/LevelViewportLayoutFourPanes.h
@@ -22,7 +22,7 @@ protected:
 	virtual TSharedRef<SWidget> MakeViewportLayout(const FString& LayoutString) override;
 
 	virtual TSharedRef<SWidget> MakeFourPanelWidget(
-		TMap<FName, TSharedPtr< IAssetViewportLayoutEntity >>& ViewportWidgets,
+		TMap<FName, TSharedPtr< IEditorViewportLayoutEntity >>& ViewportWidgets,
 		TSharedRef<SWidget> Viewport0, TSharedRef<SWidget> Viewport1, TSharedRef<SWidget> Viewport2, TSharedRef<SWidget> Viewport3,
 		float PrimarySplitterPercentage, float SecondarySplitterPercentage0, float SecondarySplitterPercentage1) = 0;
 
@@ -45,7 +45,7 @@ public:
 	virtual const FName& GetLayoutTypeName() const override { return LevelViewportConfigurationNames::FourPanesLeft; }
 
 	virtual TSharedRef<SWidget> MakeFourPanelWidget(
-		TMap<FName, TSharedPtr< IAssetViewportLayoutEntity >>& ViewportWidgets,
+		TMap<FName, TSharedPtr< IEditorViewportLayoutEntity >>& ViewportWidgets,
 		TSharedRef<SWidget> Viewport0, TSharedRef<SWidget> Viewport1, TSharedRef<SWidget> Viewport2, TSharedRef<SWidget> Viewport3,
 		float PrimarySplitterPercentage, float SecondarySplitterPercentage0, float SecondarySplitterPercentage1) override;
 };
@@ -59,7 +59,7 @@ public:
 	virtual const FName& GetLayoutTypeName() const override { return LevelViewportConfigurationNames::FourPanesRight; }
 
 	virtual TSharedRef<SWidget> MakeFourPanelWidget(
-		TMap<FName, TSharedPtr< IAssetViewportLayoutEntity >>& ViewportWidgets,
+		TMap<FName, TSharedPtr< IEditorViewportLayoutEntity >>& ViewportWidgets,
 		TSharedRef<SWidget> Viewport0, TSharedRef<SWidget> Viewport1, TSharedRef<SWidget> Viewport2, TSharedRef<SWidget> Viewport3,
 		float PrimarySplitterPercentage, float SecondarySplitterPercentage0, float SecondarySplitterPercentage1) override;
 };
@@ -73,7 +73,7 @@ public:
 	virtual const FName& GetLayoutTypeName() const override { return LevelViewportConfigurationNames::FourPanesTop; }
 
 	virtual TSharedRef<SWidget> MakeFourPanelWidget(
-		TMap<FName, TSharedPtr< IAssetViewportLayoutEntity >>& ViewportWidgets,
+		TMap<FName, TSharedPtr< IEditorViewportLayoutEntity >>& ViewportWidgets,
 		TSharedRef<SWidget> Viewport0, TSharedRef<SWidget> Viewport1, TSharedRef<SWidget> Viewport2, TSharedRef<SWidget> Viewport3,
 		float PrimarySplitterPercentage, float SecondarySplitterPercentage0, float SecondarySplitterPercentage1) override;
 };
@@ -87,7 +87,7 @@ public:
 	virtual const FName& GetLayoutTypeName() const override { return LevelViewportConfigurationNames::FourPanesBottom; }
 
 	virtual TSharedRef<SWidget> MakeFourPanelWidget(
-		TMap<FName, TSharedPtr< IAssetViewportLayoutEntity >>& ViewportWidgets,
+		TMap<FName, TSharedPtr< IEditorViewportLayoutEntity >>& ViewportWidgets,
 		TSharedRef<SWidget> Viewport0, TSharedRef<SWidget> Viewport1, TSharedRef<SWidget> Viewport2, TSharedRef<SWidget> Viewport3,
 		float PrimarySplitterPercentage, float SecondarySplitterPercentage0, float SecondarySplitterPercentage1) override;
 };
diff --git a/Engine/Source/Editor/LevelEditor/Private/LevelViewportLayoutThreePanes.cpp b/Engine/Source/Editor/LevelEditor/Private/LevelViewportLayoutThreePanes.cpp
--- a/Engine/Source/Editor/LevelEditor/Private/LevelViewportLayoutThreePanes.cpp
+++ b/Engine/Source/Editor/LevelEditor/Private/LevelViewportLayoutThreePanes.cpp
@@ -152,7 +152,7 @@ void FLevelViewportLayoutThreePanes::ReplaceWidget( TSharedRef< SWidget > Source
 // FLevelViewportLayoutThreePanesLeft /////////////////////////////
 
 TSharedRef<SWidget> FLevelViewportLayoutThreePanesLeft::MakeThreePanelWidget(
-	TMap< FName, TSharedPtr< IAssetViewportLayoutEntity >>& ViewportWidgets,
+	TMap< FName, TSharedPtr< IEditorViewportLayoutEntity >>& ViewportWidgets,
 	const TSharedRef<SWidget>& Viewport0, const TSharedRef<SWidget>& Viewport1, const TSharedRef<SWidget>& Viewport2,
 	float PrimarySplitterPercentage, float SecondarySplitterPercentage)
 {
@@ -188,7 +188,7 @@ TSharedRef<SWidget> FLevelViewportLayoutThreePanesLeft::MakeThreePanelWidget(
 // FLevelViewportLayoutThreePanesRight /////////////////////////////
 
 TSharedRef<SWidget> FLevelViewportLayoutThreePanesRight::MakeThreePanelWidget(
-	TMap< FName, TSharedPtr< IAssetViewportLayoutEntity >>& ViewportWidgets,
+	TMap< FName, TSharedPtr< IEditorViewportLayoutEntity >>& ViewportWidgets,
 	const TSharedRef<SWidget>& Viewport0, const TSharedRef<SWidget>& Viewport1, const TSharedRef<SWidget>& Viewport2,
 	float PrimarySplitterPercentage, float SecondarySplitterPercentage)
 {
@@ -224,7 +224,7 @@ TSharedRef<SWidget> FLevelViewportLayoutThreePanesRight::MakeThreePanelWidget(
 // FLevelViewportLayoutThreePanesTop /////////////////////////////
 
 TSharedRef<SWidget> FLevelViewportLayoutThreePanesTop::MakeThreePanelWidget(
-	TMap< FName, TSharedPtr< IAssetViewportLayoutEntity >>& ViewportWidgets,
+	TMap< FName, TSharedPtr< IEditorViewportLayoutEntity >>& ViewportWidgets,
 	const TSharedRef<SWidget>& Viewport0, const TSharedRef<SWidget>& Viewport1, const TSharedRef<SWidget>& Viewport2,
 	float PrimarySplitterPercentage, float SecondarySplitterPercentage)
 {
@@ -260,7 +260,7 @@ TSharedRef<SWidget> FLevelViewportLayoutThreePanesTop::MakeThreePanelWidget(
 // FLevelViewportLayoutThreePanesBottom /////////////////////////////
 
 TSharedRef<SWidget> FLevelViewportLayoutThreePanesBottom::MakeThreePanelWidget(
-	TMap< FName, TSharedPtr< IAssetViewportLayoutEntity >>& ViewportWidgets,
+	TMap< FName, TSharedPtr< IEditorViewportLayoutEntity >>& ViewportWidgets,
 	const TSharedRef<SWidget>& Viewport0, const TSharedRef<SWidget>& Viewport1, const TSharedRef<SWidget>& Viewport2,
 	float PrimarySplitterPercentage, float SecondarySplitterPercentage)
 {
diff --git a/Engine/Source/Editor/LevelEditor/Private/LevelViewportLayoutThreePanes.h b/Engine/Source/Editor/LevelEditor/Private/LevelViewportLayoutThreePanes.h
--- a/Engine/Source/Editor/LevelEditor/Private/LevelViewportLayoutThreePanes.h
+++ b/Engine/Source/Editor/LevelEditor/Private/LevelViewportLayoutThreePanes.h
@@ -22,7 +22,7 @@ protected:
 	virtual TSharedRef<SWidget> MakeViewportLayout(const FString& LayoutString) override;
 
 	virtual TSharedRef<SWidget> MakeThreePanelWidget(
-		TMap< FName, TSharedPtr< IAssetViewportLayoutEntity >>& ViewportWidgets,
+		TMap< FName, TSharedPtr< IEditorViewportLayoutEntity >>& ViewportWidgets,
 		const TSharedRef<SWidget>& ViewportKey0, const TSharedRef<SWidget>& ViewportKey1, const TSharedRef<SWidget>& ViewportKey2,
 		float PrimarySplitterPercentage, float SecondarySplitterPercentage) = 0;
 
@@ -44,7 +44,7 @@ public:
 	virtual const FName& GetLayoutTypeName() const override { return LevelViewportConfigurationNames::ThreePanesLeft; }
 
 	virtual TSharedRef<SWidget> MakeThreePanelWidget(
-		TMap< FName, TSharedPtr< IAssetViewportLayoutEntity >>& ViewportWidgets,
+		TMap< FName, TSharedPtr< IEditorViewportLayoutEntity >>& ViewportWidgets,
 		const TSharedRef<SWidget>& ViewportKey0, const TSharedRef<SWidget>& ViewportKey1, const TSharedRef<SWidget>& ViewportKey2,
 		float PrimarySplitterPercentage, float SecondarySplitterPercentage) override;
 };
@@ -58,7 +58,7 @@ public:
 	virtual const FName& GetLayoutTypeName() const override { return LevelViewportConfigurationNames::ThreePanesRight; }
 
 	virtual TSharedRef<SWidget> MakeThreePanelWidget(
-		TMap< FName, TSharedPtr< IAssetViewportLayoutEntity >>& ViewportWidgets,
+		TMap< FName, TSharedPtr< IEditorViewportLayoutEntity >>& ViewportWidgets,
 		const TSharedRef<SWidget>& ViewportKey0, const TSharedRef<SWidget>& ViewportKey1, const TSharedRef<SWidget>& ViewportKey2,
 		float PrimarySplitterPercentage, float SecondarySplitterPercentage) override;
 };
@@ -72,7 +72,7 @@ public:
 	virtual const FName& GetLayoutTypeName() const override { return LevelViewportConfigurationNames::ThreePanesTop; }
 
 	virtual TSharedRef<SWidget> MakeThreePanelWidget(
-		TMap< FName, TSharedPtr< IAssetViewportLayoutEntity >>& ViewportWidgets,
+		TMap< FName, TSharedPtr< IEditorViewportLayoutEntity >>& ViewportWidgets,
 		const TSharedRef<SWidget>& ViewportKey0, const TSharedRef<SWidget>& ViewportKey1, const TSharedRef<SWidget>& ViewportKey2,
 		float PrimarySplitterPercentage, float SecondarySplitterPercentage) override;
 };
@@ -86,7 +86,7 @@ public:
 	virtual const FName& GetLayoutTypeName() const override { return LevelViewportConfigurationNames::ThreePanesBottom; }
 
 	virtual TSharedRef<SWidget> MakeThreePanelWidget(
-		TMap< FName, TSharedPtr< IAssetViewportLayoutEntity >>& ViewportWidgets,
+		TMap< FName, TSharedPtr< IEditorViewportLayoutEntity >>& ViewportWidgets,
 		const TSharedRef<SWidget>& ViewportKey0, const TSharedRef<SWidget>& ViewportKey1, const TSharedRef<SWidget>& ViewportKey2,
 		float PrimarySplitterPercentage, float SecondarySplitterPercentage) override;
 };
diff --git a/Engine/Source/Editor/LevelEditor/Private/LevelViewportTabContent.cpp b/Engine/Source/Editor/LevelEditor/Private/LevelViewportTabContent.cpp
--- a/Engine/Source/Editor/LevelEditor/Private/LevelViewportTabContent.cpp
+++ b/Engine/Source/Editor/LevelEditor/Private/LevelViewportTabContent.cpp
@@ -73,7 +73,7 @@ bool FLevelViewportTabContent::IsVisible() const
 	return false;
 }
 
-const TMap< FName, TSharedPtr< IAssetViewportLayoutEntity > >* FLevelViewportTabContent::GetViewports() const
+const TMap< FName, TSharedPtr< IEditorViewportLayoutEntity > >* FLevelViewportTabContent::GetViewports() const
 {
 	if (ActiveViewportLayout.IsValid())
 	{
@@ -164,7 +164,7 @@ void FLevelViewportTabContent::UpdateViewportTabWidget()
 
 		if (PreviouslyFocusedViewport.IsSet())
 		{
-			TSharedPtr<IAssetViewportLayoutEntity> ViewportToFocus = ActiveViewportLayout->GetViewports().FindRef(PreviouslyFocusedViewport.GetValue());
+			TSharedPtr<IEditorViewportLayoutEntity> ViewportToFocus = ActiveViewportLayout->GetViewports().FindRef(PreviouslyFocusedViewport.GetValue());
 			if (ViewportToFocus.IsValid())
 			{
 				ViewportToFocus->SetKeyboardFocus();
diff --git a/Engine/Source/Editor/LevelEditor/Private/SLevelEditor.cpp b/Engine/Source/Editor/LevelEditor/Private/SLevelEditor.cpp
--- a/Engine/Source/Editor/LevelEditor/Private/SLevelEditor.cpp
+++ b/Engine/Source/Editor/LevelEditor/Private/SLevelEditor.cpp
@@ -344,7 +344,7 @@ bool SLevelEditor::HasActivePlayInEditorViewport() const
 		if (ViewportTab.IsValid())
 		{
 			// Get all the viewports in the layout
-			const TMap< FName, TSharedPtr< IAssetViewportLayoutEntity > >* LevelViewports = ViewportTab.Pin()->GetViewports();
+			const TMap< FName, TSharedPtr< IEditorViewportLayoutEntity > >* LevelViewports = ViewportTab.Pin()->GetViewports();
 
 			if (LevelViewports != NULL)
 			{
@@ -393,7 +393,7 @@ TSharedPtr<SLevelViewport> SLevelEditor::GetActiveViewport()
 			// Only check the viewports in the tab if its visible
 			if( ViewportTab->IsVisible() )
 			{
-				const TMap< FName, TSharedPtr< IAssetViewportLayoutEntity > >* LevelViewports = ViewportTab->GetViewports();
+				const TMap< FName, TSharedPtr< IEditorViewportLayoutEntity > >* LevelViewports = ViewportTab->GetViewports();
 
 				if (LevelViewports != nullptr)
 				{
@@ -970,7 +970,7 @@ void SLevelEditor::OnViewportTabClosed(TSharedRef<SDockTab> ClosedTab)
 
 void SLevelEditor::SaveViewportTabInfo(TSharedRef<const FLevelViewportTabContent> ViewportTabContent)
 {
-	const TMap<FName, TSharedPtr<IAssetViewportLayoutEntity>>* const Viewports = ViewportTabContent->GetViewports();
+	const TMap<FName, TSharedPtr<IEditorViewportLayoutEntity>>* const Viewports = ViewportTabContent->GetViewports();
 	if(Viewports)
 	{
 		const FString& LayoutId = ViewportTabContent->GetLayoutString();
@@ -999,7 +999,7 @@ void SLevelEditor::SaveViewportTabInfo(TSharedRef<const FLevelViewportTabContent
 
 void SLevelEditor::RestoreViewportTabInfo(TSharedRef<FLevelViewportTabContent> ViewportTabContent) const
 {
-	const TMap<FName, TSharedPtr<IAssetViewportLayoutEntity>>* const Viewports = ViewportTabContent->GetViewports();
+	const TMap<FName, TSharedPtr<IEditorViewportLayoutEntity>>* const Viewports = ViewportTabContent->GetViewports();
 	if(Viewports)
 	{
 		const FString& LayoutId = ViewportTabContent->GetLayoutString();
@@ -1511,7 +1511,7 @@ TArray< TSharedPtr< IAssetViewport > > SLevelEditor::GetViewports() const
 		
 		if (ViewportTab.IsValid())
 		{
-			const TMap< FName, TSharedPtr< IAssetViewportLayoutEntity > >* LevelViewports = ViewportTab->GetViewports();
+			const TMap< FName, TSharedPtr< IEditorViewportLayoutEntity > >* LevelViewports = ViewportTab->GetViewports();
 
 			if (LevelViewports != NULL)
 			{
diff --git a/Engine/Source/Editor/LevelEditor/Public/LevelViewportLayout.h b/Engine/Source/Editor/LevelEditor/Public/LevelViewportLayout.h
--- a/Engine/Source/Editor/LevelEditor/Public/LevelViewportLayout.h
+++ b/Engine/Source/Editor/LevelEditor/Public/LevelViewportLayout.h
@@ -11,7 +11,7 @@
 #include "Widgets/SWidget.h"
 #include "Editor/UnrealEdTypes.h"
 #include "Application/ThrottleManager.h"
-#include "AssetViewportLayout.h"
+#include "EditorViewportLayout.h"
 #include "TickableEditorObject.h"
 
 class FLevelEditorViewportClient;
@@ -46,7 +46,7 @@ struct FViewportConstructionArgs
 
 
 /** Interface that defines an entity within a viewport layout */
-class ILevelViewportLayoutEntity : public IAssetViewportLayoutEntity
+class ILevelViewportLayoutEntity : public IEditorViewportLayoutEntity
 {
 public:
 	/** Virtual destruction */
@@ -69,7 +69,7 @@ public:
  * Base class for level viewport layout configurations
  * Handles maximizing and restoring well as visibility of specific viewports.
  */
-class LEVELEDITOR_API FLevelViewportLayout : public TSharedFromThis<FLevelViewportLayout>, public FAssetViewportLayout, public FTickableEditorObject
+class LEVELEDITOR_API FLevelViewportLayout : public TSharedFromThis<FLevelViewportLayout>, public FEditorViewportLayout, public FTickableEditorObject
 {
 public:
 	/**
diff --git a/Engine/Source/Editor/LevelEditor/Public/LevelViewportTabContent.h b/Engine/Source/Editor/LevelEditor/Public/LevelViewportTabContent.h
--- a/Engine/Source/Editor/LevelEditor/Public/LevelViewportTabContent.h
+++ b/Engine/Source/Editor/LevelEditor/Public/LevelViewportTabContent.h
@@ -22,7 +22,7 @@ public:
 	bool IsVisible() const;
 
 	/** Returns an array of viewports inside this tab */
-	const TMap< FName, TSharedPtr< IAssetViewportLayoutEntity > >* GetViewports() const;
+	const TMap< FName, TSharedPtr< IEditorViewportLayoutEntity > >* GetViewports() const;
 
 	/**
 	 * Sets the current layout by changing the contained layout object
diff --git a/Engine/Source/Editor/UnrealEd/Public/AssetViewportLayout.h b/Engine/Source/Editor/UnrealEd/Public/AssetViewportLayout.h
deleted file mode 100644
diff --git a/Engine/Source/Editor/UnrealEd/Public/EditorViewportLayout.h b/Engine/Source/Editor/UnrealEd/Public/EditorViewportLayout.h
new file mode 100644
--- /dev/null
+++ b/Engine/Source/Editor/UnrealEd/Public/EditorViewportLayout.h
@@ -0,0 +1,56 @@
+// Copyright 1998-2019 Epic Games, Inc. All Rights Reserved.
+
+#pragma once
+
+#include "CoreMinimal.h"
+
+
+class IEditorViewportLayoutEntity : public TSharedFromThis<IEditorViewportLayoutEntity>
+{
+public:
+	/** Virtual destruction */
+	virtual ~IEditorViewportLayoutEntity() {};
+
+	/** Return a widget that is represents this entity */
+	virtual TSharedRef<SWidget> AsWidget() const = 0;
+
+	/** Set keyboard focus to this viewport entity */
+	virtual void SetKeyboardFocus() = 0;
+
+	/** Called when the parent layout is being destroyed */
+	virtual void OnLayoutDestroyed() = 0;
+
+	/** Called to save this item's settings in the specified config section */
+	virtual void SaveConfig(const FString& ConfigSection) = 0;
+
+	/** Get the type of this viewport as a name */
+	virtual FName GetType() const = 0;
+
+	/** Take a high res screen shot of viewport entity */
+	virtual void TakeHighResScreenShot() const = 0;
+};
+
+class FEditorViewportLayout 
+{
+public:
+	/** Virtual destruction */
+	virtual ~FEditorViewportLayout() {};
+
+	/**
+	* @return All the viewports in this configuration
+	*/
+	const TMap< FName, TSharedPtr<IEditorViewportLayoutEntity> >& GetViewports() const { return Viewports; }
+
+	/**
+	* Saves viewport layout information between editor sessions
+	*/
+	virtual void SaveLayoutString(const FString& LayoutString) const = 0;
+
+	virtual const FName& GetLayoutTypeName() const = 0;
+
+protected:
+	/** List of all of the viewports in this layout, keyed on their config key */
+	TMap< FName, TSharedPtr< IEditorViewportLayoutEntity > > Viewports;
+
+
+};
diff --git a/Engine/Source/Editor/UnrealEd/Public/ViewportTabContent.h b/Engine/Source/Editor/UnrealEd/Public/ViewportTabContent.h
--- a/Engine/Source/Editor/UnrealEd/Public/ViewportTabContent.h
+++ b/Engine/Source/Editor/UnrealEd/Public/ViewportTabContent.h
@@ -28,7 +28,7 @@ protected:
 	FString LayoutString;
 
 	/** Current layout */
-	TSharedPtr< class FAssetViewportLayout > ActiveViewportLayout;
+	TSharedPtr< class FEditorViewportLayout > ActiveViewportLayout;
 
 	TOptional<FName> PreviouslyFocusedViewport;
 };
