Convert level editor tool bar menus into UToolMenus
#rnx

diff --git a/Engine/Source/Editor/LevelEditor/Private/LevelEditorToolBar.cpp b/Engine/Source/Editor/LevelEditor/Private/LevelEditorToolBar.cpp
--- a/Engine/Source/Editor/LevelEditor/Private/LevelEditorToolBar.cpp
+++ b/Engine/Source/Editor/LevelEditor/Private/LevelEditorToolBar.cpp
@@ -303,11 +303,11 @@ namespace LevelEditorActionHelpers
 	/**
 	 * Builds a sub-menu for selecting a class
 	 *
-	 * @param InMenuBuilder		Object to append menu items/widgets to
+	 * @param InMenu		Object to append menu items/widgets to
 	 * @param InRootClass		The root class to filter the Class Viewer by to only show children of
 	 * @param InOnClassPicked	Callback delegate to fire when a class is picked
 	 */
-	void GetSelectSettingsClassSubMenu(FMenuBuilder& InMenuBuilder, UClass* InRootClass, FOnClassPicked InOnClassPicked)
+	void GetSelectSettingsClassSubMenu(UToolMenu* InMenu, UClass* InRootClass, FOnClassPicked InOnClassPicked)
 	{
 		FClassViewerInitializationOptions Options;
 		Options.Mode = EClassViewerMode::ClassPicker;
@@ -329,19 +329,18 @@ namespace LevelEditorActionHelpers
 		TSharedRef<SWidget> ClassViewer = FModuleManager::LoadModuleChecked<FClassViewerModule>("ClassViewer").CreateClassViewer(Options, InOnClassPicked);
 		FFormatNamedArguments FormatArgs;
 		FormatArgs.Add(TEXT("RootClass"), RootClassName);
-		InMenuBuilder.BeginSection(NAME_None, FText::Format(NSLOCTEXT("LevelToolBarViewMenu", "SelectGameModeLabel", "Select {RootClass} class"), FormatArgs));
-		InMenuBuilder.AddWidget(ClassViewer, FText::GetEmpty(), true);
-		InMenuBuilder.EndSection();
+		FToolMenuSection& Section = InMenu->AddSection("SelectSettingsClass", FText::Format(NSLOCTEXT("LevelToolBarViewMenu", "SelectGameModeLabel", "Select {RootClass} class"), FormatArgs));
+		Section.AddEntry(FToolMenuEntry::InitWidget("ClassViewer", ClassViewer, FText::GetEmpty(), true));
 	}
 
 	/**
 	 * Builds a sub-menu for creating a class
 	 *
-	 * @param InMenuBuilder		Object to append menu items/widgets to
+	 * @param InMenu		Object to append menu items/widgets to
 	 * @param InRootClass		The root class to filter the Class Viewer by to only show children of
 	 * @param InOnClassPicked	Callback delegate to fire when a class is picked
 	 */
-	void GetCreateSettingsClassSubMenu(FMenuBuilder& InMenuBuilder, UClass* InRootClass, FOnClassPicked InOnClassPicked)
+	void GetCreateSettingsClassSubMenu(UToolMenu* InMenu, UClass* InRootClass, FOnClassPicked InOnClassPicked)
 	{
 		FClassViewerInitializationOptions Options;
 		Options.Mode = EClassViewerMode::ClassPicker;
@@ -362,9 +361,8 @@ namespace LevelEditorActionHelpers
 		TSharedRef<SWidget> ClassViewer = FModuleManager::LoadModuleChecked<FClassViewerModule>("ClassViewer").CreateClassViewer(Options, InOnClassPicked);
 		FFormatNamedArguments FormatArgs;
 		FormatArgs.Add(TEXT("RootClass"), RootClassName);
-		InMenuBuilder.BeginSection(NAME_None, FText::Format(NSLOCTEXT("LevelToolBarViewMenu", "CreateGameModeLabel", "Select {RootClass} parent class"), FormatArgs));
-		InMenuBuilder.AddWidget(ClassViewer, FText::GetEmpty(), true);
-		InMenuBuilder.EndSection();
+		FToolMenuSection& Section = InMenu->AddSection("CreateSettingsClass", FText::Format(NSLOCTEXT("LevelToolBarViewMenu", "CreateGameModeLabel", "Select {RootClass} parent class"), FormatArgs));
+		Section.AddEntry(FToolMenuEntry::InitWidget("ClassViewer", ClassViewer, FText::GetEmpty(), true));
 	}
 
 	/** Helper struct for passing all required data to the GetBlueprintSettingsSubMenu function */
@@ -425,11 +423,11 @@ namespace LevelEditorActionHelpers
 	/**
 	 * A sub-menu for the Blueprints dropdown, facilitates all the sub-menu actions such as creating, editing, and selecting classes for the world settings game mode.
 	 *
-	 * @param InMenuBuilder		Object to append menu items/widgets to
+	 * @param InMenu		Object to append menu items/widgets to
 	 * @param InCommandList		Commandlist for menu items
 	 * @param InSettingsData	All the data needed to create the menu actions
 	 */
-	void GetBlueprintSettingsSubMenu(FMenuBuilder& InMenuBuilder, const TSharedRef<FUICommandList> InCommandList, FBlueprintMenuSettings InSettingsData);
+	void GetBlueprintSettingsSubMenu(UToolMenu* InMenu, FBlueprintMenuSettings InSettingsData);
 
 	/** Returns TRUE if the class can be edited, always TRUE for Blueprints and for native classes a compiler must be present */
 	bool CanEditClass(UClass* InClass)
@@ -523,47 +521,54 @@ namespace LevelEditorActionHelpers
 #undef LOCTEXT_NAMESPACE
 	}
 
-	void CreateGameModeSubMenu(FMenuBuilder& InMenuBuilder, TSharedRef<FUICommandList> InCommandList, TWeakPtr< SLevelEditor > InLevelEditor, bool bInProjectSettings)
+	void CreateGameModeSubMenu(FToolMenuSection& Section, const FName InName, bool bInProjectSettings)
 	{
 #define LOCTEXT_NAMESPACE "LevelToolBarViewMenu"
-		LevelEditorActionHelpers::FBlueprintMenuSettings GameModeMenuSettings;
-		GameModeMenuSettings.EditCommand = 
-			FUIAction(
-				FExecuteAction::CreateStatic< TWeakPtr< SLevelEditor > >( &OpenGameModeBlueprint, InLevelEditor, bInProjectSettings )
-			);
-		GameModeMenuSettings.OnCreateClassPicked = FOnClassPicked::CreateStatic( &LevelEditorActionHelpers::OnCreateGameModeClassPicked, InLevelEditor, bInProjectSettings );
-		GameModeMenuSettings.OnSelectClassPicked = FOnClassPicked::CreateStatic( &LevelEditorActionHelpers::OnSelectGameModeClassPicked, InLevelEditor, bInProjectSettings );
-		GameModeMenuSettings.CurrentClass = LevelEditorActionHelpers::GetGameModeClass(InLevelEditor, bInProjectSettings);
-		GameModeMenuSettings.RootClass = AGameModeBase::StaticClass();
-		GameModeMenuSettings.LevelEditor = InLevelEditor;
-		GameModeMenuSettings.bIsProjectSettings = bInProjectSettings;
-
-		auto IsGameModeActive = [](TWeakPtr< SLevelEditor > InLevelEditorPtr, bool bInProjSettings)->bool
+		Section.AddDynamicEntry(InName, FNewToolMenuSectionDelegate::CreateLambda([=](FToolMenuSection& InSection)
 		{
-			UClass* WorldSettingsGameMode = LevelEditorActionHelpers::GetGameModeClass(InLevelEditorPtr, false);
-			if((WorldSettingsGameMode == nullptr) ^ bInProjSettings ) //(WorldSettingsGameMode && !bInProjectSettings) || (!WorldSettingsGameMode && bInProjectSettings) )
+			ULevelEditorMenuContext* Context = InSection.FindContext<ULevelEditorMenuContext>();
+			if (Context && Context->LevelEditor.IsValid())
 			{
-				return false;
-			}
-			return true;
-		};
+				LevelEditorActionHelpers::FBlueprintMenuSettings GameModeMenuSettings;
+				GameModeMenuSettings.EditCommand =
+					FUIAction(
+						FExecuteAction::CreateStatic< TWeakPtr< SLevelEditor > >(&OpenGameModeBlueprint, Context->LevelEditor, bInProjectSettings)
+					);
+				GameModeMenuSettings.OnCreateClassPicked = FOnClassPicked::CreateStatic(&LevelEditorActionHelpers::OnCreateGameModeClassPicked, Context->LevelEditor, bInProjectSettings);
+				GameModeMenuSettings.OnSelectClassPicked = FOnClassPicked::CreateStatic(&LevelEditorActionHelpers::OnSelectGameModeClassPicked, Context->LevelEditor, bInProjectSettings);
+				GameModeMenuSettings.CurrentClass = LevelEditorActionHelpers::GetGameModeClass(Context->LevelEditor, bInProjectSettings);
+				GameModeMenuSettings.RootClass = AGameModeBase::StaticClass();
+				GameModeMenuSettings.LevelEditor = Context->LevelEditor;
+				GameModeMenuSettings.bIsProjectSettings = bInProjectSettings;
+
+				auto IsGameModeActive = [](TWeakPtr< SLevelEditor > InLevelEditorPtr, bool bInProjSettings)->bool
+				{
+					UClass* WorldSettingsGameMode = LevelEditorActionHelpers::GetGameModeClass(InLevelEditorPtr, false);
+					if ((WorldSettingsGameMode == nullptr) ^ bInProjSettings) //(WorldSettingsGameMode && !bInProjectSettings) || (!WorldSettingsGameMode && bInProjectSettings) )
+					{
+						return false;
+					}
+					return true;
+				};
 
-		InMenuBuilder.AddSubMenu( LevelEditorActionHelpers::GetOpenGameModeBlueprintLabel(InLevelEditor, bInProjectSettings), 
-									GetSubMenuTooltip(GameModeMenuSettings.CurrentClass, GameModeMenuSettings.RootClass, bInProjectSettings), 
-									FNewMenuDelegate::CreateStatic( &LevelEditorActionHelpers::GetBlueprintSettingsSubMenu, InCommandList, GameModeMenuSettings), 
-									FUIAction(FExecuteAction(), FCanExecuteAction(), FIsActionChecked::CreateStatic(IsGameModeActive, InLevelEditor, bInProjectSettings)),
-									NAME_None, EUserInterfaceActionType::RadioButton );
+				InSection.AddSubMenu(InName, LevelEditorActionHelpers::GetOpenGameModeBlueprintLabel(Context->LevelEditor, bInProjectSettings),
+					GetSubMenuTooltip(GameModeMenuSettings.CurrentClass, GameModeMenuSettings.RootClass, bInProjectSettings),
+					FNewToolMenuDelegate::CreateStatic(&LevelEditorActionHelpers::GetBlueprintSettingsSubMenu, GameModeMenuSettings),
+					FUIAction(FExecuteAction(), FCanExecuteAction(), FIsActionChecked::CreateStatic(IsGameModeActive, Context->LevelEditor, bInProjectSettings)),
+					EUserInterfaceActionType::RadioButton);
+			}
+		}));
 #undef LOCTEXT_NAMESPACE
 	}
 
 	/**
 	 * Builds the game mode's sub menu objects
 	 *
-	 * @param InMenuBuilder		Object to append menu items/widgets to
+	 * @param InSection			Object to append menu items/widgets to
 	 * @param InCommandList		Commandlist for menu items
 	 * @param InSettingsData	All the data needed to create the menu actions
 	 */
-	void GetGameModeSubMenu(FMenuBuilder& InMenuBuilder, const TSharedRef<FUICommandList> InCommandList, const FBlueprintMenuSettings& InSettingsData)
+	void GetGameModeSubMenu(FToolMenuSection& InSection, const FBlueprintMenuSettings& InSettingsData)
 	{
 #define LOCTEXT_NAMESPACE "LevelToolBarViewMenu"
 		// Game State
@@ -579,9 +584,9 @@ namespace LevelEditorActionHelpers
 		GameStateMenuSettings.LevelEditor = InSettingsData.LevelEditor;
 		GameStateMenuSettings.bIsProjectSettings = InSettingsData.bIsProjectSettings;
 
-		InMenuBuilder.AddSubMenu( LevelEditorActionHelpers::GetOpenGameStateBlueprintLabel(InSettingsData.LevelEditor, InSettingsData.bIsProjectSettings),
+		InSection.AddSubMenu("OpenGameStateBlueprint", LevelEditorActionHelpers::GetOpenGameStateBlueprintLabel(InSettingsData.LevelEditor, InSettingsData.bIsProjectSettings),
 			GetSubMenuTooltip(GameStateMenuSettings.CurrentClass, GameStateMenuSettings.RootClass, InSettingsData.bIsProjectSettings),
-			FNewMenuDelegate::CreateStatic( &LevelEditorActionHelpers::GetBlueprintSettingsSubMenu, InCommandList, GameStateMenuSettings )
+			FNewToolMenuDelegate::CreateStatic( &LevelEditorActionHelpers::GetBlueprintSettingsSubMenu, GameStateMenuSettings )
 		);
 
 		// Pawn
@@ -597,9 +602,9 @@ namespace LevelEditorActionHelpers
 		PawnMenuSettings.LevelEditor = InSettingsData.LevelEditor;
 		PawnMenuSettings.bIsProjectSettings = InSettingsData.bIsProjectSettings;
 
-		InMenuBuilder.AddSubMenu( LevelEditorActionHelpers::GetOpenPawnBlueprintLabel(InSettingsData.LevelEditor, InSettingsData.bIsProjectSettings),
+		InSection.AddSubMenu("OpenPawnBlueprint", LevelEditorActionHelpers::GetOpenPawnBlueprintLabel(InSettingsData.LevelEditor, InSettingsData.bIsProjectSettings),
 			GetSubMenuTooltip(PawnMenuSettings.CurrentClass, PawnMenuSettings.RootClass, InSettingsData.bIsProjectSettings),
-			FNewMenuDelegate::CreateStatic( &LevelEditorActionHelpers::GetBlueprintSettingsSubMenu, InCommandList, PawnMenuSettings )
+			FNewToolMenuDelegate::CreateStatic( &LevelEditorActionHelpers::GetBlueprintSettingsSubMenu, PawnMenuSettings )
 		);
 
 		// HUD
@@ -615,9 +620,9 @@ namespace LevelEditorActionHelpers
 		HUDMenuSettings.LevelEditor = InSettingsData.LevelEditor;
 		HUDMenuSettings.bIsProjectSettings = InSettingsData.bIsProjectSettings;
 
-		InMenuBuilder.AddSubMenu( LevelEditorActionHelpers::GetOpenHUDBlueprintLabel(InSettingsData.LevelEditor, InSettingsData.bIsProjectSettings),
+		InSection.AddSubMenu("OpenHUDBlueprint", LevelEditorActionHelpers::GetOpenHUDBlueprintLabel(InSettingsData.LevelEditor, InSettingsData.bIsProjectSettings),
 			GetSubMenuTooltip(HUDMenuSettings.CurrentClass, HUDMenuSettings.RootClass, InSettingsData.bIsProjectSettings),
-			FNewMenuDelegate::CreateStatic( &LevelEditorActionHelpers::GetBlueprintSettingsSubMenu, InCommandList, HUDMenuSettings )
+			FNewToolMenuDelegate::CreateStatic( &LevelEditorActionHelpers::GetBlueprintSettingsSubMenu, HUDMenuSettings )
 		);
 
 		// Player Controller
@@ -633,9 +638,9 @@ namespace LevelEditorActionHelpers
 		PlayerControllerMenuSettings.LevelEditor = InSettingsData.LevelEditor;
 		PlayerControllerMenuSettings.bIsProjectSettings = InSettingsData.bIsProjectSettings;
 
-		InMenuBuilder.AddSubMenu( LevelEditorActionHelpers::GetOpenPlayerControllerBlueprintLabel(InSettingsData.LevelEditor, InSettingsData.bIsProjectSettings),
+		InSection.AddSubMenu("OpenPlayerControllerBlueprint", LevelEditorActionHelpers::GetOpenPlayerControllerBlueprintLabel(InSettingsData.LevelEditor, InSettingsData.bIsProjectSettings),
 			GetSubMenuTooltip(PlayerControllerMenuSettings.CurrentClass, PlayerControllerMenuSettings.RootClass, InSettingsData.bIsProjectSettings),
-			FNewMenuDelegate::CreateStatic( &LevelEditorActionHelpers::GetBlueprintSettingsSubMenu, InCommandList, PlayerControllerMenuSettings )
+			FNewToolMenuDelegate::CreateStatic( &LevelEditorActionHelpers::GetBlueprintSettingsSubMenu, PlayerControllerMenuSettings )
 		);
 #undef LOCTEXT_NAMESPACE
 	}
@@ -656,12 +661,10 @@ namespace LevelEditorActionHelpers
 	};
 }
 
-void LevelEditorActionHelpers::GetBlueprintSettingsSubMenu(FMenuBuilder& InMenuBuilder, const TSharedRef<FUICommandList> InCommandList, FBlueprintMenuSettings InSettingsData)
+void LevelEditorActionHelpers::GetBlueprintSettingsSubMenu(UToolMenu* Menu, FBlueprintMenuSettings InSettingsData)
 {
 #define LOCTEXT_NAMESPACE "LevelToolBarViewMenu"
 
-	InMenuBuilder.PushCommandList(InCommandList);
-
 	FSlateIcon EditBPIcon(FEditorStyle::Get().GetStyleSetName(), TEXT("PropertyWindow.Button_Edit"));
 	FSlateIcon NewBPIcon(FEditorStyle::Get().GetStyleSetName(), TEXT("PropertyWindow.Button_AddToArray"));
 	FText RootClassName = FText::FromString(InSettingsData.RootClass->GetName());
@@ -673,11 +676,12 @@ void LevelEditorActionHelpers::GetBlueprintSettingsSubMenu(FMenuBuilder& InMenuB
 		Args.Add(TEXT("RootClass"), RootClassName);
 		Args.Add(TEXT("TargetLocation"), InSettingsData.bIsProjectSettings? LOCTEXT("Project", "project") : LOCTEXT("World", "world"));
 
+		FToolMenuSection& Section = Menu->AddSection("EditBlueprintOrClass");
 		if(InSettingsData.CurrentClass->ClassGeneratedBy)
 		{
 			FText BlueprintName = FText::FromString(InSettingsData.CurrentClass->ClassGeneratedBy->GetName());
 			Args.Add(TEXT("Blueprint"), BlueprintName);
-			InMenuBuilder.AddMenuEntry( FText::Format( LOCTEXT("EditBlueprint", "Edit {Blueprint}"), Args), FText::Format( LOCTEXT("EditBlueprint_Tooltip", "Open the {TargetLocation}'s assigned {RootClass} blueprint"), Args), EditBPIcon, InSettingsData.EditCommand );
+			Section.AddMenuEntry("EditBlueprint", FText::Format( LOCTEXT("EditBlueprint", "Edit {Blueprint}"), Args), FText::Format( LOCTEXT("EditBlueprint_Tooltip", "Open the {TargetLocation}'s assigned {RootClass} blueprint"), Args), EditBPIcon, InSettingsData.EditCommand );
 		}
 		else
 		{
@@ -687,7 +691,7 @@ void LevelEditorActionHelpers::GetBlueprintSettingsSubMenu(FMenuBuilder& InMenuB
 			FText MenuDescription = FText::Format( LOCTEXT("EditNativeClass", "Edit {Class}.h"), Args);
 			if(FSourceCodeNavigation::IsCompilerAvailable())
 			{
-				InMenuBuilder.AddMenuEntry( MenuDescription, FText::Format( LOCTEXT("EditNativeClass_Tooltip", "Open the {TargetLocation}'s assigned {RootClass} header"), Args), EditBPIcon, InSettingsData.EditCommand );
+				Section.AddMenuEntry("EditNativeClass", MenuDescription, FText::Format( LOCTEXT("EditNativeClass_Tooltip", "Open the {TargetLocation}'s assigned {RootClass} header"), Args), EditBPIcon, InSettingsData.EditCommand );
 			}
 			else
 			{
@@ -697,21 +701,20 @@ void LevelEditorActionHelpers::GetBlueprintSettingsSubMenu(FMenuBuilder& InMenuB
 				};
 
 				// There is no compiler present, this is always disabled with a tooltip to explain why
-				InMenuBuilder.AddMenuEntry( MenuDescription, FText::Format( LOCTEXT("CannotEditNativeClass_Tooltip", "Cannot edit the {TargetLocation}'s assigned {RootClass} header because no compiler is present!"), Args), EditBPIcon, FUIAction(FExecuteAction(), FCanExecuteAction::CreateStatic(CannotEditClass)) );
+				Section.AddMenuEntry("EditNativeClass", MenuDescription, FText::Format( LOCTEXT("CannotEditNativeClass_Tooltip", "Cannot edit the {TargetLocation}'s assigned {RootClass} header because no compiler is present!"), Args), EditBPIcon, FUIAction(FExecuteAction(), FCanExecuteAction::CreateStatic(CannotEditClass)) );
 			}
 		}
 	}
 
 	if(InSettingsData.bIsProjectSettings && InSettingsData.CurrentClass && InSettingsData.CurrentClass->IsChildOf(AGameModeBase::StaticClass()) && !FLevelEditorActionCallbacks::CanSelectGameModeBlueprint())
 	{
-		InMenuBuilder.BeginSection("CheckoutSection", LOCTEXT("CheckoutSection","Check Out Project Settings") );
+		FToolMenuSection& Section = Menu->AddSection("CheckoutSection", LOCTEXT("CheckoutSection", "Check Out Project Settings") );
 		TAttribute<FText> CheckOutLabel;
 		CheckOutLabel.BindStatic(&GetCheckOutLabel);
 
 		TAttribute<FText> CheckOutTooltip;
 		CheckOutTooltip.BindStatic(&GetCheckOutTooltip);
-		InMenuBuilder.AddMenuEntry(FLevelEditorCommands::Get().CheckOutProjectSettingsConfig, NAME_None, CheckOutLabel, CheckOutTooltip, FSlateIcon(FEditorStyle::Get().GetStyleSetName(), TEXT("Icons.Error")));
-		InMenuBuilder.EndSection();
+		Section.AddMenuEntry(FLevelEditorCommands::Get().CheckOutProjectSettingsConfig, CheckOutLabel, CheckOutTooltip, FSlateIcon(FEditorStyle::Get().GetStyleSetName(), TEXT("Icons.Error")));
 	}
 
 	auto CannotCreateSelectNativeProjectGameMode = [](bool bInIsProjectSettings) -> bool
@@ -724,49 +727,46 @@ void LevelEditorActionHelpers::GetBlueprintSettingsSubMenu(FMenuBuilder& InMenuB
 		return true;
 	};
 
+	FToolMenuSection& Section = Menu->AddSection("CreateBlueprint");
+
 	// Create a new GameMode, this is always available so the user can easily create a new one
-	InMenuBuilder.AddSubMenu( LOCTEXT("CreateBlueprint", "Create..."),
+	Section.AddSubMenu("CreateBlueprint", LOCTEXT("CreateBlueprint", "Create..."),
 		GetCreateMenuTooltip(GetGameModeClass(InSettingsData.LevelEditor, InSettingsData.bIsProjectSettings), InSettingsData.RootClass, InSettingsData.bIsProjectSettings),
-		FNewMenuDelegate::CreateStatic( &LevelEditorActionHelpers::GetCreateSettingsClassSubMenu, InSettingsData.RootClass, InSettingsData.OnCreateClassPicked ),
+		FNewToolMenuDelegate::CreateStatic( &LevelEditorActionHelpers::GetCreateSettingsClassSubMenu, InSettingsData.RootClass, InSettingsData.OnCreateClassPicked ),
 		FUIAction(
 			FExecuteAction(), 
 			InSettingsData.RootClass == AGameModeBase::StaticClass()? 
 				FCanExecuteAction::CreateStatic(CannotCreateSelectNativeProjectGameMode, InSettingsData.bIsProjectSettings) 
 				: FCanExecuteAction::CreateStatic( &CanCreateSelectSubClass, GetGameModeClass(InSettingsData.LevelEditor, InSettingsData.bIsProjectSettings), InSettingsData.bIsProjectSettings )
 		),
-		NAME_None, EUserInterfaceActionType::Button, false, NewBPIcon
+		EUserInterfaceActionType::Button, false, NewBPIcon
 	);
 
 	// Select a game mode, this is always available so the user can switch his selection
 	FFormatNamedArguments Args;
 	Args.Add(TEXT("RootClass"), RootClassName);
-	InMenuBuilder.AddSubMenu(FText::Format(LOCTEXT("SelectGameModeClass", "Select {RootClass} Class"), Args),
+	Section.AddSubMenu("SelectGameModeClass", FText::Format(LOCTEXT("SelectGameModeClass", "Select {RootClass} Class"), Args),
 		GetSelectMenuTooltip(GetGameModeClass(InSettingsData.LevelEditor, InSettingsData.bIsProjectSettings), InSettingsData.RootClass, InSettingsData.bIsProjectSettings),
-		FNewMenuDelegate::CreateStatic( &LevelEditorActionHelpers::GetSelectSettingsClassSubMenu, InSettingsData.RootClass, InSettingsData.OnSelectClassPicked ),
+		FNewToolMenuDelegate::CreateStatic( &LevelEditorActionHelpers::GetSelectSettingsClassSubMenu, InSettingsData.RootClass, InSettingsData.OnSelectClassPicked ),
 		FUIAction(
 			FExecuteAction(), 
 			InSettingsData.RootClass == AGameModeBase::StaticClass()?
 				FCanExecuteAction::CreateStatic(CannotCreateSelectNativeProjectGameMode, InSettingsData.bIsProjectSettings) 
 				: FCanExecuteAction::CreateStatic( &CanCreateSelectSubClass, GetGameModeClass(InSettingsData.LevelEditor, InSettingsData.bIsProjectSettings), InSettingsData.bIsProjectSettings )
 		),
-		NAME_None, EUserInterfaceActionType::Button
+		EUserInterfaceActionType::Button
 	);
 
 	// For GameMode classes only, there are some sub-classes we need to add to the menu
 	if(InSettingsData.RootClass == AGameModeBase::StaticClass())
 	{
-		InMenuBuilder.BeginSection(NAME_None, LOCTEXT("GameModeClasses", "Game Mode Classes"));
-
+		FToolMenuSection& GameModeClassesSection = Menu->AddSection("GameModeClasses", LOCTEXT("GameModeClasses", "Game Mode Classes"));
 		if(InSettingsData.CurrentClass)
 		{
-			GetGameModeSubMenu(InMenuBuilder, InCommandList, InSettingsData);
+			GetGameModeSubMenu(GameModeClassesSection, InSettingsData);
 		}
-
-		InMenuBuilder.EndSection();
 	}
 
-	InMenuBuilder.PopCommandList();
-
 #undef LOCTEXT_NAMESPACE
 }
 
@@ -1173,13 +1173,21 @@ void FLevelEditorToolBar::RegisterLevelEditorToolBar( const TSharedRef<FUIComman
 {
 	RegisterSourceControlMenu();
 	RegisterCinematicsMenu();
+	RegisterBuildMenu();
+
+#if WITH_LIVE_CODING
+	RegisterCompileMenu();
+#endif
+
+	RegisterQuickSettingsMenu();
+	RegisterOpenBlueprintMenu();
 
 #define LOCTEXT_NAMESPACE "LevelEditorToolBar"
 
-	UToolMenu* ToolbarBuilder = UToolMenus::Get()->RegisterMenu("LevelEditor.LevelEditorToolBar", NAME_None, EMultiBoxType::ToolBar);
+	UToolMenu* Toolbar = UToolMenus::Get()->RegisterMenu("LevelEditor.LevelEditorToolBar", NAME_None, EMultiBoxType::ToolBar);
 
 	{
-		FToolMenuSection& Section = ToolbarBuilder->AddSection("File");
+		FToolMenuSection& Section = Toolbar->AddSection("File");
 
 		// Save All Levels
 		Section.AddEntry(FToolMenuEntry::InitToolBarButton(
@@ -1266,7 +1274,7 @@ void FLevelEditorToolBar::RegisterLevelEditorToolBar( const TSharedRef<FUIComman
 			Section.AddEntry(FToolMenuEntry::InitComboButton(
 				"SourceControl",
 				FUIAction(),
-				FOnGetContent::CreateStatic(&FLevelEditorToolBar::GenerateSourceControlMenu, InCommandList),
+				FOnGetContent::CreateStatic(&FLevelEditorToolBar::GenerateSourceControlMenu, InCommandList, TWeakPtr<SLevelEditor>(InLevelEditor)),
 				LOCTEXT("SourceControl_Label", "Source Control"),
 				TAttribute<FText>::Create(&FSourceControlStatus::GetSourceControlTooltip),
 				TAttribute<FSlateIcon>::Create(&FSourceControlStatus::GetSourceControlIcon),
@@ -1276,7 +1284,7 @@ void FLevelEditorToolBar::RegisterLevelEditorToolBar( const TSharedRef<FUIComman
 	}
 
 	{
-		FToolMenuSection& Section = ToolbarBuilder->AddSection("Content");
+		FToolMenuSection& Section = Toolbar->AddSection("Content");
 		Section.AddEntry(FToolMenuEntry::InitToolBarButton(FLevelEditorCommands::Get().OpenContentBrowser, LOCTEXT( "ContentBrowser_Override", "Content" ), TAttribute<FText>(), TAttribute<FSlateIcon>(), "LevelToolbarContent"));
 		if (FLauncherPlatformModule::Get()->CanOpenLauncher(true)) 
 		{
@@ -1284,12 +1292,12 @@ void FLevelEditorToolBar::RegisterLevelEditorToolBar( const TSharedRef<FUIComman
 		}
 	}
 
-	FToolMenuSection& SettingsSection = ToolbarBuilder->AddSection("Settings");
+	FToolMenuSection& SettingsSection = Toolbar->AddSection("Settings");
 	{
 		SettingsSection.AddEntry(FToolMenuEntry::InitComboButton(
 			"LevelToolbarQuickSettings",
 			FUIAction(),
-			FOnGetContent::CreateStatic(&FLevelEditorToolBar::GenerateQuickSettingsMenu, InCommandList),
+			FOnGetContent::CreateStatic(&FLevelEditorToolBar::GenerateQuickSettingsMenu, InCommandList, TWeakPtr<SLevelEditor>(InLevelEditor)),
 			LOCTEXT("QuickSettingsCombo", "Settings"),
 			LOCTEXT("QuickSettingsCombo_ToolTip", "Project and Editor settings"),
 			FSlateIcon(FEditorStyle::GetStyleSetName(), "LevelEditor.GameSettings"),
@@ -1412,7 +1420,7 @@ void FLevelEditorToolBar::RegisterLevelEditorToolBar( const TSharedRef<FUIComman
 	}
 
 	{
-		FToolMenuSection& Section = ToolbarBuilder->AddSection("Misc");
+		FToolMenuSection& Section = Toolbar->AddSection("Misc");
 		Section.AddEntry(FToolMenuEntry::InitComboButton(
 			"OpenBlueprint",
 			FUIAction(),
@@ -1435,7 +1443,7 @@ void FLevelEditorToolBar::RegisterLevelEditorToolBar( const TSharedRef<FUIComman
 	}
 
 	{
-		FToolMenuSection& Section = ToolbarBuilder->AddSection("Compile");
+		FToolMenuSection& Section = Toolbar->AddSection("Compile");
 		// Build			
 		Section.AddEntry(FToolMenuEntry::InitToolBarButton( FLevelEditorCommands::Get().Build, LOCTEXT("BuildAll", "Build")) );
 
@@ -1443,47 +1451,50 @@ void FLevelEditorToolBar::RegisterLevelEditorToolBar( const TSharedRef<FUIComman
 		Section.AddEntry(FToolMenuEntry::InitComboButton(
 			"BuildComboButton",
 			FUIAction(),
-			FOnGetContent::CreateStatic( &FLevelEditorToolBar::GenerateBuildMenuContent, InCommandList ),
+			FOnGetContent::CreateStatic( &FLevelEditorToolBar::GenerateBuildMenuContent, InCommandList, TWeakPtr<SLevelEditor>(InLevelEditor) ),
 			LOCTEXT( "BuildCombo_Label", "Build Options" ),
 			LOCTEXT( "BuildComboToolTip", "Build options menu" ),
 			FSlateIcon(FEditorStyle::GetStyleSetName(), "LevelEditor.Build"),
 			true));
 
-		// Only show the compile options on machines with the solution (assuming they can build it)
-		if ( FSourceCodeNavigation::IsCompilerAvailable() )
-		{
-			// Since we can always add new code to the project, only hide these buttons if we haven't done so yet
-			Section.AddEntry(FToolMenuEntry::InitToolBarButton(
-				"CompileButton",
-				FUIAction(
-					FExecuteAction::CreateStatic(&FLevelEditorActionCallbacks::RecompileGameCode_Clicked),
-					FCanExecuteAction::CreateStatic(&FLevelEditorActionCallbacks::Recompile_CanExecute),
-					FIsActionChecked(),
-					FIsActionButtonVisible::CreateStatic(FLevelEditorActionCallbacks::CanShowSourceCodeActions)),
-				LOCTEXT( "CompileMenuButton", "Compile" ),
-				FLevelEditorCommands::Get().RecompileGameCode->GetDescription(),
-				FSlateIcon(FEditorStyle::GetStyleSetName(), "LevelEditor.Recompile")
+		Section.AddDynamicEntry("CompilerAvailable", FNewToolMenuSectionDelegate::CreateLambda([](FToolMenuSection& InSection)
+		{
+			// Only show the compile options on machines with the solution (assuming they can build it)
+			if (FSourceCodeNavigation::IsCompilerAvailable())
+			{
+				// Since we can always add new code to the project, only hide these buttons if we haven't done so yet
+				InSection.AddEntry(FToolMenuEntry::InitToolBarButton(
+					"CompileButton",
+					FUIAction(
+						FExecuteAction::CreateStatic(&FLevelEditorActionCallbacks::RecompileGameCode_Clicked),
+						FCanExecuteAction::CreateStatic(&FLevelEditorActionCallbacks::Recompile_CanExecute),
+						FIsActionChecked(),
+						FIsActionButtonVisible::CreateStatic(FLevelEditorActionCallbacks::CanShowSourceCodeActions)),
+					LOCTEXT("CompileMenuButton", "Compile"),
+					FLevelEditorCommands::Get().RecompileGameCode->GetDescription(),
+					FSlateIcon(FEditorStyle::GetStyleSetName(), "LevelEditor.Recompile")
 				));
 
 #if WITH_LIVE_CODING
-			Section.AddEntry(FToolMenuEntry::InitComboButton(
-				"CompileComboButton",
-				FUIAction(
-					FExecuteAction(),
-					FCanExecuteAction(),
-					FIsActionChecked(),
-					FIsActionButtonVisible::CreateStatic(FLevelEditorActionCallbacks::CanShowSourceCodeActions)), 
-				FOnGetContent::CreateStatic( &FLevelEditorToolBar::GenerateCompileMenuContent, InCommandList ),
-				LOCTEXT( "CompileCombo_Label", "Compile Options" ),
-				LOCTEXT( "CompileComboToolTip", "Compile options menu" ),
-				FSlateIcon(FEditorStyle::GetStyleSetName(), "LevelEditor.Recompile"),
-				true
+				InSection.AddEntry(FToolMenuEntry::InitComboButton(
+					"CompileComboButton",
+					FUIAction(
+						FExecuteAction(),
+						FCanExecuteAction(),
+						FIsActionChecked(),
+						FIsActionButtonVisible::CreateStatic(FLevelEditorActionCallbacks::CanShowSourceCodeActions)),
+					FNewToolMenuWidgetChoice(),
+					LOCTEXT("CompileCombo_Label", "Compile Options"),
+					LOCTEXT("CompileComboToolTip", "Compile options menu"),
+					FSlateIcon(FEditorStyle::GetStyleSetName(), "LevelEditor.Recompile"),
+					true
 				));
 #endif
-		}
+			}
+		}));
 	}
 
-	ToolbarBuilder->AddDynamicSection("Game", FNewToolBarDelegateLegacy::CreateLambda([](FToolBarBuilder& InBuilder, UToolMenu* InData)
+	Toolbar->AddDynamicSection("Game", FNewToolBarDelegateLegacy::CreateLambda([](FToolBarBuilder& InBuilder, UToolMenu* InData)
 	{
 		InBuilder.BeginSection("Game");
 
@@ -1525,38 +1536,60 @@ TSharedRef< SWidget > FLevelEditorToolBar::MakeLevelEditorToolBar( const TShared
 
 
 
-TSharedRef< SWidget > FLevelEditorToolBar::GenerateBuildMenuContent( TSharedRef<FUICommandList> InCommandList )
+TSharedRef< SWidget > FLevelEditorToolBar::GenerateBuildMenuContent( TSharedRef<FUICommandList> InCommandList, TWeakPtr<SLevelEditor> InLevelEditor)
 {
-#define LOCTEXT_NAMESPACE "LevelToolBarBuildMenu"
-
 	// Get all menu extenders for this context menu from the level editor module
 	FLevelEditorModule& LevelEditorModule = FModuleManager::GetModuleChecked<FLevelEditorModule>( TEXT("LevelEditor") );
 	TSharedPtr<FExtender> MenuExtender = LevelEditorModule.AssembleExtenders(InCommandList, LevelEditorModule.GetAllLevelEditorToolbarBuildMenuExtenders());
+	FToolMenuContext MenuContext(InCommandList, MenuExtender);
+
+	ULevelEditorMenuContext* LevelEditorMenuContext = NewObject<ULevelEditorMenuContext>();
+	LevelEditorMenuContext->LevelEditor = InLevelEditor;
+	MenuContext.AddObject(LevelEditorMenuContext);
+
+	return UToolMenus::Get()->GenerateWidget("LevelEditor.LevelEditorToolBar.BuildComboButton", MenuContext);
+}
 
-	const bool bShouldCloseWindowAfterMenuSelection = true;
-	FMenuBuilder MenuBuilder( bShouldCloseWindowAfterMenuSelection, InCommandList, MenuExtender );
+void FLevelEditorToolBar::RegisterBuildMenu()
+{
+#define LOCTEXT_NAMESPACE "LevelToolBarBuildMenu"
+
+	static const FName BaseMenuName = "LevelEditor.LevelEditorToolBar.BuildComboButton";
+	UToolMenu* Menu = UToolMenus::Get()->RegisterMenu(BaseMenuName);
 
 	struct FLightingMenus
 	{
+	public:
+
+		static void RegisterMenus(const FName BaseMenuName)
+		{
+			FLightingMenus::RegisterLightingQualityMenu(BaseMenuName);
+			FLightingMenus::RegisterLightingInfoMenu(BaseMenuName);
+		}
+
+	private:
 
 		/** Generates a lighting quality sub-menu */
-		static void MakeLightingQualityMenu( FMenuBuilder& InMenuBuilder )
+		static void RegisterLightingQualityMenu(const FName BaseMenuName)
 		{
-			InMenuBuilder.BeginSection("LevelEditorBuildLightingQuality", LOCTEXT( "LightingQualityHeading", "Quality Level" ) );
+			UToolMenu* SubMenu = UToolMenus::Get()->RegisterMenu(UToolMenus::JoinMenuPaths(BaseMenuName, "LightingQuality"));
+
 			{
-				InMenuBuilder.AddMenuEntry( FLevelEditorCommands::Get().LightingQuality_Production );
-				InMenuBuilder.AddMenuEntry( FLevelEditorCommands::Get().LightingQuality_High );
-				InMenuBuilder.AddMenuEntry( FLevelEditorCommands::Get().LightingQuality_Medium );
-				InMenuBuilder.AddMenuEntry( FLevelEditorCommands::Get().LightingQuality_Preview );
+				FToolMenuSection& Section = SubMenu->AddSection("LevelEditorBuildLightingQuality", LOCTEXT( "LightingQualityHeading", "Quality Level" ) );
+				Section.AddMenuEntry( FLevelEditorCommands::Get().LightingQuality_Production );
+				Section.AddMenuEntry( FLevelEditorCommands::Get().LightingQuality_High );
+				Section.AddMenuEntry( FLevelEditorCommands::Get().LightingQuality_Medium );
+				Section.AddMenuEntry( FLevelEditorCommands::Get().LightingQuality_Preview );
 			}
-			InMenuBuilder.EndSection();
 		}
 
 		/** Generates a lighting density sub-menu */
-		static void MakeLightingDensityMenu( FMenuBuilder& InMenuBuilder )
+		static void RegisterLightingDensityMenu(const FName BaseMenuName)
 		{
-			InMenuBuilder.BeginSection("LevelEditorBuildLightingDensity", LOCTEXT( "LightingDensityHeading", "Density Rendering" ) );
+			UToolMenu* SubMenu = UToolMenus::Get()->RegisterMenu(UToolMenus::JoinMenuPaths(BaseMenuName, "LightingDensity"));
+
 			{
+				FToolMenuSection& Section = SubMenu->AddSection("LevelEditorBuildLightingDensity", LOCTEXT( "LightingDensityHeading", "Density Rendering" ) );
 				TSharedRef<SWidget> Ideal =		SNew(SHorizontalBox)
 												+SHorizontalBox::Slot()
 												.Padding( FMargin( 27.0f, 0.0f, 0.0f, 0.0f ) )
@@ -1568,7 +1601,7 @@ TSharedRef< SWidget > FLevelEditorToolBar::GenerateBuildMenuContent( TSharedRef<
 													.Value(FLevelEditorActionCallbacks::GetLightingDensityIdeal())
 													.OnValueChanged_Static(&FLevelEditorActionCallbacks::SetLightingDensityIdeal)
 												];
-				InMenuBuilder.AddWidget(Ideal, LOCTEXT("LightingDensity_Ideal","Ideal Density"));
+				Section.AddEntry(FToolMenuEntry::InitWidget("Ideal", Ideal, LOCTEXT("LightingDensity_Ideal","Ideal Density")));
 				
 				TSharedRef<SWidget> Maximum =	SNew(SHorizontalBox)
 												+SHorizontalBox::Slot()
@@ -1580,7 +1613,7 @@ TSharedRef< SWidget > FLevelEditorToolBar::GenerateBuildMenuContent( TSharedRef<
 													.Value(FLevelEditorActionCallbacks::GetLightingDensityMaximum())
 													.OnValueChanged_Static(&FLevelEditorActionCallbacks::SetLightingDensityMaximum)
 												];
-				InMenuBuilder.AddWidget(Maximum, LOCTEXT("LightingDensity_Maximum","Maximum Density"));
+				Section.AddEntry(FToolMenuEntry::InitWidget("Maximum", Maximum, LOCTEXT("LightingDensity_Maximum","Maximum Density")));
 
 				TSharedRef<SWidget> ClrScale =	SNew(SHorizontalBox)
 												+SHorizontalBox::Slot()
@@ -1593,7 +1626,7 @@ TSharedRef< SWidget > FLevelEditorToolBar::GenerateBuildMenuContent( TSharedRef<
 													.Value(FLevelEditorActionCallbacks::GetLightingDensityColorScale())
 													.OnValueChanged_Static(&FLevelEditorActionCallbacks::SetLightingDensityColorScale)
 												];
-				InMenuBuilder.AddWidget(ClrScale, LOCTEXT("LightingDensity_ColorScale","Color Scale"));
+				Section.AddEntry(FToolMenuEntry::InitWidget("ColorScale", ClrScale, LOCTEXT("LightingDensity_ColorScale","Color Scale")));
 
 				TSharedRef<SWidget> GrayScale =	SNew(SHorizontalBox)
 												+SHorizontalBox::Slot()
@@ -1606,18 +1639,19 @@ TSharedRef< SWidget > FLevelEditorToolBar::GenerateBuildMenuContent( TSharedRef<
 													.Value(FLevelEditorActionCallbacks::GetLightingDensityGrayscaleScale())
 													.OnValueChanged_Static(&FLevelEditorActionCallbacks::SetLightingDensityGrayscaleScale)
 												];
-				InMenuBuilder.AddWidget(GrayScale, LOCTEXT("LightingDensity_GrayscaleScale","Grayscale Scale"));
+				Section.AddEntry(FToolMenuEntry::InitWidget("GrayscaleScale", GrayScale, LOCTEXT("LightingDensity_GrayscaleScale","Grayscale Scale")));
 
-				InMenuBuilder.AddMenuEntry( FLevelEditorCommands::Get().LightingDensity_RenderGrayscale );
+				Section.AddMenuEntry( FLevelEditorCommands::Get().LightingDensity_RenderGrayscale );
 			}
-			InMenuBuilder.EndSection();
 		}
 
 		/** Generates a lighting resolution sub-menu */
-		static void MakeLightingResolutionMenu( FMenuBuilder& InMenuBuilder )
+		static void RegisterLightingResolutionMenu(const FName BaseMenuName)
 		{
-			InMenuBuilder.BeginSection("LevelEditorBuildLightingResolution1", LOCTEXT( "LightingResolutionHeading1", "Primitive Types" ) );
+			UToolMenu* SubMenu = UToolMenus::Get()->RegisterMenu(UToolMenus::JoinMenuPaths(BaseMenuName, "LightingResolution"));
+
 			{
+				FToolMenuSection& Section = SubMenu->AddSection("LevelEditorBuildLightingResolution1", LOCTEXT( "LightingResolutionHeading1", "Primitive Types" ) );
 				TSharedRef<SWidget> Meshes =	SNew(SHorizontalBox)
 												+SHorizontalBox::Slot()
 												.AutoWidth()
@@ -1654,7 +1688,7 @@ TSharedRef< SWidget > FLevelEditorToolBar::GenerateBuildMenuContent( TSharedRef<
 													.Value(FLevelEditorActionCallbacks::GetLightingResolutionMaxSMs())
 													.OnValueChanged_Static(&FLevelEditorActionCallbacks::SetLightingResolutionMaxSMs)
 												];
-				InMenuBuilder.AddWidget(Meshes, FText::GetEmpty(), true);
+				Section.AddEntry(FToolMenuEntry::InitWidget("Meshes", Meshes, FText::GetEmpty(), true));
 				
 				TSharedRef<SWidget> BSPs =		SNew(SHorizontalBox)
 												+SHorizontalBox::Slot()
@@ -1692,21 +1726,19 @@ TSharedRef< SWidget > FLevelEditorToolBar::GenerateBuildMenuContent( TSharedRef<
 													.Value(FLevelEditorActionCallbacks::GetLightingResolutionMaxBSPs())
 													.OnValueChanged_Static(&FLevelEditorActionCallbacks::SetLightingResolutionMaxBSPs)
 												];
-				InMenuBuilder.AddWidget(BSPs, FText::GetEmpty(), true);
+				Section.AddEntry(FToolMenuEntry::InitWidget("BSPs", BSPs, FText::GetEmpty(), true));
 			}
-			InMenuBuilder.EndSection(); //LevelEditorBuildLightingResolution1
 
-			InMenuBuilder.BeginSection("LevelEditorBuildLightingResolution2", LOCTEXT( "LightingResolutionHeading2", "Select Options" ) );
 			{
-				InMenuBuilder.AddMenuEntry( FLevelEditorCommands::Get().LightingResolution_CurrentLevel );
-				InMenuBuilder.AddMenuEntry( FLevelEditorCommands::Get().LightingResolution_SelectedLevels );
-				InMenuBuilder.AddMenuEntry( FLevelEditorCommands::Get().LightingResolution_AllLoadedLevels );
-				InMenuBuilder.AddMenuEntry( FLevelEditorCommands::Get().LightingResolution_SelectedObjectsOnly );
+				FToolMenuSection& Section = SubMenu->AddSection("LevelEditorBuildLightingResolution2", LOCTEXT( "LightingResolutionHeading2", "Select Options" ) );
+				Section.AddMenuEntry( FLevelEditorCommands::Get().LightingResolution_CurrentLevel );
+				Section.AddMenuEntry( FLevelEditorCommands::Get().LightingResolution_SelectedLevels );
+				Section.AddMenuEntry( FLevelEditorCommands::Get().LightingResolution_AllLoadedLevels );
+				Section.AddMenuEntry( FLevelEditorCommands::Get().LightingResolution_SelectedObjectsOnly );
 			}
-			InMenuBuilder.EndSection();
 
-			InMenuBuilder.BeginSection("LevelEditorBuildLightingResolution3", LOCTEXT( "LightingResolutionHeading3", "Ratio" ) );
 			{
+				FToolMenuSection& Section = SubMenu->AddSection("LevelEditorBuildLightingResolution3", LOCTEXT( "LightingResolutionHeading3", "Ratio" ) );
 				TSharedRef<SWidget> Ratio =		SNew(SSpinBox<int32>)
 												.MinValue(0)
 												.MaxValue(400)
@@ -1714,156 +1746,156 @@ TSharedRef< SWidget > FLevelEditorToolBar::GenerateBuildMenuContent( TSharedRef<
 												.Value(FLevelEditorActionCallbacks::GetLightingResolutionRatio())
 												.OnEndSliderMovement_Static(&FLevelEditorActionCallbacks::SetLightingResolutionRatio)
 												.OnValueCommitted_Static(&FLevelEditorActionCallbacks::SetLightingResolutionRatioCommit);
-				InMenuBuilder.AddWidget(Ratio, LOCTEXT( "LightingResolutionRatio", "Ratio" ));
+				Section.AddEntry(FToolMenuEntry::InitWidget("Ratio", Ratio, LOCTEXT( "LightingResolutionRatio", "Ratio" )));
 			}
-			InMenuBuilder.EndSection();
 		}
 
 		/** Generates a lighting info dialogs sub-menu */
-		static void MakeLightingInfoMenu( FMenuBuilder& InMenuBuilder )
+		static void RegisterLightingInfoMenu(const FName BaseMenuName)
 		{
-			InMenuBuilder.BeginSection("LevelEditorBuildLightingInfo", LOCTEXT( "LightingInfoHeading", "Lighting Info Dialogs" ) );
+			RegisterLightingDensityMenu(UToolMenus::JoinMenuPaths(BaseMenuName, "LightingInfo"));
+			RegisterLightingResolutionMenu(UToolMenus::JoinMenuPaths(BaseMenuName, "LightingInfo"));
+
+			UToolMenu* SubMenu = UToolMenus::Get()->RegisterMenu(UToolMenus::JoinMenuPaths(BaseMenuName, "LightingInfo"));
+
 			{
-				InMenuBuilder.AddSubMenu(
+				FToolMenuSection& Section = SubMenu->AddSection("LevelEditorBuildLightingInfo", LOCTEXT( "LightingInfoHeading", "Lighting Info Dialogs" ) );
+				Section.AddSubMenu(
+					"LightingDensity",
 					LOCTEXT( "LightingDensityRenderingSubMenu", "LightMap Density Rendering Options" ),
 					LOCTEXT( "LightingDensityRenderingSubMenu_ToolTip", "Shows the LightMap Density Rendering viewmode options." ),
-					FNewMenuDelegate::CreateStatic( &FLightingMenus::MakeLightingDensityMenu ) );
+					FNewToolMenuChoice() );
 
-				InMenuBuilder.AddSubMenu(
+				Section.AddSubMenu(
+					"LightingResolution",
 					LOCTEXT( "LightingResolutionAdjustmentSubMenu", "LightMap Resolution Adjustment" ),
 					LOCTEXT( "LightingResolutionAdjustmentSubMenu_ToolTip", "Shows the LightMap Resolution Adjustment options." ),
-					FNewMenuDelegate::CreateStatic( &FLightingMenus::MakeLightingResolutionMenu ) );
+					FNewToolMenuChoice() );
 
-				InMenuBuilder.AddMenuEntry( FLevelEditorCommands::Get().LightingStaticMeshInfo, NAME_None, LOCTEXT( "BuildLightingInfo_LightingStaticMeshInfo", "Lighting StaticMesh Info..." ) );
+				Section.AddMenuEntry( FLevelEditorCommands::Get().LightingStaticMeshInfo, LOCTEXT( "BuildLightingInfo_LightingStaticMeshInfo", "Lighting StaticMesh Info..." ) );
 			}
-			InMenuBuilder.EndSection();
 		}
 	};
 
-	MenuBuilder.BeginSection("LevelEditorLighting", LOCTEXT( "LightingHeading", "Lighting" ) );
+	FLightingMenus::RegisterMenus(BaseMenuName);
+
 	{
-		MenuBuilder.AddMenuEntry( FLevelEditorCommands::Get().BuildLightingOnly, NAME_None, LOCTEXT( "BuildLightingOnlyHeading", "Build Lighting Only" ) );
+		FToolMenuSection& Section = Menu->AddSection( "LevelEditorLighting", LOCTEXT( "LightingHeading", "Lighting" ) );
+		Section.AddMenuEntry( FLevelEditorCommands::Get().BuildLightingOnly, LOCTEXT( "BuildLightingOnlyHeading", "Build Lighting Only" ) );
 
-		MenuBuilder.AddSubMenu(
+		Section.AddSubMenu(
+			"LightingQuality",
 			LOCTEXT( "LightingQualitySubMenu", "Lighting Quality" ),
 			LOCTEXT( "LightingQualitySubMenu_ToolTip", "Allows you to select the quality level for precomputed lighting" ),
-			FNewMenuDelegate::CreateStatic( &FLightingMenus::MakeLightingQualityMenu ) );
+			FNewToolMenuChoice() );
 
-		MenuBuilder.AddSubMenu(
+		Section.AddSubMenu(
+			"LightingInfo",
 			LOCTEXT( "BuildLightingInfoSubMenu", "Lighting Info" ),
 			LOCTEXT( "BuildLightingInfoSubMenu_ToolTip", "Access the lighting info dialogs" ),
-			FNewMenuDelegate::CreateStatic( &FLightingMenus::MakeLightingInfoMenu ) );
+			FNewToolMenuChoice() );
 
-		MenuBuilder.AddMenuEntry( FLevelEditorCommands::Get().LightingBuildOptions_UseErrorColoring );
-		MenuBuilder.AddMenuEntry( FLevelEditorCommands::Get().LightingBuildOptions_ShowLightingStats );
+		Section.AddMenuEntry( FLevelEditorCommands::Get().LightingBuildOptions_UseErrorColoring );
+		Section.AddMenuEntry( FLevelEditorCommands::Get().LightingBuildOptions_ShowLightingStats );
 	}
-	MenuBuilder.EndSection();
 
-	MenuBuilder.BeginSection("LevelEditorReflections", LOCTEXT( "ReflectionHeading", "Reflections" ) );
 	{
-		MenuBuilder.AddMenuEntry( FLevelEditorCommands::Get().BuildReflectionCapturesOnly );
+		FToolMenuSection& Section = Menu->AddSection("LevelEditorReflections", LOCTEXT( "ReflectionHeading", "Reflections" ) );
+		Section.AddMenuEntry( FLevelEditorCommands::Get().BuildReflectionCapturesOnly );
 	}
-	MenuBuilder.EndSection();
 
-	MenuBuilder.BeginSection("LevelEditorVisibility", LOCTEXT( "VisibilityHeading", "Visibility" ) );
 	{
-		MenuBuilder.AddMenuEntry( FLevelEditorCommands::Get().BuildLightingOnly_VisibilityOnly );
+		FToolMenuSection& Section = Menu->AddSection("LevelEditorVisibility", LOCTEXT( "VisibilityHeading", "Visibility" ) );
+		Section.AddMenuEntry( FLevelEditorCommands::Get().BuildLightingOnly_VisibilityOnly );
 	}
-	MenuBuilder.EndSection();
 
-	MenuBuilder.BeginSection("LevelEditorGeometry", LOCTEXT( "GeometryHeading", "Geometry" ) );
 	{
-		MenuBuilder.AddMenuEntry( FLevelEditorCommands::Get().BuildGeometryOnly );
-		MenuBuilder.AddMenuEntry( FLevelEditorCommands::Get().BuildGeometryOnly_OnlyCurrentLevel );
+		FToolMenuSection& Section = Menu->AddSection("LevelEditorGeometry", LOCTEXT( "GeometryHeading", "Geometry" ) );
+		Section.AddMenuEntry( FLevelEditorCommands::Get().BuildGeometryOnly );
+		Section.AddMenuEntry( FLevelEditorCommands::Get().BuildGeometryOnly_OnlyCurrentLevel );
 	}
-	MenuBuilder.EndSection();
 
-	MenuBuilder.BeginSection("LevelEditorNavigation", LOCTEXT( "NavigationHeading", "Navigation" ) );
 	{
-		MenuBuilder.AddMenuEntry( FLevelEditorCommands::Get().BuildPathsOnly );
+		FToolMenuSection& Section = Menu->AddSection("LevelEditorNavigation", LOCTEXT( "NavigationHeading", "Navigation" ) );
+		Section.AddMenuEntry( FLevelEditorCommands::Get().BuildPathsOnly );
 	}
-	MenuBuilder.EndSection();
 
-	MenuBuilder.BeginSection("LevelEditorLOD", LOCTEXT("LODHeading", "Hierarchical LOD"));
 	{
-		MenuBuilder.AddMenuEntry(FLevelEditorCommands::Get().BuildLODsOnly);
+		FToolMenuSection& Section = Menu->AddSection("LevelEditorLOD", LOCTEXT("LODHeading", "Hierarchical LOD"));
+		Section.AddMenuEntry(FLevelEditorCommands::Get().BuildLODsOnly);
 	}
-	MenuBuilder.EndSection();
-
-	MenuBuilder.BeginSection("LevelEditorTextureStreaming", LOCTEXT("TextureStreamingHeading", "Texture Streaming"));
+	
 	{
-		if (CVarStreamingUseNewMetrics.GetValueOnAnyThread() != 0) // There is not point of in building texture streaming data with the old system.
+		FToolMenuSection& Section = Menu->AddSection("LevelEditorTextureStreaming", LOCTEXT("TextureStreamingHeading", "Texture Streaming"));
+		Section.AddDynamicEntry("BuildTextureStreamingOnly", FNewToolMenuSectionDelegate::CreateLambda([](FToolMenuSection& InSection)
 		{
-			MenuBuilder.AddMenuEntry(FLevelEditorCommands::Get().BuildTextureStreamingOnly);
-		}
-		MenuBuilder.AddMenuEntry(FLevelEditorCommands::Get().BuildVirtualTextureOnly);
+			if (CVarStreamingUseNewMetrics.GetValueOnAnyThread() != 0) // There is no point of in building texture streaming data with the old system.
+			{
+				InSection.AddMenuEntry(FLevelEditorCommands::Get().BuildTextureStreamingOnly);
+			}
+		}));
+		Section.AddMenuEntry(FLevelEditorCommands::Get().BuildVirtualTextureOnly);
 	}
-	MenuBuilder.EndSection();
 
-	MenuBuilder.BeginSection("LevelEditorAutomation", LOCTEXT( "AutomationHeading", "Automation" ) );
+	
 	{
-		MenuBuilder.AddMenuEntry( 
+		FToolMenuSection& Section = Menu->AddSection("LevelEditorAutomation", LOCTEXT( "AutomationHeading", "Automation" ) );
+		Section.AddMenuEntry( 
 			FLevelEditorCommands::Get().BuildAndSubmitToSourceControl,
-			NAME_None,
 			TAttribute<FText>(),
 			TAttribute<FText>(),
 			FSlateIcon(FEditorStyle::GetStyleSetName(), "LevelEditor.Tabs.BuildAndSubmit")
 			);
 	}
-	MenuBuilder.EndSection();
 
 	// Map Check
-	MenuBuilder.BeginSection("LevelEditorVerification", LOCTEXT( "VerificationHeading", "Verification" ) );
 	{
-		MenuBuilder.AddMenuEntry( FLevelEditorCommands::Get().MapCheck, NAME_None, LOCTEXT("OpenMapCheck", "Map Check") );
+		FToolMenuSection& Section = Menu->AddSection("LevelEditorVerification", LOCTEXT( "VerificationHeading", "Verification" ) );
+		Section.AddMenuEntry( FLevelEditorCommands::Get().MapCheck, LOCTEXT("OpenMapCheck", "Map Check") );
 	}
-	MenuBuilder.EndSection();
 
 #undef LOCTEXT_NAMESPACE
-
-	return MenuBuilder.MakeWidget();
 }
 
-static void MakeMaterialQualityLevelMenu( FMenuBuilder& MenuBuilder )
+static void MakeMaterialQualityLevelMenu( UToolMenu* InMenu )
 {
-	MenuBuilder.BeginSection("LevelEditorMaterialQualityLevel", NSLOCTEXT( "LevelToolBarViewMenu", "MaterialQualityLevelHeading", "Material Quality Level" ) );
 	{
-		MenuBuilder.AddMenuEntry(FLevelEditorCommands::Get().MaterialQualityLevel_Low);
-		MenuBuilder.AddMenuEntry(FLevelEditorCommands::Get().MaterialQualityLevel_Medium);
-		MenuBuilder.AddMenuEntry(FLevelEditorCommands::Get().MaterialQualityLevel_High);
+		FToolMenuSection& Section = InMenu->AddSection("LevelEditorMaterialQualityLevel", NSLOCTEXT( "LevelToolBarViewMenu", "MaterialQualityLevelHeading", "Material Quality Level" ) );
+		Section.AddMenuEntry(FLevelEditorCommands::Get().MaterialQualityLevel_Low);
+		Section.AddMenuEntry(FLevelEditorCommands::Get().MaterialQualityLevel_Medium);
+		Section.AddMenuEntry(FLevelEditorCommands::Get().MaterialQualityLevel_High);
 	}
-	MenuBuilder.EndSection();
 }
 
-static void MakeShaderModelPreviewMenu( FMenuBuilder& MenuBuilder )
+static void MakeShaderModelPreviewMenu( UToolMenu* InMenu )
 {
 #define LOCTEXT_NAMESPACE "LevelToolBarViewMenu"
 
-	MenuBuilder.BeginSection("EditorPreviewMode", LOCTEXT("EditorPreviewModeDevices", "Preview Devices"));
+	FToolMenuSection& Section = InMenu->AddSection("EditorPreviewMode", LOCTEXT("EditorPreviewModeDevices", "Preview Devices"));
 
 	// SM5
-	MenuBuilder.AddMenuEntry(FLevelEditorCommands::Get().PreviewPlatformOverride_SM5);
+	Section.AddMenuEntry(FLevelEditorCommands::Get().PreviewPlatformOverride_SM5);
 
 	// Android
 	bool bAndroidBuildForES31 = false;
 	GConfig->GetBool(TEXT("/Script/AndroidRuntimeSettings.AndroidRuntimeSettings"), TEXT("bBuildForES31"), bAndroidBuildForES31, GEngineIni);
 	if (bAndroidBuildForES31)
 	{
-		MenuBuilder.AddMenuEntry(FLevelEditorCommands::Get().PreviewPlatformOverride_AndroidGLES31);
+		Section.AddMenuEntry(FLevelEditorCommands::Get().PreviewPlatformOverride_AndroidGLES31);
 	}
 
 	bool bAndroidSupportsVulkan = false;
 	GConfig->GetBool(TEXT("/Script/AndroidRuntimeSettings.AndroidRuntimeSettings"), TEXT("bSupportsVulkan"), bAndroidSupportsVulkan, GEngineIni);
 	if (bAndroidSupportsVulkan)
 	{
-		MenuBuilder.AddMenuEntry(FLevelEditorCommands::Get().PreviewPlatformOverride_AndroidVulkanES31);
+		Section.AddMenuEntry(FLevelEditorCommands::Get().PreviewPlatformOverride_AndroidVulkanES31);
 	}
 
 	bool bAndroidBuildForES2 = false;
 	GConfig->GetBool(TEXT("/Script/AndroidRuntimeSettings.AndroidRuntimeSettings"), TEXT("bBuildForES2"), bAndroidBuildForES2, GEngineIni);
 	if (bAndroidBuildForES2)
 	{
-    	MenuBuilder.AddMenuEntry(FLevelEditorCommands::Get().PreviewPlatformOverride_AndroidGLES2);
+    	Section.AddMenuEntry(FLevelEditorCommands::Get().PreviewPlatformOverride_AndroidGLES2);
 	}
 
 	// iOS
@@ -1871,78 +1903,81 @@ static void MakeShaderModelPreviewMenu( FMenuBuilder& MenuBuilder )
 	GConfig->GetBool(TEXT("/Script/IOSRuntimeSettings.IOSRuntimeSettings"), TEXT("bSupportsMetal"), bIOSSupportsMetal, GEngineIni);
 	if (bIOSSupportsMetal)
 	{
-		MenuBuilder.AddMenuEntry(FLevelEditorCommands::Get().PreviewPlatformOverride_IOSMetalES31);
+		Section.AddMenuEntry(FLevelEditorCommands::Get().PreviewPlatformOverride_IOSMetalES31);
 	}
 
 	// HTML5
-	MenuBuilder.AddMenuEntry(FLevelEditorCommands::Get().PreviewPlatformOverride_HTML5);
-
-    MenuBuilder.EndSection();
+	Section.AddMenuEntry(FLevelEditorCommands::Get().PreviewPlatformOverride_HTML5);
 
 #undef LOCTEXT_NAMESPACE
 }
 
-static void MakeScalabilityMenu( FMenuBuilder& MenuBuilder )
+static void MakeScalabilityMenu( UToolMenu* InMenu )
 {
-	MenuBuilder.AddWidget(SNew(SScalabilitySettings), FText(), true);
+	{
+		FToolMenuSection& Section = InMenu->AddSection("Section");
+		Section.AddEntry(FToolMenuEntry::InitWidget("ScalabilitySettings", SNew(SScalabilitySettings), FText(), true));
+	}
 }
 
-static void MakePreviewSettingsMenu( FMenuBuilder& MenuBuilder )
+static void MakePreviewSettingsMenu( UToolMenu* InMenu )
 {
 #define LOCTEXT_NAMESPACE "LevelToolBarViewMenu"
-
-	MenuBuilder.BeginSection("LevelEditorPreview", LOCTEXT("PreviewHeading", "Previewing"));
 	{
-		MenuBuilder.AddMenuEntry(FLevelEditorCommands::Get().DrawBrushMarkerPolys);
-		MenuBuilder.AddMenuEntry(FLevelEditorCommands::Get().OnlyLoadVisibleInPIE);
-		MenuBuilder.AddMenuEntry(FLevelEditorCommands::Get().ToggleParticleSystemLOD);
-		MenuBuilder.AddMenuEntry(FLevelEditorCommands::Get().ToggleParticleSystemHelpers);
-		MenuBuilder.AddMenuEntry(FLevelEditorCommands::Get().ToggleFreezeParticleSimulation);
-		MenuBuilder.AddMenuEntry(FLevelEditorCommands::Get().ToggleLODViewLocking);
-		MenuBuilder.AddMenuEntry(FLevelEditorCommands::Get().LevelStreamingVolumePrevis);
+		FToolMenuSection& Section = InMenu->AddSection("LevelEditorPreview", LOCTEXT("PreviewHeading", "Previewing"));
+		Section.AddMenuEntry(FLevelEditorCommands::Get().DrawBrushMarkerPolys);
+		Section.AddMenuEntry(FLevelEditorCommands::Get().OnlyLoadVisibleInPIE);
+		Section.AddMenuEntry(FLevelEditorCommands::Get().ToggleParticleSystemLOD);
+		Section.AddMenuEntry(FLevelEditorCommands::Get().ToggleParticleSystemHelpers);
+		Section.AddMenuEntry(FLevelEditorCommands::Get().ToggleFreezeParticleSimulation);
+		Section.AddMenuEntry(FLevelEditorCommands::Get().ToggleLODViewLocking);
+		Section.AddMenuEntry(FLevelEditorCommands::Get().LevelStreamingVolumePrevis);
 	}
-	MenuBuilder.EndSection();
 #undef LOCTEXT_NAMESPACE
 }
 
 #if WITH_LIVE_CODING
-TSharedRef< SWidget > FLevelEditorToolBar::GenerateCompileMenuContent( TSharedRef<FUICommandList> InCommandList )
+void FLevelEditorToolBar::RegisterCompileMenu()
 {
 #define LOCTEXT_NAMESPACE "LevelToolBarCompileMenu"
 
-	const bool bShouldCloseWindowAfterMenuSelection = true;
-	FMenuBuilder MenuBuilder( bShouldCloseWindowAfterMenuSelection, InCommandList );
+	UToolMenu* Menu = UToolMenus::Get()->RegisterMenu("LevelEditor.LevelEditorToolBar.CompileComboButton");
 
-	MenuBuilder.BeginSection("LiveCodingMode", LOCTEXT( "LiveCodingMode", "General" ) );
 	{
-		MenuBuilder.AddMenuEntry( FLevelEditorCommands::Get().LiveCoding_Enable );
+		FToolMenuSection& Section = Menu->AddSection("LiveCodingMode", LOCTEXT( "LiveCodingMode", "General" ) );
+		Section.AddMenuEntry( FLevelEditorCommands::Get().LiveCoding_Enable );
 	}
-	MenuBuilder.EndSection();
 
-	MenuBuilder.BeginSection("LiveCodingActions", LOCTEXT( "LiveCodingActions", "Actions" ) );
 	{
-		MenuBuilder.AddMenuEntry( FLevelEditorCommands::Get().LiveCoding_StartSession );
-		MenuBuilder.AddMenuEntry( FLevelEditorCommands::Get().LiveCoding_ShowConsole );
-		MenuBuilder.AddMenuEntry( FLevelEditorCommands::Get().LiveCoding_Settings );
+		FToolMenuSection& Section = Menu->AddSection("LiveCodingActions", LOCTEXT( "LiveCodingActions", "Actions" ) );
+		Section.AddMenuEntry( FLevelEditorCommands::Get().LiveCoding_StartSession );
+		Section.AddMenuEntry( FLevelEditorCommands::Get().LiveCoding_ShowConsole );
+		Section.AddMenuEntry( FLevelEditorCommands::Get().LiveCoding_Settings );
 	}
-	MenuBuilder.EndSection();
-
-	return MenuBuilder.MakeWidget();
 
 #undef LOCTEXT_NAMESPACE
 }
 #endif
 
-TSharedRef< SWidget > FLevelEditorToolBar::GenerateQuickSettingsMenu( TSharedRef<FUICommandList> InCommandList )
+TSharedRef< SWidget > FLevelEditorToolBar::GenerateQuickSettingsMenu(TSharedRef<FUICommandList> InCommandList, TWeakPtr<SLevelEditor> InLevelEditor)
 {
-#define LOCTEXT_NAMESPACE "LevelToolBarViewMenu"
-
 	// Get all menu extenders for this context menu from the level editor module
 	FLevelEditorModule& LevelEditorModule = FModuleManager::GetModuleChecked<FLevelEditorModule>( TEXT("LevelEditor") );
 	TSharedPtr<FExtender> MenuExtender = LevelEditorModule.AssembleExtenders(InCommandList, LevelEditorModule.GetAllLevelEditorToolbarViewMenuExtenders());
 
-	const bool bShouldCloseWindowAfterMenuSelection = true;
-	FMenuBuilder MenuBuilder( bShouldCloseWindowAfterMenuSelection, InCommandList, MenuExtender );
+	FToolMenuContext MenuContext(InCommandList, MenuExtender);
+	ULevelEditorMenuContext* LevelEditorMenuContext = NewObject<ULevelEditorMenuContext>();
+	LevelEditorMenuContext->LevelEditor = InLevelEditor;
+	MenuContext.AddObject(LevelEditorMenuContext);
+
+	return UToolMenus::Get()->GenerateWidget("LevelEditor.LevelEditorToolBar.LevelToolbarQuickSettings", MenuContext);
+}
+
+void FLevelEditorToolBar::RegisterQuickSettingsMenu()
+{
+#define LOCTEXT_NAMESPACE "LevelToolBarViewMenu"
+
+	UToolMenu* Menu = UToolMenus::Get()->RegisterMenu("LevelEditor.LevelEditorToolBar.LevelToolbarQuickSettings");
 
 	struct Local
 	{
@@ -1952,56 +1987,60 @@ TSharedRef< SWidget > FLevelEditorToolBar::GenerateQuickSettingsMenu( TSharedRef
 		}
 	};
 
-	MenuBuilder.BeginSection("ProjectSettingsSection", LOCTEXT("ProjectSettings","Game Specific Settings") );
 	{
+		FToolMenuSection& Section = Menu->AddSection("ProjectSettingsSection", LOCTEXT("ProjectSettings", "Game Specific Settings"));
 
-		MenuBuilder.AddMenuEntry(FLevelEditorCommands::Get().WorldProperties);
+		Section.AddMenuEntry(FLevelEditorCommands::Get().WorldProperties);
 
-		MenuBuilder.AddMenuEntry(
+		Section.AddMenuEntry(
+			"ProjectSettings",
 			LOCTEXT("ProjectSettingsMenuLabel", "Project Settings..."),
 			LOCTEXT("ProjectSettingsMenuToolTip", "Change the settings of the currently loaded project"),
 			FSlateIcon(FEditorStyle::GetStyleSetName(), "ProjectSettings.TabIcon"),
 			FUIAction(FExecuteAction::CreateStatic(&Local::OpenSettings, FName("Project"), FName("Project"), FName("General")))
 			);
 
-		if (IModularFeatures::Get().IsModularFeatureAvailable(EditorFeatures::PluginsEditor))
+		Section.AddDynamicEntry("PluginsEditor", FNewToolMenuDelegateLegacy::CreateLambda([](FMenuBuilder& InMenuBuilder, UToolMenu* InMenu)
 		{
-			FGlobalTabmanager::Get()->PopulateTabSpawnerMenu(MenuBuilder, "PluginsEditor");
-		}
+			if (IModularFeatures::Get().IsModularFeatureAvailable(EditorFeatures::PluginsEditor))
+			{
+				FGlobalTabmanager::Get()->PopulateTabSpawnerMenu(InMenuBuilder, "PluginsEditor");
+			}
+		}));
 	}
-	MenuBuilder.EndSection();
 
-	MenuBuilder.BeginSection("LevelEditorSelection", LOCTEXT("SelectionHeading","Selection") );
 	{
-		MenuBuilder.AddMenuEntry( FLevelEditorCommands::Get().AllowTranslucentSelection );
-		MenuBuilder.AddMenuEntry( FLevelEditorCommands::Get().AllowGroupSelection );
-		MenuBuilder.AddMenuEntry( FLevelEditorCommands::Get().StrictBoxSelect );
-		MenuBuilder.AddMenuEntry( FLevelEditorCommands::Get().TransparentBoxSelect );
-		MenuBuilder.AddMenuEntry( FLevelEditorCommands::Get().ShowTransformWidget );
+		FToolMenuSection& Section = Menu->AddSection("LevelEditorSelection", LOCTEXT("SelectionHeading","Selection") );
+		Section.AddMenuEntry( FLevelEditorCommands::Get().AllowTranslucentSelection );
+		Section.AddMenuEntry( FLevelEditorCommands::Get().AllowGroupSelection );
+		Section.AddMenuEntry( FLevelEditorCommands::Get().StrictBoxSelect );
+		Section.AddMenuEntry( FLevelEditorCommands::Get().TransparentBoxSelect );
+		Section.AddMenuEntry( FLevelEditorCommands::Get().ShowTransformWidget );
 	}
-	MenuBuilder.EndSection();
 
-	MenuBuilder.BeginSection("LevelEditorScalability", LOCTEXT("ScalabilityHeading", "Scalability") );
 	{
-		MenuBuilder.AddSubMenu(
+		FToolMenuSection& Section = Menu->AddSection("LevelEditorScalability", LOCTEXT("ScalabilityHeading", "Scalability") );
+		Section.AddSubMenu(
+			"Scalability",
 			LOCTEXT( "ScalabilitySubMenu", "Engine Scalability Settings" ),
 			LOCTEXT( "ScalabilitySubMenu_ToolTip", "Open the engine scalability settings" ),
-			FNewMenuDelegate::CreateStatic( &MakeScalabilityMenu ) );
+			FNewToolMenuDelegate::CreateStatic( &MakeScalabilityMenu ) );
 
-		MenuBuilder.AddSubMenu(
+		Section.AddSubMenu(
+			"MaterialQualityLevel",
 			LOCTEXT( "MaterialQualityLevelSubMenu", "Material Quality Level" ),
 			LOCTEXT( "MaterialQualityLevelSubMenu_ToolTip", "Sets the value of the CVar \"r.MaterialQualityLevel\" (low=0, high=1, medium=2). This affects materials via the QualitySwitch material expression." ),
-			FNewMenuDelegate::CreateStatic( &MakeMaterialQualityLevelMenu ) );
+			FNewToolMenuDelegate::CreateStatic( &MakeMaterialQualityLevelMenu ) );
 
-		MenuBuilder.AddSubMenu(
+		Section.AddSubMenu(
+			"FeatureLevelPreview",
 			LOCTEXT("FeatureLevelPreviewSubMenu", "Preview Rendering Level"),
 			LOCTEXT("FeatureLevelPreviewSubMenu_ToolTip", "Sets the rendering level used by the main editor"),
-			FNewMenuDelegate::CreateStatic(&MakeShaderModelPreviewMenu));
+			FNewToolMenuDelegate::CreateStatic(&MakeShaderModelPreviewMenu));
 	}
-	MenuBuilder.EndSection();
 
-	MenuBuilder.BeginSection("LevelEditorAudio", LOCTEXT("AudioHeading", "Real Time Audio") );
 	{
+		FToolMenuSection& Section = Menu->AddSection("LevelEditorAudio", LOCTEXT("AudioHeading", "Real Time Audio") );
 		TSharedRef<SWidget> VolumeItem = SNew(SHorizontalBox)
 											+SHorizontalBox::Slot()
 											.FillWidth(0.9f)
@@ -2017,13 +2056,12 @@ TSharedRef< SWidget > FLevelEditorToolBar::GenerateQuickSettingsMenu( TSharedRef
 											+SHorizontalBox::Slot()
 											.FillWidth(0.1f);
 
-		MenuBuilder.AddWidget(VolumeItem, LOCTEXT("VolumeControlLabel","Volume"));
+		Section.AddEntry(FToolMenuEntry::InitWidget("Volume", VolumeItem, LOCTEXT("VolumeControlLabel","Volume")));
 	}
-	MenuBuilder.EndSection();
 
-	MenuBuilder.BeginSection( "Snapping", LOCTEXT("SnappingHeading","Snapping") );
 	{
-		MenuBuilder.AddMenuEntry( FLevelEditorCommands::Get().EnableActorSnap );
+		FToolMenuSection& Section = Menu->AddSection( "Snapping", LOCTEXT("SnappingHeading","Snapping") );
+		Section.AddMenuEntry( FLevelEditorCommands::Get().EnableActorSnap );
 		TSharedRef<SWidget> SnapItem = 
 		SNew(SHorizontalBox)
 	          +SHorizontalBox::Slot()
@@ -2036,35 +2074,35 @@ TSharedRef< SWidget > FLevelEditorToolBar::GenerateQuickSettingsMenu( TSharedRef
 	          ]
 	          +SHorizontalBox::Slot()
 	          .FillWidth(0.1f);
-		MenuBuilder.AddWidget(SnapItem, LOCTEXT("ActorSnapLabel","Distance"));
+		Section.AddEntry(FToolMenuEntry::InitWidget("Snap", SnapItem, LOCTEXT("ActorSnapLabel", "Distance")));
 
-		MenuBuilder.AddMenuEntry( FLevelEditorCommands::Get().ToggleSocketSnapping );
-		MenuBuilder.AddMenuEntry( FLevelEditorCommands::Get().EnableVertexSnap );
+		Section.AddMenuEntry( FLevelEditorCommands::Get().ToggleSocketSnapping );
+		Section.AddMenuEntry( FLevelEditorCommands::Get().EnableVertexSnap );
 	}
-	MenuBuilder.EndSection();
 
-	MenuBuilder.BeginSection("LevelEditorViewport", LOCTEXT("ViewportHeading", "Viewport") );
 	{
-		MenuBuilder.AddMenuEntry( FLevelEditorCommands::Get().ToggleHideViewportUI );
+		FToolMenuSection& Section = Menu->AddSection("LevelEditorViewport", LOCTEXT("ViewportHeading", "Viewport") );
+		Section.AddMenuEntry( FLevelEditorCommands::Get().ToggleHideViewportUI );
 
-		MenuBuilder.AddSubMenu( LOCTEXT("PreviewMenu", "Previewing"), LOCTEXT("PreviewMenuTooltip","Game Preview Settings"), FNewMenuDelegate::CreateStatic( &MakePreviewSettingsMenu ) );
+		Section.AddSubMenu( "Preview", LOCTEXT("PreviewMenu", "Previewing"), LOCTEXT("PreviewMenuTooltip","Game Preview Settings"), FNewToolMenuDelegate::CreateStatic( &MakePreviewSettingsMenu ) );
 	}
-	MenuBuilder.EndSection();
 
 #undef LOCTEXT_NAMESPACE
-
-	return MenuBuilder.MakeWidget();
 }
 
 
-TSharedRef< SWidget > FLevelEditorToolBar::GenerateSourceControlMenu(TSharedRef<FUICommandList> InCommandList)
+TSharedRef< SWidget > FLevelEditorToolBar::GenerateSourceControlMenu(TSharedRef<FUICommandList> InCommandList, TWeakPtr< SLevelEditor > InLevelEditor)
 {
 	// Get all menu extenders for this context menu from the level editor module
 	FLevelEditorModule& LevelEditorModule = FModuleManager::GetModuleChecked<FLevelEditorModule>(TEXT("LevelEditor"));
 	TSharedPtr<FExtender> MenuExtender = LevelEditorModule.AssembleExtenders(InCommandList, LevelEditorModule.GetAllLevelEditorToolbarSourceControlMenuExtenders());
 
-	FToolMenuContext Context(InCommandList, MenuExtender);
-	return UToolMenus::Get()->GenerateWidget("LevelEditor.LevelEditorToolBar.SourceControl", Context);
+	FToolMenuContext MenuContext(InCommandList, MenuExtender);
+	ULevelEditorMenuContext* LevelEditorMenuContext = NewObject<ULevelEditorMenuContext>();
+	LevelEditorMenuContext->LevelEditor = InLevelEditor;
+	MenuContext.AddObject(LevelEditorMenuContext);
+
+	return UToolMenus::Get()->GenerateWidget("LevelEditor.LevelEditorToolBar.SourceControl", MenuContext);
 }
 
 void FLevelEditorToolBar::RegisterSourceControlMenu()
@@ -2122,35 +2160,54 @@ TSharedRef< SWidget > FLevelEditorToolBar::GenerateOpenBlueprintMenuContent( TSh
 {
 #define LOCTEXT_NAMESPACE "LevelToolBarViewMenu"
 
+	FLevelEditorModule& LevelEditorModule = FModuleManager::LoadModuleChecked<FLevelEditorModule>("LevelEditor");
+	TSharedPtr<FExtender> MenuExtender = FExtender::Combine(LevelEditorModule.GetAllLevelEditorToolbarBlueprintsMenuExtenders());
+
+	FToolMenuContext MenuContext(InCommandList, MenuExtender);
+	ULevelEditorMenuContext* LevelEditorMenuContext = NewObject<ULevelEditorMenuContext>();
+	LevelEditorMenuContext->LevelEditor = InLevelEditor;
+	MenuContext.AddObject(LevelEditorMenuContext);
+
+	return UToolMenus::Get()->GenerateWidget("LevelEditor.LevelEditorToolBar.OpenBlueprint", MenuContext);
+}
+
+void FLevelEditorToolBar::RegisterOpenBlueprintMenu()
+{
+#define LOCTEXT_NAMESPACE "LevelToolBarViewMenu"
+	UToolMenu* Menu = UToolMenus::Get()->RegisterMenu("LevelEditor.LevelEditorToolBar.OpenBlueprint");
+
 	struct FBlueprintMenus
 	{
 		/** Generates a sub-level Blueprints sub-menu */
-		static void MakeSubLevelsMenu(FMenuBuilder& InMenuBuilder, TWeakPtr< SLevelEditor > InLvlEditor)
+		static void MakeSubLevelsMenu(UToolMenu* InMenu)
 		{
-			FSlateIcon EditBP(FEditorStyle::Get().GetStyleSetName(), TEXT("LevelEditor.OpenLevelBlueprint"));
-
-			InMenuBuilder.BeginSection(NAME_None, LOCTEXT("SubLevelsHeading", "Sub-Level Blueprints"));
+			ULevelEditorMenuContext* Context = InMenu->FindContext<ULevelEditorMenuContext>();
+			if (Context && Context->LevelEditor.IsValid())
 			{
-				UWorld* World = InLvlEditor.Pin()->GetWorld();
-				// Sort the levels alphabetically 
-				TArray<ULevel*> SortedLevels = World->GetLevels();
-				Algo::Sort(SortedLevels, LevelEditorActionHelpers::FLevelSortByName());
+				FSlateIcon EditBP(FEditorStyle::Get().GetStyleSetName(), TEXT("LevelEditor.OpenLevelBlueprint"));
 
-				for (ULevel* const Level : SortedLevels)
 				{
-					if (Level != NULL && Level->GetOutermost() != NULL && !Level->IsPersistentLevel())
+					FToolMenuSection& Section = InMenu->AddSection("SubLevels", LOCTEXT("SubLevelsHeading", "Sub-Level Blueprints"));
+					UWorld* World = Context->LevelEditor.Pin()->GetWorld();
+					// Sort the levels alphabetically 
+					TArray<ULevel*> SortedLevels = World->GetLevels();
+					Algo::Sort(SortedLevels, LevelEditorActionHelpers::FLevelSortByName());
+
+					for (ULevel* const Level : SortedLevels)
 					{
-						FUIAction UIAction
-						(
-							FExecuteAction::CreateStatic(&FLevelEditorToolBar::OnOpenSubLevelBlueprint, Level)
-						);
+						if (Level != NULL && Level->GetOutermost() != NULL && !Level->IsPersistentLevel())
+						{
+							FUIAction UIAction
+							(
+								FExecuteAction::CreateStatic(&FLevelEditorToolBar::OnOpenSubLevelBlueprint, Level)
+							);
 
-						FText DisplayName = FText::Format(LOCTEXT("SubLevelBlueprintItem", "Edit {0}"), FText::FromString(FPaths::GetCleanFilename(Level->GetOutermost()->GetName())));
-						InMenuBuilder.AddMenuEntry(DisplayName, FText::GetEmpty(), EditBP, UIAction);
+							FText DisplayName = FText::Format(LOCTEXT("SubLevelBlueprintItem", "Edit {0}"), FText::FromString(FPaths::GetCleanFilename(Level->GetOutermost()->GetName())));
+							Section.AddMenuEntry(NAME_None, DisplayName, FText::GetEmpty(), EditBP, UIAction);
+						}
 					}
 				}
 			}
-			InMenuBuilder.EndSection();
 		}
 
 		/** Handle BP being selected from popup picker */
@@ -2165,7 +2222,7 @@ TSharedRef< SWidget > FLevelEditorToolBar::GenerateOpenBlueprintMenuContent( TSh
 
 
 		/** Generates 'open blueprint' sub-menu */
-		static void MakeOpenBPClassMenu(FMenuBuilder& InMenuBuilder)
+		static void MakeOpenBPClassMenu(UToolMenu* InMenu)
 		{
 			FContentBrowserModule& ContentBrowserModule = FModuleManager::Get().LoadModuleChecked<FContentBrowserModule>(TEXT("ContentBrowser"));
 
@@ -2188,80 +2245,77 @@ TSharedRef< SWidget > FLevelEditorToolBar::GenerateOpenBlueprintMenuContent( TSh
 				];
 		
 
-			InMenuBuilder.BeginSection(NAME_None, LOCTEXT("BrowseHeader", "Browse"));
 			{
-				InMenuBuilder.AddWidget(Widget, FText::GetEmpty());
+				FToolMenuSection& Section = InMenu->AddSection("Browse", LOCTEXT("BrowseHeader", "Browse"));
+				Section.AddEntry(FToolMenuEntry::InitWidget("PickClassWidget", Widget, FText::GetEmpty()));
 			}
-			InMenuBuilder.EndSection();
 		}
 	};
 
-	FLevelEditorModule& LevelEditorModule = FModuleManager::LoadModuleChecked<FLevelEditorModule>("LevelEditor");
-	TSharedPtr<FExtender> Extender = FExtender::Combine(LevelEditorModule.GetAllLevelEditorToolbarBlueprintsMenuExtenders());
-
-	const bool bShouldCloseWindowAfterMenuSelection = true;
-	FMenuBuilder MenuBuilder( bShouldCloseWindowAfterMenuSelection, InCommandList, Extender);
-
-	MenuBuilder.BeginSection("BlueprintClass", LOCTEXT("BlueprintClass", "Blueprint Class"));
 	{
+		FToolMenuSection& Section = Menu->AddSection("BlueprintClass", LOCTEXT("BlueprintClass", "Blueprint Class"));
+
 		// Create a blank BP
-		MenuBuilder.AddMenuEntry(FLevelEditorCommands::Get().CreateBlankBlueprintClass);
+		Section.AddMenuEntry(FLevelEditorCommands::Get().CreateBlankBlueprintClass);
 
 		// Convert selection to BP
-		MenuBuilder.AddMenuEntry(FLevelEditorCommands::Get().ConvertSelectionToBlueprintViaHarvest);
-		MenuBuilder.AddMenuEntry(FLevelEditorCommands::Get().ConvertSelectionToBlueprintViaSubclass);
+		Section.AddMenuEntry(FLevelEditorCommands::Get().ConvertSelectionToBlueprintViaHarvest);
+		Section.AddMenuEntry(FLevelEditorCommands::Get().ConvertSelectionToBlueprintViaSubclass);
 
 		// Open an existing Blueprint Class...
 		FSlateIcon OpenBPIcon(FEditorStyle::GetStyleSetName(), "LevelEditor.OpenClassBlueprint");
-		MenuBuilder.AddSubMenu(
+		Section.AddSubMenu(
+			"OpenBlueprintClass",
 			LOCTEXT("OpenBlueprintClassSubMenu", "Open Blueprint Class..."),
 			LOCTEXT("OpenBlueprintClassSubMenu_ToolTip", "Open an existing Blueprint Class in this project"),
-			FNewMenuDelegate::CreateStatic(&FBlueprintMenus::MakeOpenBPClassMenu),
+			FNewToolMenuDelegate::CreateStatic(&FBlueprintMenus::MakeOpenBPClassMenu),
 			false,
 			OpenBPIcon);
 	}
-	MenuBuilder.EndSection();
 
-	MenuBuilder.BeginSection("LevelScriptBlueprints", LOCTEXT("LevelScriptBlueprints", "Level Blueprints"));
 	{
-		MenuBuilder.AddMenuEntry( FLevelEditorCommands::Get().OpenLevelBlueprint );
+		FToolMenuSection& Section = Menu->AddSection("LevelScriptBlueprints", LOCTEXT("LevelScriptBlueprints", "Level Blueprints"));
+		Section.AddMenuEntry( FLevelEditorCommands::Get().OpenLevelBlueprint );
 
-		// If there are any sub-levels, display the sub-menu. A single level means there is only the persistent level
-		UWorld* World = InLevelEditor.Pin()->GetWorld();
-		if(World->GetNumLevels() > 1)
+		Section.AddDynamicEntry("SubLevels", FNewToolMenuSectionDelegate::CreateLambda([](FToolMenuSection& InSection)
 		{
-			MenuBuilder.AddSubMenu(
-				LOCTEXT( "SubLevelsSubMenu", "Sub-Levels" ),
-				LOCTEXT( "SubLevelsSubMenu_ToolTip", "Shows available sub-level Blueprints that can be edited." ),
-				FNewMenuDelegate::CreateStatic( &FBlueprintMenus::MakeSubLevelsMenu, InLevelEditor ), 
-				FUIAction(), NAME_None, EUserInterfaceActionType::Button, false, FSlateIcon(FEditorStyle::Get().GetStyleSetName(), TEXT("LevelEditor.OpenLevelBlueprint")) );
-		}
+			ULevelEditorMenuContext* Context = InSection.FindContext<ULevelEditorMenuContext>();
+			if (Context && Context->LevelEditor.IsValid())
+			{
+				// If there are any sub-levels, display the sub-menu. A single level means there is only the persistent level
+				UWorld* World = Context->LevelEditor.Pin()->GetWorld();
+				if (World->GetNumLevels() > 1)
+				{
+					InSection.AddSubMenu(
+						"SubLevels",
+						LOCTEXT("SubLevelsSubMenu", "Sub-Levels"),
+						LOCTEXT("SubLevelsSubMenu_ToolTip", "Shows available sub-level Blueprints that can be edited."),
+						FNewToolMenuDelegate::CreateStatic(&FBlueprintMenus::MakeSubLevelsMenu),
+						FUIAction(), EUserInterfaceActionType::Button, false, FSlateIcon(FEditorStyle::Get().GetStyleSetName(), TEXT("LevelEditor.OpenLevelBlueprint")));
+				}
+			}
+		}));
 	}
-	MenuBuilder.EndSection();
 
-	MenuBuilder.BeginSection("ProjectSettingsClasses", LOCTEXT("ProjectSettingsClasses", "Project Settings"));
 	{
-		// If source control is enabled, queue up a query to the status of the config file so it is (hopefully) ready before we get to the sub-menu
-		if(ISourceControlModule::Get().IsEnabled())
-		{
-			FString ConfigFilePath = FPaths::ConvertRelativePathToFull(FString::Printf(TEXT("%sDefaultEngine.ini"), *FPaths::SourceConfigDir()));
-
-			// note: calling QueueStatusUpdate often does not spam status updates as an internal timer prevents this
-			ISourceControlModule::Get().QueueStatusUpdate(ConfigFilePath);
-		}
-		LevelEditorActionHelpers::CreateGameModeSubMenu(MenuBuilder, InCommandList, InLevelEditor, true);
+		FToolMenuSection& Section = Menu->AddSection("ProjectSettingsClasses", LOCTEXT("ProjectSettingsClasses", "Project Settings"));
+		LevelEditorActionHelpers::CreateGameModeSubMenu(Section, "ProjectSettingsClasses", true);
 	}
-	MenuBuilder.EndSection();
 
-	MenuBuilder.BeginSection("WorldSettingsClasses", LOCTEXT("WorldSettingsClasses", "World Override"));
 	{
-		LevelEditorActionHelpers::CreateGameModeSubMenu(MenuBuilder, InCommandList, InLevelEditor, false);
+		FToolMenuSection& Section = Menu->AddSection("WorldSettingsClasses", LOCTEXT("WorldSettingsClasses", "World Override"));
+		LevelEditorActionHelpers::CreateGameModeSubMenu(Section, "WorldSettingsClasses", false);
 	}
-	MenuBuilder.EndSection();
 
-#undef LOCTEXT_NAMESPACE
+	// If source control is enabled, queue up a query to the status of the config file so it is (hopefully) ready before we get to the sub-menu
+	if(ISourceControlModule::Get().IsEnabled())
+	{
+		FString ConfigFilePath = FPaths::ConvertRelativePathToFull(FString::Printf(TEXT("%sDefaultEngine.ini"), *FPaths::SourceConfigDir()));
 
-	return MenuBuilder.MakeWidget();
+		// note: calling QueueStatusUpdate often does not spam status updates as an internal timer prevents this
+		ISourceControlModule::Get().QueueStatusUpdate(ConfigFilePath);
+	}
+#undef LOCTEXT_NAMESPACE
 }
 
 void FLevelEditorToolBar::OnOpenSubLevelBlueprint( ULevel* InLevel )
@@ -2278,14 +2332,16 @@ void FLevelEditorToolBar::OnOpenSubLevelBlueprint( ULevel* InLevel )
 	}
 }
 
-TSharedRef< SWidget > FLevelEditorToolBar::GenerateCinematicsMenuContent( TSharedRef<FUICommandList> InCommandList, TWeakPtr<SLevelEditor> LevelEditorWeakPtr )
+TSharedRef< SWidget > FLevelEditorToolBar::GenerateCinematicsMenuContent(TSharedRef<FUICommandList> InCommandList, TWeakPtr<SLevelEditor> InLevelEditor)
 {
 	FLevelEditorModule& LevelEditorModule = FModuleManager::LoadModuleChecked<FLevelEditorModule>("LevelEditor");
-	FToolMenuContext Context(InCommandList, FExtender::Combine(LevelEditorModule.GetAllLevelEditorToolbarCinematicsMenuExtenders()));
+
+	FToolMenuContext MenuContext(InCommandList, FExtender::Combine(LevelEditorModule.GetAllLevelEditorToolbarCinematicsMenuExtenders()));
 	ULevelEditorMenuContext* LevelEditorMenuContext = NewObject<ULevelEditorMenuContext>();
-	LevelEditorMenuContext->LevelEditor = LevelEditorWeakPtr;
-	Context.AddObject(LevelEditorMenuContext);
-	return UToolMenus::Get()->GenerateWidget("LevelEditor.LevelEditorToolBar.Cinematics", Context);
+	LevelEditorMenuContext->LevelEditor = InLevelEditor;
+	MenuContext.AddObject(LevelEditorMenuContext);
+
+	return UToolMenus::Get()->GenerateWidget("LevelEditor.LevelEditorToolBar.Cinematics", MenuContext);
 }
 
 void FLevelEditorToolBar::RegisterCinematicsMenu()
diff --git a/Engine/Source/Editor/LevelEditor/Private/LevelEditorToolBar.h b/Engine/Source/Editor/LevelEditor/Private/LevelEditorToolBar.h
--- a/Engine/Source/Editor/LevelEditor/Private/LevelEditorToolBar.h
+++ b/Engine/Source/Editor/LevelEditor/Private/LevelEditorToolBar.h
@@ -34,51 +34,35 @@ protected:
 	 *
 	 * @return	Menu content widget
 	 */
-	static TSharedRef< SWidget > GenerateBuildMenuContent( TSharedRef<FUICommandList> InCommandList );
-
-#if WITH_LIVE_CODING
-	/**
-	 * Generates menu content for the compile combo button drop down menu
-	 *
-	 * @return	Menu content widget
-	 */
-	static TSharedRef< SWidget > GenerateCompileMenuContent( TSharedRef<FUICommandList> InCommandList );
-#endif
-
-	/**
-	 * Generates menu content for the create actor combo button drop down menu
-	 *
-	 * @return	Menu content widget
-	 */
-	static TSharedRef< SWidget > GenerateCreateContent( TSharedRef<FUICommandList> InCommandList );
+	static TSharedRef< SWidget > GenerateBuildMenuContent(TSharedRef<FUICommandList> InCommandList, TWeakPtr<SLevelEditor> InLevelEditor);
 
 	/**
 	 * Generates menu content for the quick settings combo button drop down menu
 	 *
 	 * @return	Menu content widget
 	 */
-	static TSharedRef< SWidget > GenerateQuickSettingsMenu( TSharedRef<FUICommandList> InCommandList );
+	static TSharedRef< SWidget > GenerateQuickSettingsMenu(TSharedRef<FUICommandList> InCommandList, TWeakPtr<SLevelEditor> InLevelEditor);
 
 	/**
 	 * Generates menu content for the source control combo button drop down menu
 	 *
 	 * @return	Menu content widget
 	 */
-	static TSharedRef< SWidget > GenerateSourceControlMenu(TSharedRef<FUICommandList> InCommandList);
+	static TSharedRef< SWidget > GenerateSourceControlMenu(TSharedRef<FUICommandList> InCommandList, TWeakPtr<SLevelEditor> InLevelEditor);
 
 	/**
 	 * Generates menu content for the compile combo button drop down menu
 	 *
 	 * @return	Menu content widget
 	 */
-	static TSharedRef< SWidget > GenerateOpenBlueprintMenuContent( TSharedRef<FUICommandList> InCommandList, TWeakPtr< SLevelEditor > InLevelEditor );
+	static TSharedRef< SWidget > GenerateOpenBlueprintMenuContent(TSharedRef<FUICommandList> InCommandList, TWeakPtr<SLevelEditor> InLevelEditor);
 
 	/**
 	 * Generates menu content for the Cinematics combo button drop down menu
 	 *
 	 * @return	Menu content widget
 	 */
-	static TSharedRef< SWidget > GenerateCinematicsMenuContent( TSharedRef<FUICommandList> InCommandList, TWeakPtr<SLevelEditor> LevelEditorWeakPtr );
+	static TSharedRef< SWidget > GenerateCinematicsMenuContent(TSharedRef<FUICommandList> InCommandList, TWeakPtr<SLevelEditor> InLevelEditor);
 
 	/**
 	 * Delegate for actor selection within the Cinematics popup menu's SceneOutliner.
@@ -97,4 +81,15 @@ private:
 
 	static void RegisterSourceControlMenu();
 	static void RegisterCinematicsMenu();
+	static void RegisterBuildMenu();
+
+#if WITH_LIVE_CODING
+	/**
+	 * Generates menu content for the compile combo button drop down menu
+	 */
+	static void RegisterCompileMenu();
+#endif
+
+	static void RegisterQuickSettingsMenu();
+	static void RegisterOpenBlueprintMenu();
 };
