Fix mac build warning
#rnx

diff --git a/Engine/Source/Editor/LevelEditor/Private/LevelEditorToolBar.cpp b/Engine/Source/Editor/LevelEditor/Private/LevelEditorToolBar.cpp
--- a/Engine/Source/Editor/LevelEditor/Private/LevelEditorToolBar.cpp
+++ b/Engine/Source/Editor/LevelEditor/Private/LevelEditorToolBar.cpp
@@ -1561,18 +1561,18 @@ void FLevelEditorToolBar::RegisterBuildMenu()
 	{
 	public:
 
-		static void RegisterMenus(const FName BaseMenuName)
+		static void RegisterMenus(const FName InBaseMenuName)
 		{
-			FLightingMenus::RegisterLightingQualityMenu(BaseMenuName);
-			FLightingMenus::RegisterLightingInfoMenu(BaseMenuName);
+			FLightingMenus::RegisterLightingQualityMenu(InBaseMenuName);
+			FLightingMenus::RegisterLightingInfoMenu(InBaseMenuName);
 		}
 
 	private:
 
 		/** Generates a lighting quality sub-menu */
-		static void RegisterLightingQualityMenu(const FName BaseMenuName)
+		static void RegisterLightingQualityMenu(const FName InBaseMenuName)
 		{
-			UToolMenu* SubMenu = UToolMenus::Get()->RegisterMenu(UToolMenus::JoinMenuPaths(BaseMenuName, "LightingQuality"));
+			UToolMenu* SubMenu = UToolMenus::Get()->RegisterMenu(UToolMenus::JoinMenuPaths(InBaseMenuName, "LightingQuality"));
 
 			{
 				FToolMenuSection& Section = SubMenu->AddSection("LevelEditorBuildLightingQuality", LOCTEXT( "LightingQualityHeading", "Quality Level" ) );
@@ -1584,9 +1584,9 @@ void FLevelEditorToolBar::RegisterBuildMenu()
 		}
 
 		/** Generates a lighting density sub-menu */
-		static void RegisterLightingDensityMenu(const FName BaseMenuName)
+		static void RegisterLightingDensityMenu(const FName InBaseMenuName)
 		{
-			UToolMenu* SubMenu = UToolMenus::Get()->RegisterMenu(UToolMenus::JoinMenuPaths(BaseMenuName, "LightingDensity"));
+			UToolMenu* SubMenu = UToolMenus::Get()->RegisterMenu(UToolMenus::JoinMenuPaths(InBaseMenuName, "LightingDensity"));
 
 			{
 				FToolMenuSection& Section = SubMenu->AddSection("LevelEditorBuildLightingDensity", LOCTEXT( "LightingDensityHeading", "Density Rendering" ) );
@@ -1646,9 +1646,9 @@ void FLevelEditorToolBar::RegisterBuildMenu()
 		}
 
 		/** Generates a lighting resolution sub-menu */
-		static void RegisterLightingResolutionMenu(const FName BaseMenuName)
+		static void RegisterLightingResolutionMenu(const FName InBaseMenuName)
 		{
-			UToolMenu* SubMenu = UToolMenus::Get()->RegisterMenu(UToolMenus::JoinMenuPaths(BaseMenuName, "LightingResolution"));
+			UToolMenu* SubMenu = UToolMenus::Get()->RegisterMenu(UToolMenus::JoinMenuPaths(InBaseMenuName, "LightingResolution"));
 
 			{
 				FToolMenuSection& Section = SubMenu->AddSection("LevelEditorBuildLightingResolution1", LOCTEXT( "LightingResolutionHeading1", "Primitive Types" ) );
@@ -1751,12 +1751,12 @@ void FLevelEditorToolBar::RegisterBuildMenu()
 		}
 
 		/** Generates a lighting info dialogs sub-menu */
-		static void RegisterLightingInfoMenu(const FName BaseMenuName)
+		static void RegisterLightingInfoMenu(const FName InBaseMenuName)
 		{
-			RegisterLightingDensityMenu(UToolMenus::JoinMenuPaths(BaseMenuName, "LightingInfo"));
-			RegisterLightingResolutionMenu(UToolMenus::JoinMenuPaths(BaseMenuName, "LightingInfo"));
+			RegisterLightingDensityMenu(UToolMenus::JoinMenuPaths(InBaseMenuName, "LightingInfo"));
+			RegisterLightingResolutionMenu(UToolMenus::JoinMenuPaths(InBaseMenuName, "LightingInfo"));
 
-			UToolMenu* SubMenu = UToolMenus::Get()->RegisterMenu(UToolMenus::JoinMenuPaths(BaseMenuName, "LightingInfo"));
+			UToolMenu* SubMenu = UToolMenus::Get()->RegisterMenu(UToolMenus::JoinMenuPaths(InBaseMenuName, "LightingInfo"));
 
 			{
 				FToolMenuSection& Section = SubMenu->AddSection("LevelEditorBuildLightingInfo", LOCTEXT( "LightingInfoHeading", "Lighting Info Dialogs" ) );
