Removed EActorCreateMode::Placement.   Was unused and rotting

diff --git a/Engine/Source/Editor/LevelEditor/Private/LevelEditorActions.cpp b/Engine/Source/Editor/LevelEditor/Private/LevelEditorActions.cpp
--- a/Engine/Source/Editor/LevelEditor/Private/LevelEditorActions.cpp
+++ b/Engine/Source/Editor/LevelEditor/Private/LevelEditorActions.cpp
@@ -1363,27 +1363,9 @@ bool FLevelEditorActionCallbacks::LockActorMovement_IsChecked()
 	return GEditor->HasLockedActors();
 }
 
-void FLevelEditorActionCallbacks::AddActor_Clicked( UActorFactory* ActorFactory, FAssetData AssetData, bool bUsePlacement )
+void FLevelEditorActionCallbacks::AddActor_Clicked( UActorFactory* ActorFactory, FAssetData AssetData)
 {
-	UObject* Object = AssetData.GetAsset();
-	if(bUsePlacement && IPlacementModeModule::IsAvailable() && Object != NULL)
-	{
-		FLevelEditorModule& LevelEditorModule = FModuleManager::LoadModuleChecked<FLevelEditorModule>("LevelEditor");
-		LevelEditorModule.FocusViewport();
-
-		// Make sure we're in actor placement mode
-		GLevelEditorModeTools().ActivateMode(FBuiltinEditorModes::EM_Placement);
-
-		TArray<UObject*> AssetsToPlace;
-		AssetsToPlace.Add(Object);
-
-		auto* PlacementMode = GLevelEditorModeTools().GetActiveModeTyped<IPlacementMode>(FBuiltinEditorModes::EM_Placement);
-		PlacementMode->StartPlacing(AssetsToPlace, ActorFactory);
-	}
-	else
-	{
-		FLevelEditorActionCallbacks::AddActor(ActorFactory, AssetData, NULL);
-	}
+	FLevelEditorActionCallbacks::AddActor(ActorFactory, AssetData, NULL);
 }
 
 AActor* FLevelEditorActionCallbacks::AddActor( UActorFactory* ActorFactory, const FAssetData& AssetData, const FTransform* ActorTransform )
diff --git a/Engine/Source/Editor/LevelEditor/Private/LevelEditorCreateActorMenu.cpp b/Engine/Source/Editor/LevelEditor/Private/LevelEditorCreateActorMenu.cpp
--- a/Engine/Source/Editor/LevelEditor/Private/LevelEditorCreateActorMenu.cpp
+++ b/Engine/Source/Editor/LevelEditor/Private/LevelEditorCreateActorMenu.cpp
@@ -304,7 +304,7 @@ static void FillAssetAddReplaceActorMenu(UToolMenu* Menu, const FAssetData Asset
 		}
 		else
 		{
-			Action = FUIAction( FExecuteAction::CreateStatic( &FLevelEditorActionCallbacks::AddActor_Clicked, MenuItem.FactoryToUse,  MenuItem.AssetData, CreateMode == EActorCreateMode::Placement ) );
+			Action = FUIAction( FExecuteAction::CreateStatic( &FLevelEditorActionCallbacks::AddActor_Clicked, MenuItem.FactoryToUse,  MenuItem.AssetData) );
 		}
 
 		Section.AddMenuEntry(NAME_None, Label, ToolTip, Icon, Action);
@@ -337,7 +337,7 @@ static void BuildSingleAssetAddReplaceActorMenu(FToolMenuSection& Section, const
 		}
 		else
 		{
-			Action = FUIAction( FExecuteAction::CreateStatic( &FLevelEditorActionCallbacks::AddActor_Clicked, MenuItem.FactoryToUse,  MenuItem.AssetData, CreateMode == EActorCreateMode::Placement ) );
+			Action = FUIAction( FExecuteAction::CreateStatic( &FLevelEditorActionCallbacks::AddActor_Clicked, MenuItem.FactoryToUse,  MenuItem.AssetData) );
 		}
 
 		Section.AddEntry(FToolMenuEntry::InitMenuEntry(NAME_None, Action, SNew(SAssetMenuEntry, Asset, AssetMenuOptions).LabelOverride(LabelOverride)));
@@ -387,7 +387,7 @@ void LevelEditorCreateActorMenu::FillAddReplaceViewportContextMenuSections(UTool
 
 		{
 			FToolMenuSection& Section = Menu->AddSection("AddActor", NSLOCTEXT("LevelViewportContextMenu", "AddActorHeading", "Place Actor"));
-			FUIAction Action( FExecuteAction::CreateStatic( &FLevelEditorActionCallbacks::AddActor_Clicked, AssetMenuOptions[0].FactoryToUse,  AssetMenuOptions[0].AssetData, false ) );
+			FUIAction Action( FExecuteAction::CreateStatic( &FLevelEditorActionCallbacks::AddActor_Clicked, AssetMenuOptions[0].FactoryToUse,  AssetMenuOptions[0].AssetData ) );
 			TSharedRef<SWidget> Widget = SNew(SAssetMenuEntry, TargetAssetData, AssetMenuOptions);
 			Section.AddEntry(FToolMenuEntry::InitSubMenu(
 				"AddActor",
diff --git a/Engine/Source/Editor/LevelEditor/Private/LevelEditorCreateActorMenu.h b/Engine/Source/Editor/LevelEditor/Private/LevelEditorCreateActorMenu.h
--- a/Engine/Source/Editor/LevelEditor/Private/LevelEditorCreateActorMenu.h
+++ b/Engine/Source/Editor/LevelEditor/Private/LevelEditorCreateActorMenu.h
@@ -19,8 +19,6 @@ namespace EActorCreateMode
 		/** Replace the actor that was last clicked on */
 		Replace,
 
-		/** Temp actor is attached to cursor & placed on click */
-		Placement,
 	};
 }
 
diff --git a/Engine/Source/Editor/LevelEditor/Private/SLevelViewport.cpp b/Engine/Source/Editor/LevelEditor/Private/SLevelViewport.cpp
--- a/Engine/Source/Editor/LevelEditor/Private/SLevelViewport.cpp
+++ b/Engine/Source/Editor/LevelEditor/Private/SLevelViewport.cpp
@@ -1212,7 +1212,7 @@ void FLevelViewportDropContextMenuImpl::FillDropAddReplaceActorMenu( bool bRepla
 				}
 				else
 				{
-					FUIAction Action( FExecuteAction::CreateStatic( &FLevelEditorActionCallbacks::AddActor_Clicked, MenuItem.FactoryToUse, MenuItem.AssetData, false ) );
+					FUIAction Action( FExecuteAction::CreateStatic( &FLevelEditorActionCallbacks::AddActor_Clicked, MenuItem.FactoryToUse, MenuItem.AssetData ) );
 				
 					FText MenuEntryName = FText::Format( NSLOCTEXT("SLevelViewport", "AddActorMenuFormat", "Add {0}"), MenuItem.FactoryToUse->GetDisplayName() );
 					if ( MenuItem.AssetData.IsValid() )
diff --git a/Engine/Source/Editor/LevelEditor/Public/LevelEditorActions.h b/Engine/Source/Editor/LevelEditor/Public/LevelEditorActions.h
--- a/Engine/Source/Editor/LevelEditor/Public/LevelEditorActions.h
+++ b/Engine/Source/Editor/LevelEditor/Public/LevelEditorActions.h
@@ -951,7 +951,7 @@ public:
 	 * @param bUsePlacement		Whether to use the placement editor. If not, the actor will be placed at the last click.
 	 * @param ActorLocation		[opt] If NULL, positions the actor at the mouse location, otherwise the location specified. Default is true.
 	 */
-	static void AddActor_Clicked( UActorFactory* ActorFactory, FAssetData AssetData, bool bUsePlacement );
+	static void AddActor_Clicked( UActorFactory* ActorFactory, FAssetData AssetData);
 	static AActor* AddActor( UActorFactory* ActorFactory, const FAssetData& AssetData, const FTransform* ActorLocation );
 
 	/**
