Change ToggleVR hotkey to add 'Shift' modifer

diff --git a/Engine/Source/Editor/LevelEditor/Private/LevelEditorActions.cpp b/Engine/Source/Editor/LevelEditor/Private/LevelEditorActions.cpp
--- a/Engine/Source/Editor/LevelEditor/Private/LevelEditorActions.cpp
+++ b/Engine/Source/Editor/LevelEditor/Private/LevelEditorActions.cpp
@@ -3302,7 +3302,7 @@ void FLevelEditorCommands::RegisterCommands()
 	UI_COMMAND( OpenContentBrowser, "Open Content Browser", "Opens the Content Browser", EUserInterfaceActionType::Button, FInputChord(EModifierKey::Control|EModifierKey::Shift, EKeys::F) );
 	UI_COMMAND( OpenMarketplace, "Open Marketplace", "Opens the Marketplace", EUserInterfaceActionType::Button, FInputChord() );
 
-	UI_COMMAND( ToggleVR, "Toggle VR", "Toggles VR (Virtual Reality) mode", EUserInterfaceActionType::ToggleButton, FInputChord( EModifierKey::Alt, EKeys::V ) );
+	UI_COMMAND( ToggleVR, "Toggle VR", "Toggles VR (Virtual Reality) mode", EUserInterfaceActionType::ToggleButton, FInputChord(EModifierKey::Shift, EKeys::V ) );
 
 	UI_COMMAND( OpenLevelBlueprint, "Open Level Blueprint", "Edit the Level Blueprint for the current level", EUserInterfaceActionType::Button, FInputChord() );
 	UI_COMMAND( CheckOutProjectSettingsConfig, "Check Out", "Checks out the project settings config file so the game mode can be set.", EUserInterfaceActionType::Button, FInputChord() );
