Editor mode changes:
Made placement browser its own tab and not the default mode.
The default mode is now "Select" which does nothing
Modes are now actived from a toolbar drop down (old keybindings still work) and the old tabbed UI for changing modes has been deprecated.
Mode tab no longer appears when there is no active mode and will appear when a mode becomes active

This change resets the default layout of the editor

diff --git a/Engine/Plugins/Editor/GeometryMode/Source/GeometryMode/Private/GeometryModeModule.cpp b/Engine/Plugins/Editor/GeometryMode/Source/GeometryMode/Private/GeometryModeModule.cpp
--- a/Engine/Plugins/Editor/GeometryMode/Source/GeometryMode/Private/GeometryModeModule.cpp
+++ b/Engine/Plugins/Editor/GeometryMode/Source/GeometryMode/Private/GeometryModeModule.cpp
@@ -16,7 +16,7 @@ void FGeometryModeModule::StartupModule()
 {
 	FEditorModeRegistry::Get().RegisterMode<FEdModeGeometry>(
 		FGeometryEditingModes::EM_Geometry,
-		NSLOCTEXT("EditorModes", "GeometryMode", "Geometry Editing"),
+		NSLOCTEXT("EditorModes", "GeometryMode", "Brush Editing"),
 		FSlateIcon(FEditorStyle::GetStyleSetName(), "LevelEditor.BspMode", "LevelEditor.BspMode.Small"),
 		true, 500
 	);
diff --git a/Engine/Plugins/Experimental/ChaosEditor/Source/FractureEditor/Private/FractureEditor.cpp b/Engine/Plugins/Experimental/ChaosEditor/Source/FractureEditor/Private/FractureEditor.cpp
--- a/Engine/Plugins/Experimental/ChaosEditor/Source/FractureEditor/Private/FractureEditor.cpp
+++ b/Engine/Plugins/Experimental/ChaosEditor/Source/FractureEditor/Private/FractureEditor.cpp
@@ -16,7 +16,7 @@ void FFractureEditorModule::StartupModule()
 
 	FEditorModeRegistry::Get().RegisterMode<FFractureEditorMode>(
 		FFractureEditorMode::EM_FractureEditorModeId, 
-		LOCTEXT("FractureEditorModeName", "FractureEditorMode"), 
+		LOCTEXT("FractureEditorModeName", "Fracture"), 
 		FSlateIcon("FractureEditorStyle", "LevelEditor.FractureMode", "LevelEditor.FractureMode.Small"),
 		true
 		);
diff --git a/Engine/Source/Editor/EditorStyle/Private/SlateEditorStyle.cpp b/Engine/Source/Editor/EditorStyle/Private/SlateEditorStyle.cpp
--- a/Engine/Source/Editor/EditorStyle/Private/SlateEditorStyle.cpp
+++ b/Engine/Source/Editor/EditorStyle/Private/SlateEditorStyle.cpp
@@ -4899,6 +4899,7 @@ void FSlateEditorStyle::FStyle::SetupLevelEditorStyle()
 			Set( "LevelEditor.Tabs.Cinematics", new IMAGE_BRUSH("/Icons/icon_tab_Cinematics_16x", Icon16x16));
 			Set( "LevelEditor.Tabs.EditorModes", new IMAGE_BRUSH( "/Icons/icon_Editor_Modes_16x", Icon16x16 ) );
 			Set( "LevelEditor.Tabs.Modes", new IMAGE_BRUSH( "/Icons/icon_Editor_Modes_16x", Icon16x16 ) );
+			Set( "LevelEditor.Tabs.PlacementBrowser", new IMAGE_BRUSH("/Icons/icon_Editor_Modes_16x", Icon16x16));
 			Set( "LevelEditor.Tabs.Properties", new IMAGE_BRUSH( "/Icons/properties_16x", Icon16x16 ) );
 			Set( "LevelEditor.Tabs.Outliner", new IMAGE_BRUSH( "/Icons/icon_tab_SceneOutliner_16x", Icon16x16 ) );
 			Set( "LevelEditor.Tabs.ContentBrowser", new IMAGE_BRUSH( "/Icons/icon_tab_ContentBrowser_16x", Icon16x16 ) );
@@ -5005,6 +5006,11 @@ void FSlateEditorStyle::FStyle::SetupLevelEditorStyle()
 		);	
 		Set( "ToolPalette.DockingWell", new FSlateColorBrush(FLinearColor(FColor(34, 34, 34, 255))));
 
+		Set("LevelEditor.SelectMode", new IMAGE_BRUSH("Icons/icon_DefaultEdMode_40x", Icon40x40));
+		Set("LevelEditor.SelectMode.Small", new IMAGE_BRUSH("Icons/icon_DefaultEdMode_40x", Icon20x20));
+		Set("LevelEditor.SelectMode.Selected", new IMAGE_BRUSH("Icons/icon_DefaultEdMode_40x", Icon40x40));
+		Set("LevelEditor.SelectMode.Selected.Small", new IMAGE_BRUSH("Icons/icon_DefaultEdMode_40x", Icon20x20));
+
 		Set( "LevelEditor.PlacementMode", new IMAGE_BRUSH( "Icons/icon_Mode_Placement_40px", Icon40x40 ) );
 		Set( "LevelEditor.PlacementMode.Small", new IMAGE_BRUSH( "Icons/icon_Mode_Placement_40px", Icon20x20 ) );
 		Set( "LevelEditor.PlacementMode.Selected", new IMAGE_BRUSH( "Icons/icon_Mode_Placement_selected_40x", Icon40x40 ) );
diff --git a/Engine/Source/Editor/FoliageEdit/Private/FoliageEdModeToolkit.cpp b/Engine/Source/Editor/FoliageEdit/Private/FoliageEdModeToolkit.cpp
--- a/Engine/Source/Editor/FoliageEdit/Private/FoliageEdModeToolkit.cpp
+++ b/Engine/Source/Editor/FoliageEdit/Private/FoliageEdModeToolkit.cpp
@@ -31,7 +31,7 @@ FName FFoliageEdModeToolkit::GetToolkitFName() const
 
 FText FFoliageEdModeToolkit::GetBaseToolkitName() const
 {
-	return LOCTEXT( "ToolkitName", "Foliage Edit Mode" );
+	return LOCTEXT( "ToolkitName", "Foliage" );
 }
 
 class FEdMode* FFoliageEdModeToolkit::GetEditorMode() const
diff --git a/Engine/Source/Editor/LandscapeEditor/Private/SLandscapeEditor.cpp b/Engine/Source/Editor/LandscapeEditor/Private/SLandscapeEditor.cpp
--- a/Engine/Source/Editor/LandscapeEditor/Private/SLandscapeEditor.cpp
+++ b/Engine/Source/Editor/LandscapeEditor/Private/SLandscapeEditor.cpp
@@ -145,7 +145,7 @@ FName FLandscapeToolKit::GetToolkitFName() const
 
 FText FLandscapeToolKit::GetBaseToolkitName() const
 {
-	return LOCTEXT("ToolkitName", "Landscape Editor");
+	return LOCTEXT("ToolkitName", "Landscape");
 }
 
 FEdModeLandscape* FLandscapeToolKit::GetEditorMode() const
diff --git a/Engine/Source/Editor/LevelEditor/Private/LevelEditor.cpp b/Engine/Source/Editor/LevelEditor/Private/LevelEditor.cpp
--- a/Engine/Source/Editor/LevelEditor/Private/LevelEditor.cpp
+++ b/Engine/Source/Editor/LevelEditor/Private/LevelEditor.cpp
@@ -39,6 +39,7 @@
 #include "Interfaces/IMainFrameModule.h"
 #include "Framework/Commands/GenericCommands.h"
 #include "Misc/EngineBuildSettings.h"
+#include "Settings/LevelEditorMiscSettings.h"
 
 #define LOCTEXT_NAMESPACE "LevelEditor"
 
@@ -173,11 +174,30 @@ TSharedRef<SDockTab> FLevelEditorModule::SpawnLevelEditor( const FSpawnTabArgs&
 		SetLevelEditorInstance(LevelEditorTmp);
 		LevelEditorTmp->Initialize( LevelEditorTab, OwnerWindow.ToSharedRef() );
 
-		GLevelEditorModeTools().RemoveDefaultMode( FBuiltinEditorModes::EM_Default );
-		GLevelEditorModeTools().AddDefaultMode( FBuiltinEditorModes::EM_Placement );
+		if (GetDefault<ULevelEditorMiscSettings>()->bEnableLegacyEditorModeUI)
+		{
+			GLevelEditorModeTools().RemoveDefaultMode(FBuiltinEditorModes::EM_Default);
+			GLevelEditorModeTools().AddDefaultMode(FBuiltinEditorModes::EM_Placement);
+		}
+
 		GLevelEditorModeTools().DeactivateAllModes();
 		GLevelEditorModeTools().ActivateDefaultMode();
 
+		if (GetDefault<ULevelEditorMiscSettings>()->bEnableLegacyEditorModeUI)
+		{
+			// In legacy mode this toolbox should always be open
+			static const FTabId ToolboxTabId("LevelEditorToolBox");
+			LevelEditorTabManager->InvokeTab(ToolboxTabId);
+
+			// In legacy mode the standalone placement browser tab should not be opened
+			static const FTabId PlacementBrowserTabId("PlacementBrowser");
+			TSharedPtr<SDockTab> PlacementBrowserTab = LevelEditorTabManager->FindExistingLiveTab(PlacementBrowserTabId);
+			if (PlacementBrowserTab.IsValid())
+			{
+				PlacementBrowserTab->RequestCloseTab();
+			}
+		}
+
 	}
 
 	IIntroTutorials& IntroTutorials = FModuleManager::LoadModuleChecked<IIntroTutorials>(TEXT("IntroTutorials"));
diff --git a/Engine/Source/Editor/LevelEditor/Private/LevelEditorToolBar.cpp b/Engine/Source/Editor/LevelEditor/Private/LevelEditorToolBar.cpp
--- a/Engine/Source/Editor/LevelEditor/Private/LevelEditorToolBar.cpp
+++ b/Engine/Source/Editor/LevelEditor/Private/LevelEditorToolBar.cpp
@@ -61,6 +61,7 @@
 #include "LevelEditorMenuContext.h"
 #include "ToolMenus.h"
 #include "Subsystems/AssetEditorSubsystem.h"
+#include "LevelEditorModesActions.h"
 
 static TAutoConsoleVariable<int32> CVarAllowMatineeActors(
 	TEXT("Matinee.AllowMatineeActors"),
@@ -1173,8 +1174,7 @@ void FLevelEditorToolBar::RegisterLevelEditorToolBar( const TSharedRef<FUIComman
 {
 	RegisterSourceControlMenu();
 	RegisterCinematicsMenu();
-	RegisterBuildMenu();
-
+	RegisterEditorModesMenu();
 #if WITH_LIVE_CODING
 	RegisterCompileMenu();
 #endif
@@ -1283,6 +1283,18 @@ void FLevelEditorToolBar::RegisterLevelEditorToolBar( const TSharedRef<FUIComman
 		}
 	}
 
+	{
+		FToolMenuSection& Section = Toolbar->AddSection("Modes");
+		Section.AddEntry(FToolMenuEntry::InitComboButton(
+			"EditorModes",
+			FUIAction(),
+			FOnGetContent::CreateStatic(&FLevelEditorToolBar::GenerateEditorModesMenu, InCommandList, TWeakPtr<SLevelEditor>(InLevelEditor)),
+			LOCTEXT("EditorModes_Label", "Modes"),
+			LOCTEXT("EditorModes_Tooltip", "Displays a list of editing modes that can be toggled"),
+			FSlateIcon(FEditorStyle::GetStyleSetName(), "LevelEditor.EditorModes")
+		));
+	}
+
 	{
 		FToolMenuSection& Section = Toolbar->AddSection("Content");
 		Section.AddEntry(FToolMenuEntry::InitToolBarButton(FLevelEditorCommands::Get().OpenContentBrowser, LOCTEXT( "ContentBrowser_Override", "Content" ), TAttribute<FText>(), TAttribute<FSlateIcon>(), "LevelToolbarContent"));
@@ -2102,6 +2114,19 @@ TSharedRef< SWidget > FLevelEditorToolBar::GenerateSourceControlMenu(TSharedRef<
 	return UToolMenus::Get()->GenerateWidget("LevelEditor.LevelEditorToolBar.SourceControl", MenuContext);
 }
 
+TSharedRef< SWidget > FLevelEditorToolBar::GenerateEditorModesMenu(TSharedRef<FUICommandList> InCommandList, TWeakPtr<SLevelEditor> InLevelEditor)
+{
+	// Get all menu extenders for this context menu from the level editor module
+	FLevelEditorModule& LevelEditorModule = FModuleManager::GetModuleChecked<FLevelEditorModule>(TEXT("LevelEditor"));
+
+	FToolMenuContext MenuContext(InCommandList);
+	ULevelEditorMenuContext* LevelEditorMenuContext = NewObject<ULevelEditorMenuContext>();
+	LevelEditorMenuContext->LevelEditor = InLevelEditor;
+	MenuContext.AddObject(LevelEditorMenuContext);
+
+	return UToolMenus::Get()->GenerateWidget("LevelEditor.LevelEditorToolBar.EditorModes", MenuContext);
+}
+
 void FLevelEditorToolBar::RegisterSourceControlMenu()
 {
 #define LOCTEXT_NAMESPACE "LevelToolBarSourceControlMenu"
@@ -2411,6 +2436,77 @@ void FLevelEditorToolBar::RegisterCinematicsMenu()
 #undef LOCTEXT_NAMESPACE
 }
 
+void FLevelEditorToolBar::RegisterEditorModesMenu()
+{
+#define LOCTEXT_NAMESPACE "LevelToolBarEditorModesMenu"
+	UToolMenu* Menu = UToolMenus::Get()->RegisterMenu("LevelEditor.LevelEditorToolBar.EditorModes");
+
+	FToolMenuSection& Section = Menu->AddSection("EditorModes", LOCTEXT("EditorModesMenu_NewHeading", "Editor Modes"));
+
+	Section.AddDynamicEntry("ModesList", FNewToolMenuSectionDelegate::CreateLambda([](FToolMenuSection& InSection)
+	{
+
+		const FLevelEditorModule& LevelEditorModule = FModuleManager::LoadModuleChecked<FLevelEditorModule>("LevelEditor");
+		const FLevelEditorModesCommands& Commands = LevelEditorModule.GetLevelEditorModesCommands();
+
+		TArray<FEditorModeInfo, TInlineAllocator<1>> DefaultModes;
+
+		TArray<FEditorModeInfo, TInlineAllocator<10>> NonDefaultModes;
+
+		for (const FEditorModeInfo& Mode : FEditorModeRegistry::Get().GetSortedModeInfo())
+		{
+			// If the mode isn't visible don't create a menu option for it.
+			if (!Mode.bVisible)
+			{
+				continue;
+			}
+
+			if (GLevelEditorModeTools().IsDefaultMode(Mode.ID))
+			{
+				DefaultModes.Add(Mode.ID);
+			}
+			else
+			{
+				NonDefaultModes.Add(Mode.ID);
+			}
+			
+		}
+
+		auto BuildEditorModes = 
+			[&Commands, &InSection](const TArrayView<FEditorModeInfo>& Modes)
+			{
+				for (const FEditorModeInfo& Mode : Modes)
+				{
+					FName EditorModeCommandName = FName(*(FString("EditorMode.") + Mode.ID.ToString()));
+
+					TSharedPtr<FUICommandInfo> EditorModeCommand =
+						FInputBindingManager::Get().FindCommandInContext(Commands.GetContextName(), EditorModeCommandName);
+
+					// If a command isn't yet registered for this mode, we need to register one.
+					if (!EditorModeCommand.IsValid())
+					{
+						continue;
+					}
+
+					InSection.AddMenuEntry(EditorModeCommand);
+				}
+
+			};
+
+		// Build Default Modes first
+		BuildEditorModes(MakeArrayView(DefaultModes));
+
+		InSection.AddMenuSeparator(NAME_None);
+			
+		// Build non-default modes second
+		BuildEditorModes(MakeArrayView(NonDefaultModes));
+
+	}));
+
+#undef LOCTEXT_NAMESPACE
+}
+
+
 void FLevelEditorToolBar::OnCinematicsActorPicked( AActor* Actor )
 {
 	//The matinee editor will not tick unless the editor viewport is in realtime mode.
diff --git a/Engine/Source/Editor/LevelEditor/Private/LevelEditorToolBar.h b/Engine/Source/Editor/LevelEditor/Private/LevelEditorToolBar.h
--- a/Engine/Source/Editor/LevelEditor/Private/LevelEditorToolBar.h
+++ b/Engine/Source/Editor/LevelEditor/Private/LevelEditorToolBar.h
@@ -50,6 +50,14 @@ protected:
 	 */
 	static TSharedRef< SWidget > GenerateSourceControlMenu(TSharedRef<FUICommandList> InCommandList, TWeakPtr<SLevelEditor> InLevelEditor);
 
+	
+	/**
+	 * Generates menu content for the modes combo button drop down menu
+	 *
+	 * @return	Menu content widget
+	 */
+	static TSharedRef< SWidget > GenerateEditorModesMenu(TSharedRef<FUICommandList> InCommandList, TWeakPtr<SLevelEditor> InLevelEditor);
+
 	/**
 	 * Generates menu content for the compile combo button drop down menu
 	 *
@@ -81,6 +89,7 @@ private:
 
 	static void RegisterSourceControlMenu();
 	static void RegisterCinematicsMenu();
+	static void RegisterEditorModesMenu();
 	static void RegisterBuildMenu();
 
 #if WITH_LIVE_CODING
diff --git a/Engine/Source/Editor/LevelEditor/Private/SLevelEditor.cpp b/Engine/Source/Editor/LevelEditor/Private/SLevelEditor.cpp
--- a/Engine/Source/Editor/LevelEditor/Private/SLevelEditor.cpp
+++ b/Engine/Source/Editor/LevelEditor/Private/SLevelEditor.cpp
@@ -50,6 +50,8 @@
 #include "Framework/Docking/LayoutExtender.h"
 #include "HierarchicalLODOutlinerModule.h"
 #include "EditorViewportCommands.h"
+#include "IPlacementModeModule.h"
+#include "Settings/LevelEditorMiscSettings.h"
 
 static const FName LevelEditorBuildAndSubmitTab("LevelEditorBuildAndSubmit");
 static const FName LevelEditorStatsViewerTab("LevelEditorStatsViewer");
@@ -668,9 +670,9 @@ TSharedRef<SDockTab> SLevelEditor::SpawnLevelEditorTab( const FSpawnTabArgs& Arg
 	{
 		TSharedRef<SLevelEditorToolBox> NewToolBox = StaticCastSharedRef<SLevelEditorToolBox>( CreateToolBox() );
 
-		return SNew( SDockTab )
+		TSharedRef<SDockTab> DockTab = SNew( SDockTab )
 			.Icon( FEditorStyle::GetBrush( "LevelEditor.Tabs.Modes" ) )
-			.Label( NSLOCTEXT( "LevelEditor", "ToolsTabTitle", "Modes" ) )
+			.Label( NSLOCTEXT( "LevelEditor", "ToolsTabTitle", "Toolbox" ) )
 			[
 				SNew( SBox )
 				.AddMetaData<FTutorialMetaData>(FTutorialMetaData(TEXT("ToolsPanel"), TEXT("LevelEditorToolBox")))
@@ -678,6 +680,24 @@ TSharedRef<SDockTab> SLevelEditor::SpawnLevelEditorTab( const FSpawnTabArgs& Arg
 					NewToolBox
 				]
 			];
+
+		NewToolBox->SetParentTab(DockTab);
+
+		return DockTab;
+	}
+	else if (TabIdentifier == TEXT("PlacementBrowser"))
+	{
+		if(!GetDefault<ULevelEditorMiscSettings>()->bEnableLegacyEditorModeUI)
+		{
+			return
+				SNew(SDockTab)
+				.Icon(FEditorStyle::GetBrush("LevelEditor.Tabs.PlacementBrowser"))
+				.Label(NSLOCTEXT("LevelEditor", "PlacementBrowserTitle", "Place Actors"))
+				.AddMetaData<FTutorialMetaData>(FTutorialMetaData(TEXT("PlacementBrowser"), TEXT("PlacementBrowser")))
+				[
+					IPlacementModeModule::Get().CreatePlacementModeBrowser()
+				];
+		}
 	}
 	else if( TabIdentifier == LevelEditorBuildAndSubmitTab )
 	{
@@ -755,20 +775,19 @@ TSharedRef<SDockTab> SLevelEditor::SpawnLevelEditorTab( const FSpawnTabArgs& Arg
 	}
 	else if( TabIdentifier == TEXT("LevelEditorLayerBrowser") )
 	{
-
-			FLayersModule& LayersModule = FModuleManager::LoadModuleChecked<FLayersModule>( "Layers" );
-			return SNew( SDockTab )
-				.Icon( FEditorStyle::GetBrush( "LevelEditor.Tabs.Layers" ) )
-				.Label( NSLOCTEXT("LevelEditor", "LayersTabTitle", "Layers") )
+		FLayersModule& LayersModule = FModuleManager::LoadModuleChecked<FLayersModule>( "Layers" );
+		return SNew( SDockTab )
+			.Icon( FEditorStyle::GetBrush( "LevelEditor.Tabs.Layers" ) )
+			.Label( NSLOCTEXT("LevelEditor", "LayersTabTitle", "Layers") )
+			[
+				SNew(SBorder)
+				.Padding( 0 )
+				.BorderImage( FEditorStyle::GetBrush("ToolPanel.GroupBorder") )
+				.AddMetaData<FTutorialMetaData>(FTutorialMetaData(TEXT("LayerBrowser"), TEXT("LevelEditorLayerBrowser")))
 				[
-					SNew(SBorder)
-					.Padding( 0 )
-					.BorderImage( FEditorStyle::GetBrush("ToolPanel.GroupBorder") )
-					.AddMetaData<FTutorialMetaData>(FTutorialMetaData(TEXT("LayerBrowser"), TEXT("LevelEditorLayerBrowser")))
-					[
-						LayersModule.CreateLayerBrowser()
-					]
-				];
+					LayersModule.CreateLayerBrowser()
+				]
+			];
 	}
 	else if (TabIdentifier == TEXT("LevelEditorHierarchicalLODOutliner"))
 	{
@@ -877,6 +896,11 @@ bool SLevelEditor::CanSpawnEditorModeToolbarTab(const FSpawnTabArgs& Args) const
 	return GLevelEditorModeTools().ShouldShowModeToolbar();
 }
 
+bool SLevelEditor::CanSpawnEditorModeToolboxTab(const FSpawnTabArgs& Args) const
+{
+	return GLevelEditorModeTools().ShouldShowModeToolbox();
+}
+
 TSharedRef<SDockTab> SLevelEditor::InvokeTab( FName TabID )
 {
 	TSharedPtr<FTabManager> LevelEditorTabManager = GetTabManager();
@@ -1123,22 +1147,35 @@ TSharedRef<SWidget> SLevelEditor::RestoreContentArea( const TSharedRef<SDockTab>
 		}
 
 		{
-		
+
+			FCanSpawnTab CanSpawnTabDelegate = FCanSpawnTab::CreateSP(this, &SLevelEditor::CanSpawnEditorModeToolboxTab);
 			const FSlateIcon ToolsIcon(FEditorStyle::GetStyleSetName(), "LevelEditor.Tabs.Modes");
-			LevelEditorTabManager->RegisterTabSpawner("LevelEditorToolBox", FOnSpawnTab::CreateSP<SLevelEditor, FName, FString>(this, &SLevelEditor::SpawnLevelEditorTab, FName("LevelEditorToolBox"), FString()))
-				.SetDisplayName(NSLOCTEXT("LevelEditorTabs", "LevelEditorToolBox", "Modes Panel"))
-				.SetTooltipText(NSLOCTEXT("LevelEditorTabs", "LevelEditorToolBoxTooltipText", "Open the Modes tab, which specifies all the available editing modes."))
+			LevelEditorTabManager->RegisterTabSpawner("LevelEditorToolBox", FOnSpawnTab::CreateSP<SLevelEditor, FName, FString>(this, &SLevelEditor::SpawnLevelEditorTab, FName("LevelEditorToolBox"), FString()), CanSpawnTabDelegate)
+				.SetDisplayName(NSLOCTEXT("LevelEditorTabs", "LevelEditorModesToolboxTab", "Active Mode Toolbox"))
+				.SetTooltipText(NSLOCTEXT("LevelEditorTabs", "LevelEditorModesToolboxTabTooltipText", "Open the Modes tab, which contains the active editor mode's settings."))
 				.SetGroup(MenuStructure.GetLevelEditorModesCategory())
 				.SetIcon(ToolsIcon);
+		}
 
+		{
+			const FSlateIcon ToolsIcon(FEditorStyle::GetStyleSetName(), "LevelEditor.Tabs.PlacementBrowser");
 			FCanSpawnTab CanSpawnTabDelegate = FCanSpawnTab::CreateSP(this, &SLevelEditor::CanSpawnEditorModeToolbarTab);
 			LevelEditorTabManager->RegisterTabSpawner(FEditorModeTools::EditorModeToolbarTabName, FOnSpawnTab::CreateSP<SLevelEditor, FName, FString>(this, &SLevelEditor::SpawnLevelEditorTab, FEditorModeTools::EditorModeToolbarTabName, FString()), CanSpawnTabDelegate)
-				.SetDisplayName(NSLOCTEXT("LevelEditorTabs", "LevelEditorModesTab", "Active Mode Tools"))
-				.SetTooltipText(NSLOCTEXT("LevelEditorTabs", "LevelEditorModesTabTooltipText", "Open the modes toolbar which has the active mode's tools"))
+				.SetDisplayName(NSLOCTEXT("LevelEditorTabs", "LevelEditorModesToolbarTab", "Active Mode Toolbar"))
+				.SetTooltipText(NSLOCTEXT("LevelEditorTabs", "LevelEditorModesToolbarTabTooltipText", "Opens a toolbar for the active editor mode"))
 				.SetGroup(MenuStructure.GetLevelEditorModesCategory())
 				.SetIcon(ToolsIcon);
 		}
 
+		{
+			const FSlateIcon ToolsIcon(FEditorStyle::GetStyleSetName(), "LevelEditor.Tabs.PlacementBrowser");
+			LevelEditorTabManager->RegisterTabSpawner("PlacementBrowser", FOnSpawnTab::CreateSP<SLevelEditor, FName, FString>(this, &SLevelEditor::SpawnLevelEditorTab, FName("PlacementBrowser"), FString()))
+				.SetDisplayName(NSLOCTEXT("LevelEditorTabs", "PlacementBrowser", "Place Actors"))
+				.SetTooltipText(NSLOCTEXT("LevelEditorTabs", "PlacementBrowserTooltipText", "Actor Placement Browser"))
+				.SetGroup(MenuStructure.GetLevelEditorCategory())
+				.SetIcon(ToolsIcon);
+		}
+
 		{
 			const FSlateIcon OutlinerIcon(FEditorStyle::GetStyleSetName(), "LevelEditor.Tabs.Outliner");
 		    LevelEditorTabManager->RegisterTabSpawner( "LevelEditorSceneOutliner", FOnSpawnTab::CreateSP<SLevelEditor, FName, FString>(this, &SLevelEditor::SpawnLevelEditorTab, FName("LevelEditorSceneOutliner"), FString()) )
@@ -1245,7 +1282,7 @@ TSharedRef<SWidget> SLevelEditor::RestoreContentArea( const TSharedRef<SDockTab>
 	// 10. Push the new "DefaultLayout.ini" together with your new code.
 	// 11. Also update these instructions if you change the version number (e.g., from "UnrealEd_Layout_v1.4" to "UnrealEd_Layout_v1.5").
 	const TSharedRef<FTabManager::FLayout> Layout = FLayoutSaveRestore::LoadFromConfig(GEditorLayoutIni,
-		FTabManager::NewLayout( "LevelEditor_Layout_v1.1" )
+		FTabManager::NewLayout( "LevelEditor_Layout_v1.2" )
 		->AddArea
 		(
 			FTabManager::NewPrimaryArea()
@@ -1265,7 +1302,9 @@ TSharedRef<SWidget> SLevelEditor::RestoreContentArea( const TSharedRef<SDockTab>
 					(
 						FTabManager::NewStack()
 						->SetSizeCoefficient( 0.3f )
-						->AddTab( "LevelEditorToolBox", ETabState::OpenedTab )
+						->AddTab("PlacementBrowser", ETabState::OpenedTab)
+						->AddTab("LevelEditorToolBox", ETabState::ClosedTab)
+						->SetForegroundTab(FName("PlacementBrowser"))
 					)
 					->Split
 					(
@@ -1327,9 +1366,6 @@ TSharedRef<SWidget> SLevelEditor::RestoreContentArea( const TSharedRef<SDockTab>
 	
 	FLayoutExtender LayoutExtender;
 
-	// Make a layout extension for the mode toolbar.  This avoids bumping level editor layout version to add a mandatory tab
-	LayoutExtender.ExtendLayout(FName("LevelEditorToolBar"), ELayoutExtensionPosition::Below, FTabManager::FTab(FEditorModeTools::EditorModeToolbarTabName, ETabState::ClosedTab));
-
 	LevelEditorModule.OnRegisterLayoutExtensions().Broadcast(LayoutExtender);
 	Layout->ProcessExtensions(LayoutExtender);
 
@@ -1380,7 +1416,21 @@ void SLevelEditor::ToggleEditorMode( FEditorModeID ModeID )
 	FLevelEditorModule& LevelEditorModule = FModuleManager::GetModuleChecked<FLevelEditorModule>( "LevelEditor" );
 	TSharedPtr<FTabManager> LevelEditorTabManager = LevelEditorModule.GetLevelEditorTabManager();
 
-	TSharedRef<SDockTab> ToolboxTab = LevelEditorTabManager->InvokeTab( FTabId("LevelEditorToolBox") );
+	static const FTabId ToolboxTabId("LevelEditorToolBox");
+
+	if (GLevelEditorModeTools().IsModeActive(FBuiltinEditorModes::EM_Default) && !GetDefault<ULevelEditorMiscSettings>()->bEnableLegacyEditorModeUI)
+	{
+		TSharedPtr<SDockTab> ToolboxTab = LevelEditorTabManager->FindExistingLiveTab(ToolboxTabId);
+		if (ToolboxTab.IsValid())
+		{
+			ToolboxTab->RequestCloseTab();
+		}
+	}
+	else
+	{
+		LevelEditorTabManager->InvokeTab(ToolboxTabId);
+	}
+
 }
 
 bool SLevelEditor::IsModeActive( FEditorModeID ModeID )
diff --git a/Engine/Source/Editor/LevelEditor/Private/SLevelEditor.h b/Engine/Source/Editor/LevelEditor/Private/SLevelEditor.h
--- a/Engine/Source/Editor/LevelEditor/Private/SLevelEditor.h
+++ b/Engine/Source/Editor/LevelEditor/Private/SLevelEditor.h
@@ -139,6 +139,7 @@ private:
 	
 	TSharedRef<SDockTab> SpawnLevelEditorTab(const FSpawnTabArgs& Args, FName TabIdentifier, FString InitializationPayload);
 	bool CanSpawnEditorModeToolbarTab(const FSpawnTabArgs& Args) const;
+	bool CanSpawnEditorModeToolboxTab(const FSpawnTabArgs& Args) const;
 
 	//TSharedRef<SDockTab> SpawnLevelEditorModeTab(const FSpawnTabArgs& Args, FEdMode* EditorMode);
 	TSharedRef<SDockTab> SummonDetailsPanel( FName Identifier );
diff --git a/Engine/Source/Editor/LevelEditor/Private/SLevelEditorToolBox.cpp b/Engine/Source/Editor/LevelEditor/Private/SLevelEditorToolBox.cpp
--- a/Engine/Source/Editor/LevelEditor/Private/SLevelEditorToolBox.cpp
+++ b/Engine/Source/Editor/LevelEditor/Private/SLevelEditorToolBox.cpp
@@ -13,6 +13,8 @@
 #include "LevelEditor.h"
 #include "LevelEditorActions.h"
 #include "LevelEditorModesActions.h"
+#include "Settings/LevelEditorMiscSettings.h"
+#include "EdMode.h"
 
 #define LOCTEXT_NAMESPACE "SLevelEditorToolBox"
 
@@ -26,6 +28,7 @@ SLevelEditorToolBox::~SLevelEditorToolBox()
 
 void SLevelEditorToolBox::Construct( const FArguments& InArgs, const TSharedRef< class ILevelEditor >& OwningLevelEditor )
 {
+	TabIcon = FEditorStyle::Get().GetBrush("LevelEditor.Tabs.Modes");
 	LevelEditor = OwningLevelEditor;
 
 	// Important: We use a raw binding here because we are releasing our binding in our destructor (where a weak pointer would be invalid)
@@ -36,7 +39,6 @@ void SLevelEditorToolBox::Construct( const FArguments& InArgs, const TSharedRef<
 	ChildSlot
 	[
 		SNew( SVerticalBox )
-
 		+ SVerticalBox::Slot()
 		.AutoHeight()
 		.HAlign( HAlign_Left )
@@ -73,20 +75,27 @@ void SLevelEditorToolBox::Construct( const FArguments& InArgs, const TSharedRef<
 		]
 	];
 
-	UpdateModeToolBar();
+	UpdateModeLegacyToolBar();
 }
 
 void SLevelEditorToolBox::HandleUserSettingsChange( FName PropertyName )
 {
-	UpdateModeToolBar();
+	UpdateModeLegacyToolBar();
 }
 
 void SLevelEditorToolBox::OnEditorModeCommandsChanged()
 {
-	UpdateModeToolBar();
+	UpdateModeLegacyToolBar();
+}
+
+void SLevelEditorToolBox::SetParentTab(TSharedRef<SDockTab>& InDockTab)
+{
+	ParentTab = InDockTab;
+	InDockTab->SetLabel(TabName);
+	InDockTab->SetTabIcon(TabIcon);
 }
 
-void SLevelEditorToolBox::UpdateModeToolBar()
+void SLevelEditorToolBox::UpdateModeLegacyToolBar()
 {
 	FLevelEditorModule& LevelEditorModule = FModuleManager::GetModuleChecked<FLevelEditorModule>( "LevelEditor");
 	const TSharedPtr< const FUICommandList > CommandList = LevelEditorModule.GetGlobalLevelEditorActions();
@@ -126,12 +135,19 @@ void SLevelEditorToolBox::UpdateModeToolBar()
 		}
 	}
 
-	ModeToolBarContainer->SetContent( EditorModeTools.MakeWidget() );
+	if (GetDefault<ULevelEditorMiscSettings>()->bEnableLegacyEditorModeUI)
+	{
+		ModeToolBarContainer->SetContent(EditorModeTools.MakeWidget());
+	}
+	else
+	{
+		ModeToolBarContainer->SetVisibility(EVisibility::Collapsed);
+	}
 
-	const TArray< TSharedPtr< IToolkit > >& HostedToolkits = LevelEditor.Pin()->GetHostedToolkits();
-	for( auto HostedToolkitIt = HostedToolkits.CreateConstIterator(); HostedToolkitIt; ++HostedToolkitIt )
+	const TArray<TSharedPtr<IToolkit>>& HostedToolkits = LevelEditor.Pin()->GetHostedToolkits();
+	for(const TSharedPtr<IToolkit>& HostedToolkitIt : HostedToolkits)
 	{
-		UpdateInlineContent( ( *HostedToolkitIt )->GetInlineContent() );
+		UpdateInlineContent(HostedToolkitIt, HostedToolkitIt->GetInlineContent());
 		break;
 	}
 }
@@ -146,50 +162,57 @@ EVisibility SLevelEditorToolBox::GetNoToolSelectedTextVisibility() const
 	return InlineContentHolder->GetContent() == SNullWidget::NullWidget ? EVisibility::Visible : EVisibility::Collapsed;
 }
 
-void SLevelEditorToolBox::UpdateInlineContent(TSharedPtr<SWidget> InlineContent) const
+void SLevelEditorToolBox::UpdateInlineContent(const TSharedPtr<IToolkit>& Toolkit, TSharedPtr<SWidget> InlineContent) 
 {
+	TabName = NSLOCTEXT("LevelEditor", "ToolsTabTitle", "Toolbox");
+	TabIcon = FEditorStyle::Get().GetBrush("LevelEditor.Tabs.Modes");
+
+	if (Toolkit.IsValid() && Toolkit->GetEditorMode())
+	{
+		TabName = Toolkit->GetEditorMode()->GetModeInfo().Name;
+		TabIcon = Toolkit->GetEditorMode()->GetModeInfo().IconBrush.GetSmallIcon();
+	}
+	else if (Toolkit.IsValid())
+	{
+		TabName = Toolkit->GetToolkitName();
+	}
+
+	TSharedPtr<SDockTab> ParentTabPinned = ParentTab.Pin();
 	if (InlineContent.IsValid() && InlineContentHolder.IsValid())
 	{
 		InlineContentHolder->SetContent(InlineContent.ToSharedRef());
 	}
+
+	if (ParentTabPinned.IsValid())
+	{
+		ParentTabPinned->SetLabel(TabName);
+		ParentTabPinned->SetTabIcon(TabIcon);
+	}
 }
 
-void SLevelEditorToolBox::OnToolkitHostingStarted( const TSharedRef< class IToolkit >& Toolkit )
+void SLevelEditorToolBox::OnToolkitHostingStarted(const TSharedRef<IToolkit>& Toolkit)
 {
-	UpdateInlineContent( Toolkit->GetInlineContent() );
+	UpdateInlineContent(Toolkit, Toolkit->GetInlineContent());
 }
 
-void SLevelEditorToolBox::OnToolkitHostingFinished( const TSharedRef< class IToolkit >& Toolkit )
+void SLevelEditorToolBox::OnToolkitHostingFinished(const TSharedRef<IToolkit>& Toolkit)
 {
 	bool FoundAnotherToolkit = false;
-	const TArray< TSharedPtr< IToolkit > >& HostedToolkits = LevelEditor.Pin()->GetHostedToolkits();
-	for( auto HostedToolkitIt = HostedToolkits.CreateConstIterator(); HostedToolkitIt; ++HostedToolkitIt )
+	const TArray<TSharedPtr<IToolkit>>& HostedToolkits = LevelEditor.Pin()->GetHostedToolkits();
+	for (const TSharedPtr<IToolkit>& HostedToolkitIt : HostedToolkits)
 	{
-		if ( ( *HostedToolkitIt ) != Toolkit )
+		if (HostedToolkitIt != Toolkit)
 		{
-			UpdateInlineContent( ( *HostedToolkitIt )->GetInlineContent() );
+			UpdateInlineContent(HostedToolkitIt, HostedToolkitIt->GetInlineContent());
 			FoundAnotherToolkit = true;
 			break;
 		}
 	}
 
-	if ( !FoundAnotherToolkit )
-	{
-		UpdateInlineContent( SNullWidget::NullWidget );
-	}
-}
-
-FSlateIcon SLevelEditorToolBox::GetEditorModeIcon(TSharedPtr< FUICommandInfo > EditorModeUICommand, FEditorModeID EditorMode )
-{
-	const FSlateIcon& Icon = EditorModeUICommand->GetIcon();
-
-	FName IconName = Icon.GetStyleName();
-	if( FLevelEditorActionCallbacks::IsEditorModeActive(EditorMode) )
+	if (!FoundAnotherToolkit)
 	{
-		IconName = FEditorStyle::Join(IconName, ".Selected");
+		UpdateInlineContent(nullptr, SNullWidget::NullWidget );
 	}
-
-	return FSlateIcon(Icon.GetStyleSetName(), IconName);
 }
 
 #undef LOCTEXT_NAMESPACE
diff --git a/Engine/Source/Editor/LevelEditor/Private/SLevelEditorToolBox.h b/Engine/Source/Editor/LevelEditor/Private/SLevelEditorToolBox.h
--- a/Engine/Source/Editor/LevelEditor/Private/SLevelEditorToolBox.h
+++ b/Engine/Source/Editor/LevelEditor/Private/SLevelEditorToolBox.h
@@ -14,6 +14,9 @@
 
 class FExtender;
 class SBorder;
+class IToolkit;
+class SDockTab;
+class ILevelEditor;
 
 /**
  * Tools for the level editor                   
@@ -26,17 +29,19 @@ public:
 
 	~SLevelEditorToolBox();
 
-	void Construct( const FArguments& InArgs, const TSharedRef< class ILevelEditor >& OwningLevelEditor );
+	void Construct(const FArguments& InArgs, const TSharedRef<ILevelEditor>& OwningLevelEditor);
 
 	/** Called by SLevelEditor to notify the toolbox about a new toolkit being hosted */
-	void OnToolkitHostingStarted( const TSharedRef< class IToolkit >& Toolkit );
+	void OnToolkitHostingStarted( const TSharedRef<IToolkit>& Toolkit);
 
 	/** Called by SLevelEditor to notify the toolbox about an existing toolkit no longer being hosted */
-	void OnToolkitHostingFinished(const TSharedRef< class IToolkit >& Toolkit);
+	void OnToolkitHostingFinished(const TSharedRef<IToolkit>& Toolkit);
 
 	/** Handles updating the mode toolbar when the registered mode commands change */
 	void OnEditorModeCommandsChanged();
 
+	/** Sets the parent tab of this toolbox */
+	void SetParentTab(TSharedRef<SDockTab>& InDockTab);
 private:
 
 	/** Gets the visibility for the SBorder showing toolbox editor mode inline content */
@@ -46,28 +51,30 @@ private:
 	EVisibility GetNoToolSelectedTextVisibility() const;
 
 	/** Updates the widget for showing toolbox editor mode inline content */
-	void UpdateInlineContent(TSharedPtr<SWidget> InlineContent) const;
+	void UpdateInlineContent(const TSharedPtr<IToolkit>& Toolkit, TSharedPtr<SWidget> InlineContent);
 
 	/** Creates and sets the mode toolbar */
-	void UpdateModeToolBar();
+	void UpdateModeLegacyToolBar();
 
 	/** Handles updating the mode toolbar when the user settings change */
 	void HandleUserSettingsChange( FName PropertyName );
 
-	/** Returns specified Editor modes icon, if that mode is active it adds ".Selected" onto the name */
-	FSlateIcon GetEditorModeIcon(TSharedPtr< FUICommandInfo > EditorModeUICommand, FEditorModeID EditorMode);
-
 private:
+	/** Parent tab where this toolbox is hosted */
+	TWeakPtr<SDockTab> ParentTab;
 
 	/** Level editor that we're associated with */
-	TWeakPtr<class ILevelEditor> LevelEditor;
+	TWeakPtr<ILevelEditor> LevelEditor;
 
 	/** Inline content area for editor modes */
 	TSharedPtr<SBorder> InlineContentHolder;
 
-	/** The menu extenders to populate the toolbox*/
-	TArray< TSharedPtr<FExtender> > ToolboxExtenders;
-
 	/** The container holding the mode toolbar */
-	TSharedPtr< SBorder > ModeToolBarContainer;
+	TSharedPtr<SBorder> ModeToolBarContainer;
+
+	/** The display name that the parent tab should have as its label */
+	FText TabName;
+
+	/** The icon that should be displayed in the parent tab */
+	const FSlateBrush* TabIcon;
 };
diff --git a/Engine/Source/Editor/MeshPaintMode/Private/MeshPaintModeModule.cpp b/Engine/Source/Editor/MeshPaintMode/Private/MeshPaintModeModule.cpp
--- a/Engine/Source/Editor/MeshPaintMode/Private/MeshPaintModeModule.cpp
+++ b/Engine/Source/Editor/MeshPaintMode/Private/MeshPaintModeModule.cpp
@@ -21,7 +21,7 @@ void FMeshPaintModeModule::StartupModule()
 {
 	FEditorModeRegistry::Get().RegisterMode<FEdModeMeshPaint>(
 		FBuiltinEditorModes::EM_MeshPaint,
-		NSLOCTEXT("MeshPaint_Mode", "MeshPaint_ModeName", "Paint"),
+		NSLOCTEXT("MeshPaint_Mode", "MeshPaint_ModeName", "Mesh Paint"),
 		FSlateIcon(FEditorStyle::GetStyleSetName(), "LevelEditor.MeshPaintMode", "LevelEditor.MeshPaintMode.Small"),
 		true, 200 );
 
diff --git a/Engine/Source/Editor/PlacementMode/Private/PlacementMode.cpp b/Engine/Source/Editor/PlacementMode/Private/PlacementMode.cpp
--- a/Engine/Source/Editor/PlacementMode/Private/PlacementMode.cpp
+++ b/Engine/Source/Editor/PlacementMode/Private/PlacementMode.cpp
@@ -14,6 +14,7 @@
 
 #include "Toolkits/ToolkitManager.h"
 #include "IAssetViewport.h"
+#include "Settings/LevelEditorMiscSettings.h"
 
 FPlacementMode::FPlacementMode()
 	: AssetsToPlace()
@@ -30,13 +31,11 @@ FPlacementMode::~FPlacementMode()
 
 void FPlacementMode::Initialize()
 {
-	//GConfig->GetFloat(TEXT("PlacementMode"), TEXT("AssetThumbnailScale"), AssetThumbnailScale, GEditorPerProjectIni);
-	//GConfig->GetBool( TEXT( "PlacementMode" ), TEXT( "ShowOtherDeveloperAssets" ), ShowOtherDeveloperAssets, GEditorPerProjectIni );
 }
 
 bool FPlacementMode::UsesToolkits() const
 {
-	return true;
+	return GetDefault<ULevelEditorMiscSettings>()->bEnableLegacyEditorModeUI;
 }
 
 void FPlacementMode::Enter()
@@ -44,7 +43,8 @@ void FPlacementMode::Enter()
 	// Call parent implementation
 	FEdMode::Enter();
 
-	if ( !Toolkit.IsValid() )
+	
+	if(!Toolkit.IsValid() && GetDefault<ULevelEditorMiscSettings>()->bEnableLegacyEditorModeUI)
 	{
 		Toolkit = MakeShareable( new FPlacementModeToolkit );
 		Toolkit->Init(Owner->GetToolkitHost());
diff --git a/Engine/Source/Editor/PlacementMode/Private/PlacementModeModule.cpp b/Engine/Source/Editor/PlacementMode/Private/PlacementModeModule.cpp
--- a/Engine/Source/Editor/PlacementMode/Private/PlacementModeModule.cpp
+++ b/Engine/Source/Editor/PlacementMode/Private/PlacementModeModule.cpp
@@ -48,6 +48,8 @@
 #include "AssetToolsModule.h"
 #include "Kismet2/KismetEditorUtilities.h"
 #include "ActorFactories/ActorFactoryPlanarReflection.h"
+#include "SPlacementModeTools.h"
+#include "Settings/LevelEditorMiscSettings.h"
 
 
 TOptional<FLinearColor> GetBasicShapeColorOverride()
@@ -81,12 +83,11 @@ void FPlacementModeModule::StartupModule()
 		RecentlyPlaced.Add(FActorPlacementInfo(RecentlyPlacedAsStrings[Index]));
 	}
 
-
 	FEditorModeRegistry::Get().RegisterMode<FPlacementMode>(
 		FBuiltinEditorModes::EM_Placement,
 		NSLOCTEXT("PlacementMode", "DisplayName", "Place"),
 		FSlateIcon(FEditorStyle::GetStyleSetName(), "LevelEditor.PlacementMode", "LevelEditor.PlacementMode.Small"),
-		true, 0);
+		GetDefault<ULevelEditorMiscSettings>()->bEnableLegacyEditorModeUI, 0);
 
 	FAssetRegistryModule& AssetRegistryModule = FModuleManager::LoadModuleChecked<FAssetRegistryModule>(TEXT("AssetRegistry"));
 	AssetRegistryModule.Get().OnAssetRemoved().AddRaw(this, &FPlacementModeModule::OnAssetRemoved);
@@ -322,6 +323,11 @@ void FPlacementModeModule::AddToRecentlyPlaced(UObject* Asset, UActorFactory* Fa
 	AddToRecentlyPlaced(Assets, FactoryUsed);
 }
 
+TSharedRef<SWidget> FPlacementModeModule::CreatePlacementModeBrowser()
+{
+	return SNew(SPlacementModeTools);
+}
+
 bool FPlacementModeModule::RegisterPlacementCategory(const FPlacementCategoryInfo& Info)
 {
 	if (Categories.Contains(Info.UniqueHandle))
diff --git a/Engine/Source/Editor/PlacementMode/Private/SPlacementModeTools.cpp b/Engine/Source/Editor/PlacementMode/Private/SPlacementModeTools.cpp
--- a/Engine/Source/Editor/PlacementMode/Private/SPlacementModeTools.cpp
+++ b/Engine/Source/Editor/PlacementMode/Private/SPlacementModeTools.cpp
@@ -298,7 +298,10 @@ void SPlacementModeTools::Construct( const FArguments& InArgs )
 	bNeedsUpdate = true;
 
 	FPlacementMode* PlacementEditMode = (FPlacementMode*)GLevelEditorModeTools().GetActiveMode( FBuiltinEditorModes::EM_Placement );
-	PlacementEditMode->AddValidFocusTargetForPlacement( SharedThis( this ) );
+	if (PlacementEditMode)
+	{
+		PlacementEditMode->AddValidFocusTargetForPlacement(SharedThis(this));
+	}
 
 	SearchTextFilter = MakeShareable(new FPlacementAssetEntryTextFilter(
 		FPlacementAssetEntryTextFilter::FItemToStringArray::CreateStatic(&PlacementViewFilter::GetBasicStrings)
diff --git a/Engine/Source/Editor/PlacementMode/Public/IPlacementModeModule.h b/Engine/Source/Editor/PlacementMode/Public/IPlacementModeModule.h
--- a/Engine/Source/Editor/PlacementMode/Public/IPlacementModeModule.h
+++ b/Engine/Source/Editor/PlacementMode/Public/IPlacementModeModule.h
@@ -249,6 +249,10 @@ public:
 	virtual FOnStoppedPlacingEvent& OnStoppedPlacing() = 0;
 	virtual void BroadcastStoppedPlacing( bool bWasSuccessfullyPlaced ) = 0;
 
+	/**
+	 * Creates the placement browser widget
+	 */
+	virtual TSharedRef<SWidget> CreatePlacementModeBrowser() = 0;
 public:
 
 	/**
@@ -321,4 +325,3 @@ public:
 	 */
 	virtual void RegenerateItemsForCategory(FName Category) = 0;
 };
-
diff --git a/Engine/Source/Editor/PlacementMode/Public/PlacementModeModule.h b/Engine/Source/Editor/PlacementMode/Public/PlacementModeModule.h
--- a/Engine/Source/Editor/PlacementMode/Public/PlacementModeModule.h
+++ b/Engine/Source/Editor/PlacementMode/Public/PlacementModeModule.h
@@ -101,7 +101,7 @@ public:
 		StoppedPlacingEvent.Broadcast(bWasSuccessfullyPlaced);
 	}
 
-public:
+	virtual TSharedRef<SWidget> CreatePlacementModeBrowser() override;
 
 	virtual bool RegisterPlacementCategory(const FPlacementCategoryInfo& Info);
 
diff --git a/Engine/Source/Editor/UnrealEd/Private/EditorModeManager.cpp b/Engine/Source/Editor/UnrealEd/Private/EditorModeManager.cpp
--- a/Engine/Source/Editor/UnrealEd/Private/EditorModeManager.cpp
+++ b/Engine/Source/Editor/UnrealEd/Private/EditorModeManager.cpp
@@ -913,11 +913,17 @@ TSharedRef<SDockTab> FEditorModeTools::MakeModeToolbarTab()
 
 }
 
-bool FEditorModeTools::ShouldShowModeToolbar()
+bool FEditorModeTools::ShouldShowModeToolbar() const
 {
 	return ActiveToolBarRows.Num() > 0;
 }
 
+bool FEditorModeTools::ShouldShowModeToolbox() const
+{
+	// This could ideally ask each active mode if it has any tools but when developing a new mode the toolbox not appearing by default could be hard to understand
+	return !IsModeActive(FBuiltinEditorModes::EM_Default);
+}
+
 void FEditorModeTools::ActivateMode(FEditorModeID InID, bool bToggle)
 {
 	static bool bReentrant = false;
diff --git a/Engine/Source/Editor/UnrealEd/Private/EditorModeRegistry.cpp b/Engine/Source/Editor/UnrealEd/Private/EditorModeRegistry.cpp
--- a/Engine/Source/Editor/UnrealEd/Private/EditorModeRegistry.cpp
+++ b/Engine/Source/Editor/UnrealEd/Private/EditorModeRegistry.cpp
@@ -15,6 +15,7 @@
 #include "Editor/FoliageEdit/Public/FoliageEditModule.h"
 #include "Editor/VirtualTexturingEditor/Public/VirtualTexturingEditorModule.h"
 #include "Tools/UEdMode.h"
+#include "Settings/LevelEditorMiscSettings.h"
 
 FEditorModeInfo::FEditorModeInfo()
 	: ID(NAME_None)
@@ -49,8 +50,20 @@ void FEditorModeRegistry::Initialize()
 {
 	Get();
 
-	// Add default editor modes
-	GModeRegistry->RegisterMode<FEdModeDefault>(FBuiltinEditorModes::EM_Default);
+	if(!GetDefault<ULevelEditorMiscSettings>()->bEnableLegacyEditorModeUI)
+	{
+		// Add default editor modes
+		FEditorModeRegistry::Get().RegisterMode<FEdModeDefault>(
+			FBuiltinEditorModes::EM_Default,
+			NSLOCTEXT("DefaultMode", "DisplayName", "Select"),
+			FSlateIcon(FEditorStyle::GetStyleSetName(), "LevelEditor.SelectMode", "LevelEditor.SelectMode.Small"),
+			true, 0);
+	}
+	else
+	{
+		GModeRegistry->RegisterMode<FEdModeDefault>(FBuiltinEditorModes::EM_Default);
+	}
+
 	GModeRegistry->RegisterMode<FEdModeInterpEdit>(FBuiltinEditorModes::EM_InterpEdit);
 
 	// Load editor mode modules that will automatically register their editor modes, and clean themselves up on unload.
diff --git a/Engine/Source/Editor/UnrealEd/Public/EdMode.h b/Engine/Source/Editor/UnrealEd/Public/EdMode.h
--- a/Engine/Source/Editor/UnrealEd/Public/EdMode.h
+++ b/Engine/Source/Editor/UnrealEd/Public/EdMode.h
@@ -403,7 +403,7 @@ protected:
 ------------------------------------------------------------------------------*/
 
 /**
- * The default editing mode.  User can work with BSP and the builder brush. Vector and array properties are also visually editable.
+ * The default editing mode. 
  */
 class FEdModeDefault : public FEdMode
 {
diff --git a/Engine/Source/Editor/UnrealEd/Public/EditorModeManager.h b/Engine/Source/Editor/UnrealEd/Public/EditorModeManager.h
--- a/Engine/Source/Editor/UnrealEd/Public/EditorModeManager.h
+++ b/Engine/Source/Editor/UnrealEd/Public/EditorModeManager.h
@@ -53,6 +53,11 @@ public:
 	 */
 	void RemoveDefaultMode( const FEditorModeID DefaultModeID );
 
+	/**
+	 * Returns whether or not the provided mode ID is a default mode
+	 */
+	bool IsDefaultMode(const FEditorModeID ModeID) const { return DefaultModeIDs.Contains(ModeID); }
+
 	/**
 	 * Activates the default modes defined by this class.  Note that there can be more than one default mode, and this call will activate them all in sequence.
 	 */
@@ -91,8 +96,12 @@ public:
 	/**
 	 * Whether or not the mode toolbar should be shown.  If any active modes generated a toolbar this method will return true
 	 */
-	bool ShouldShowModeToolbar();
+	bool ShouldShowModeToolbar() const;
 
+	/**
+	 * Whether or not the mode toolbox (where mode details panels and some tools are) should be shown.
+	 */
+	bool ShouldShowModeToolbox() const;
 protected:
 	
 	/** Deactivates the editor mode at the specified index */
