Reduce duplication of Level editor tab FNames for readability, reusability, and perf reasons

diff --git a/Engine/Source/Editor/LevelEditor/Private/LevelEditor.cpp b/Engine/Source/Editor/LevelEditor/Private/LevelEditor.cpp
--- a/Engine/Source/Editor/LevelEditor/Private/LevelEditor.cpp
+++ b/Engine/Source/Editor/LevelEditor/Private/LevelEditor.cpp
@@ -49,6 +49,31 @@ const FName LevelEditorApp = FName(TEXT("LevelEditorApp"));
 const FName MainFrame("MainFrame");
 const FName CommonMenuExtensionsName(TEXT("CommonMenuExtensions"));
 
+const FName LevelEditorTabIds::LevelEditorViewport(TEXT("LevelEditorViewport"));
+const FName LevelEditorTabIds::LevelEditorViewport_Clone1(TEXT("LevelEditorViewport_Clone1"));
+const FName LevelEditorTabIds::LevelEditorViewport_Clone2(TEXT("LevelEditorViewport_Clone2"));
+const FName LevelEditorTabIds::LevelEditorViewport_Clone3(TEXT("LevelEditorViewport_Clone3"));
+const FName LevelEditorTabIds::LevelEditorViewport_Clone4(TEXT("LevelEditorViewport_Clone4"));
+const FName LevelEditorTabIds::LevelEditorToolBar(TEXT("LevelEditorToolBar"));
+const FName LevelEditorTabIds::LevelEditorToolBox(TEXT("LevelEditorToolBox"));
+const FName LevelEditorTabIds::LevelEditorSelectionDetails(TEXT("LevelEditorSelectionDetails"));
+const FName LevelEditorTabIds::LevelEditorSelectionDetails2(TEXT("LevelEditorSelectionDetails2"));
+const FName LevelEditorTabIds::LevelEditorSelectionDetails3(TEXT("LevelEditorSelectionDetails3"));
+const FName LevelEditorTabIds::LevelEditorSelectionDetails4(TEXT("LevelEditorSelectionDetails4"));
+const FName LevelEditorTabIds::PlacementBrowser(TEXT("PlacementBrowser"));
+const FName LevelEditorTabIds::LevelEditorBuildAndSubmit(TEXT("LevelEditorBuildAndSubmit"));
+const FName LevelEditorTabIds::LevelEditorSceneOutliner(TEXT("LevelEditorSceneOutliner"));
+const FName LevelEditorTabIds::LevelEditorStatsViewer(TEXT("LevelEditorStatsViewer"));
+const FName LevelEditorTabIds::LevelEditorLayerBrowser(TEXT("LevelEditorLayerBrowser"));
+const FName LevelEditorTabIds::Sequencer(TEXT("Sequencer"));
+const FName LevelEditorTabIds::SequencerGraphEditor(TEXT("SequencerGraphEditor"));
+const FName LevelEditorTabIds::WorldSettings(TEXT("WorldSettingsTab"));
+const FName LevelEditorTabIds::WorldBrowserComposition(TEXT("WorldBrowserComposition"));
+const FName LevelEditorTabIds::WorldBrowserHierarchy(TEXT("WorldBrowserHierarchy"));
+const FName LevelEditorTabIds::WorldBrowserDetails(TEXT("WorldBrowserDetails"));
+const FName LevelEditorTabIds::LevelEditorHierarchicalLODOutliner(TEXT("LevelEditorHierarchicalLODOutliner"));
+const FName LevelEditorTabIds::OutputLog(TEXT("OutputLog"));
+
 FLevelEditorModule::FLevelEditorModule()
 	: ToggleImmersiveConsoleCommand(
 	TEXT( "LevelEditor.ToggleImmersive" ),
@@ -344,31 +369,26 @@ void FLevelEditorModule::SummonSelectionDetails()
 void FLevelEditorModule::SummonBuildAndSubmit()
 {
 	TSharedPtr<SLevelEditor> LevelEditorInstance = LevelEditorInstancePtr.Pin();
-	LevelEditorInstance->InvokeTab("LevelEditorBuildAndSubmit");
+	LevelEditorInstance->InvokeTab(LevelEditorTabIds::LevelEditorBuildAndSubmit);
 }
 
-void FLevelEditorModule::SummonLevelBrowser()
-{
-	TSharedPtr<SLevelEditor> LevelEditorInstance = LevelEditorInstancePtr.Pin();
-	LevelEditorInstance->InvokeTab("LevelEditorLevelBrowser");
-}
 
 void FLevelEditorModule::SummonWorldBrowserHierarchy()
 {
 	TSharedPtr<SLevelEditor> LevelEditorInstance = LevelEditorInstancePtr.Pin();
-	LevelEditorInstance->InvokeTab("WorldBrowserHierarchy");
+	LevelEditorInstance->InvokeTab(LevelEditorTabIds::WorldBrowserHierarchy);
 }
 
 void FLevelEditorModule::SummonWorldBrowserDetails()
 {
 	TSharedPtr<SLevelEditor> LevelEditorInstance = LevelEditorInstancePtr.Pin();
-	LevelEditorInstance->InvokeTab("WorldBrowserDetails");
+	LevelEditorInstance->InvokeTab(LevelEditorTabIds::WorldBrowserDetails);
 }
 
 void FLevelEditorModule::SummonWorldBrowserComposition()
 {
 	TSharedPtr<SLevelEditor> LevelEditorInstance = LevelEditorInstancePtr.Pin();
-	LevelEditorInstance->InvokeTab("WorldBrowserComposition");
+	LevelEditorInstance->InvokeTab(LevelEditorTabIds::WorldBrowserComposition);
 }
 
 // @todo remove when world-centric mode is added
diff --git a/Engine/Source/Editor/LevelEditor/Private/SLevelEditor.cpp b/Engine/Source/Editor/LevelEditor/Private/SLevelEditor.cpp
--- a/Engine/Source/Editor/LevelEditor/Private/SLevelEditor.cpp
+++ b/Engine/Source/Editor/LevelEditor/Private/SLevelEditor.cpp
@@ -53,13 +53,8 @@
 #include "IPlacementModeModule.h"
 #include "Settings/LevelEditorMiscSettings.h"
 
-static const FName LevelEditorBuildAndSubmitTab("LevelEditorBuildAndSubmit");
-static const FName LevelEditorStatsViewerTab("LevelEditorStatsViewer");
 static const FName MainFrameModuleName("MainFrame");
 static const FName LevelEditorModuleName("LevelEditor");
-static const FName WorldBrowserHierarchyTab("WorldBrowserHierarchy");
-static const FName WorldBrowserDetailsTab("WorldBrowserDetails");
-static const FName WorldBrowserCompositionTab("WorldBrowserComposition");
 
 
 namespace LevelEditorConstants
@@ -562,7 +557,7 @@ void SLevelEditor::AttachSequencer( TSharedPtr<SWidget> SequencerWidget, TShared
 	if( !bIsReentrant )
 	{
 		TSharedRef<SDockTab> Tab = SNew(SDockTab);
-		Tab = InvokeTab("Sequencer");
+		Tab = InvokeTab(LevelEditorTabIds::Sequencer);
 
 		// Close the sequence editor after invoking a sequencer tab instead of before so that the existing asset editor doesn't refer to a stale sequencer.
 		if(SequencerAssetEditor.IsValid())
@@ -621,23 +616,23 @@ TSharedRef<SDockTab> SLevelEditor::SummonDetailsPanel( FName TabIdentifier )
 /** Method to call when a tab needs to be spawned by the FLayoutService */
 TSharedRef<SDockTab> SLevelEditor::SpawnLevelEditorTab( const FSpawnTabArgs& Args, FName TabIdentifier, FString InitializationPayload )
 {
-	if( TabIdentifier == TEXT("LevelEditorViewport" ) )
+	if( TabIdentifier == LevelEditorTabIds::LevelEditorViewport)
 	{
 		return this->BuildViewportTab( NSLOCTEXT("LevelViewportTypes", "LevelEditorViewport", "Viewport 1"), TEXT("Viewport 1"), InitializationPayload );
 	}
-	else if( TabIdentifier == TEXT("LevelEditorViewport_Clone1" ) )
+	else if( TabIdentifier == LevelEditorTabIds::LevelEditorViewport_Clone1)
 	{
 		return this->BuildViewportTab( NSLOCTEXT("LevelViewportTypes", "LevelEditorViewport_Clone1", "Viewport 2"), TEXT("Viewport 2"), InitializationPayload );
 	}
-	else if( TabIdentifier == TEXT("LevelEditorViewport_Clone2" ) )
+	else if( TabIdentifier == LevelEditorTabIds::LevelEditorViewport_Clone2)
 	{
 		return this->BuildViewportTab( NSLOCTEXT("LevelViewportTypes", "LevelEditorViewport_Clone2", "Viewport 3"), TEXT("Viewport 3"), InitializationPayload );
 	}
-	else if( TabIdentifier == TEXT("LevelEditorViewport_Clone3" ) )
+	else if( TabIdentifier == LevelEditorTabIds::LevelEditorViewport_Clone3)
 	{
 		return this->BuildViewportTab( NSLOCTEXT("LevelViewportTypes", "LevelEditorViewport_Clone3", "Viewport 4"), TEXT("Viewport 4"), InitializationPayload );
 	}
-	else if( TabIdentifier == TEXT( "LevelEditorToolBar") )
+	else if( TabIdentifier == LevelEditorTabIds::LevelEditorToolBar)
 	{
 		return SNew( SDockTab )
 			.Label( NSLOCTEXT("LevelEditor", "ToolBarTabTitle", "Toolbar") )
@@ -666,7 +661,7 @@ TSharedRef<SDockTab> SLevelEditor::SpawnLevelEditorTab( const FSpawnTabArgs& Arg
 		GUnrealEd->UpdateFloatingPropertyWindows();
 		return DetailsPanel;
 	}
-	else if( TabIdentifier == TEXT("LevelEditorToolBox") )
+	else if( TabIdentifier == LevelEditorTabIds::LevelEditorToolBox )
 	{
 		TSharedRef<SLevelEditorToolBox> NewToolBox = StaticCastSharedRef<SLevelEditorToolBox>( CreateToolBox() );
 
@@ -685,7 +680,7 @@ TSharedRef<SDockTab> SLevelEditor::SpawnLevelEditorTab( const FSpawnTabArgs& Arg
 
 		return DockTab;
 	}
-	else if (TabIdentifier == TEXT("PlacementBrowser"))
+	else if (TabIdentifier == LevelEditorTabIds::PlacementBrowser)
 	{
 		if(!GetDefault<ULevelEditorMiscSettings>()->bEnableLegacyEditorModeUI)
 		{
@@ -699,7 +694,7 @@ TSharedRef<SDockTab> SLevelEditor::SpawnLevelEditorTab( const FSpawnTabArgs& Arg
 				];
 		}
 	}
-	else if( TabIdentifier == LevelEditorBuildAndSubmitTab )
+	else if( TabIdentifier == LevelEditorTabIds::LevelEditorBuildAndSubmit )
 	{
 		TSharedRef<SLevelEditorBuildAndSubmit> NewBuildAndSubmit = SNew( SLevelEditorBuildAndSubmit, SharedThis( this ) );
 
@@ -714,7 +709,7 @@ TSharedRef<SDockTab> SLevelEditor::SpawnLevelEditorTab( const FSpawnTabArgs& Arg
 
 		return NewTab;
 	}
-	else if( TabIdentifier == TEXT("LevelEditorSceneOutliner") )
+	else if( TabIdentifier == LevelEditorTabIds::LevelEditorSceneOutliner)
 	{
 		SceneOutliner::FInitializationOptions InitOptions;
 		InitOptions.bShowTransient = true;
@@ -773,7 +768,7 @@ TSharedRef<SDockTab> SLevelEditor::SpawnLevelEditorTab( const FSpawnTabArgs& Arg
 				]
 			];
 	}
-	else if( TabIdentifier == TEXT("LevelEditorLayerBrowser") )
+	else if(TabIdentifier == LevelEditorTabIds::LevelEditorLayerBrowser)
 	{
 		FLayersModule& LayersModule = FModuleManager::LoadModuleChecked<FLayersModule>( "Layers" );
 		return SNew( SDockTab )
@@ -789,7 +784,7 @@ TSharedRef<SDockTab> SLevelEditor::SpawnLevelEditorTab( const FSpawnTabArgs& Arg
 				]
 			];
 	}
-	else if (TabIdentifier == TEXT("LevelEditorHierarchicalLODOutliner"))
+	else if (TabIdentifier == LevelEditorTabIds::LevelEditorHierarchicalLODOutliner)
 	{
 		FText Label = NSLOCTEXT("LevelEditor", "HLODOutlinerTabTitle", "Hierarchical LOD Outliner");
 
@@ -802,7 +797,7 @@ TSharedRef<SDockTab> SLevelEditor::SpawnLevelEditorTab( const FSpawnTabArgs& Arg
 				HLODModule.CreateHLODOutlinerWidget()
 			];
 	}
-	else if( TabIdentifier == WorldBrowserHierarchyTab )
+	else if( TabIdentifier == LevelEditorTabIds::WorldBrowserHierarchy)
 	{
 		FWorldBrowserModule& WorldBrowserModule = FModuleManager::LoadModuleChecked<FWorldBrowserModule>( "WorldBrowser" );
 		return SNew( SDockTab )
@@ -812,7 +807,7 @@ TSharedRef<SDockTab> SLevelEditor::SpawnLevelEditorTab( const FSpawnTabArgs& Arg
 				WorldBrowserModule.CreateWorldBrowserHierarchy()
 			];
 	}
-	else if( TabIdentifier == WorldBrowserDetailsTab )
+	else if( TabIdentifier == LevelEditorTabIds::WorldBrowserDetails)
 	{
 		FWorldBrowserModule& WorldBrowserModule = FModuleManager::LoadModuleChecked<FWorldBrowserModule>( "WorldBrowser" );
 		return SNew( SDockTab )
@@ -822,7 +817,7 @@ TSharedRef<SDockTab> SLevelEditor::SpawnLevelEditorTab( const FSpawnTabArgs& Arg
 				WorldBrowserModule.CreateWorldBrowserDetails()
 			];
 	}
-	else if( TabIdentifier == WorldBrowserCompositionTab )
+	else if( TabIdentifier == LevelEditorTabIds::WorldBrowserComposition )
 	{
 		FWorldBrowserModule& WorldBrowserModule = FModuleManager::LoadModuleChecked<FWorldBrowserModule>( "WorldBrowser" );
 		return SNew( SDockTab )
@@ -832,7 +827,7 @@ TSharedRef<SDockTab> SLevelEditor::SpawnLevelEditorTab( const FSpawnTabArgs& Arg
 				WorldBrowserModule.CreateWorldBrowserComposition()
 			];
 	}
-	else if( TabIdentifier == TEXT("Sequencer") )
+	else if(TabIdentifier == LevelEditorTabIds::Sequencer)
 	{
 		if (FSlateStyleRegistry::FindSlateStyle("LevelSequenceEditorStyle"))
 		{
@@ -845,7 +840,7 @@ TSharedRef<SDockTab> SLevelEditor::SpawnLevelEditorTab( const FSpawnTabArgs& Arg
 				];
 		}
 	}
-	else if( TabIdentifier == TEXT("SequencerGraphEditor") )
+	else if(TabIdentifier == LevelEditorTabIds::SequencerGraphEditor )
 	{
 		const FSlateIcon SequencerGraphIcon = FSlateIcon(FEditorStyle::GetStyleSetName(), "GenericCurveEditor.TabIcon");
 		// @todo sequencer: remove when world-centric mode is added
@@ -856,7 +851,7 @@ TSharedRef<SDockTab> SLevelEditor::SpawnLevelEditorTab( const FSpawnTabArgs& Arg
 				SNullWidget::NullWidget
 			];
 	}
-	else if( TabIdentifier == LevelEditorStatsViewerTab )
+	else if( TabIdentifier == LevelEditorTabIds::LevelEditorStatsViewer )
 	{
 		FStatsViewerModule& StatsViewerModule = FModuleManager::Get().LoadModuleChecked<FStatsViewerModule>( "StatsViewer" );
 		return SNew( SDockTab )
@@ -866,7 +861,7 @@ TSharedRef<SDockTab> SLevelEditor::SpawnLevelEditorTab( const FSpawnTabArgs& Arg
 				StatsViewerModule.CreateStatsViewer()
 			];
 	}
-	else if ( TabIdentifier == "WorldSettingsTab" )
+	else if (TabIdentifier == LevelEditorTabIds::WorldSettings)
 	{
 		FPropertyEditorModule& PropPlugin = FModuleManager::LoadModuleChecked<FPropertyEditorModule>("PropertyEditor");
 		FDetailsViewArgs DetailsViewArgs( false, false, true, FDetailsViewArgs::HideNameArea, false, GUnrealEd );
@@ -909,7 +904,11 @@ TSharedRef<SDockTab> SLevelEditor::InvokeTab( FName TabID )
 
 void SLevelEditor::SyncDetailsToSelection()
 {
-	static const FName DetailsTabIdentifiers[] = { "LevelEditorSelectionDetails", "LevelEditorSelectionDetails2", "LevelEditorSelectionDetails3", "LevelEditorSelectionDetails4" };
+	static const FName DetailsTabIdentifiers[] = { 
+		LevelEditorTabIds::LevelEditorSelectionDetails, 
+		LevelEditorTabIds::LevelEditorSelectionDetails2, 
+		LevelEditorTabIds::LevelEditorSelectionDetails3, 
+		LevelEditorTabIds::LevelEditorSelectionDetails4 };
 
 	FPropertyEditorModule& PropPlugin = FModuleManager::LoadModuleChecked<FPropertyEditorModule>("PropertyEditor");
 
@@ -1083,25 +1082,25 @@ TSharedRef<SWidget> SLevelEditor::RestoreContentArea( const TSharedRef<SDockTab>
 			const FText ViewportTooltip = NSLOCTEXT("LevelEditorTabs", "LevelEditorViewportTooltip", "Open a Viewport tab. Use this to view and edit the current level.");
 			const FSlateIcon ViewportIcon(FEditorStyle::GetStyleSetName(), "LevelEditor.Tabs.Viewports");
 
-			LevelEditorTabManager->RegisterTabSpawner( "LevelEditorViewport", FOnSpawnTab::CreateSP(this, &SLevelEditor::SpawnLevelEditorTab, FName("LevelEditorViewport"), FString()) )
+			LevelEditorTabManager->RegisterTabSpawner(LevelEditorTabIds::LevelEditorViewport, FOnSpawnTab::CreateSP(this, &SLevelEditor::SpawnLevelEditorTab, LevelEditorTabIds::LevelEditorViewport, FString()) )
 				.SetDisplayName(NSLOCTEXT("LevelEditorTabs", "LevelEditorViewport", "Viewport 1"))
 				.SetTooltipText(ViewportTooltip)
 				.SetGroup( MenuStructure.GetLevelEditorViewportsCategory() )
 				.SetIcon(ViewportIcon);
 
-			LevelEditorTabManager->RegisterTabSpawner( "LevelEditorViewport_Clone1", FOnSpawnTab::CreateSP(this, &SLevelEditor::SpawnLevelEditorTab, FName("LevelEditorViewport_Clone1"), FString()) )
+			LevelEditorTabManager->RegisterTabSpawner(LevelEditorTabIds::LevelEditorViewport_Clone1, FOnSpawnTab::CreateSP(this, &SLevelEditor::SpawnLevelEditorTab, LevelEditorTabIds::LevelEditorViewport_Clone1, FString()) )
 				.SetDisplayName(NSLOCTEXT("LevelEditorTabs", "LevelEditorViewport_Clone1", "Viewport 2"))
 				.SetTooltipText(ViewportTooltip)
 				.SetGroup( MenuStructure.GetLevelEditorViewportsCategory() )
 				.SetIcon(ViewportIcon);
 
-			LevelEditorTabManager->RegisterTabSpawner( "LevelEditorViewport_Clone2", FOnSpawnTab::CreateSP(this, &SLevelEditor::SpawnLevelEditorTab, FName("LevelEditorViewport_Clone2"), FString()) )
+			LevelEditorTabManager->RegisterTabSpawner(LevelEditorTabIds::LevelEditorViewport_Clone2, FOnSpawnTab::CreateSP(this, &SLevelEditor::SpawnLevelEditorTab, LevelEditorTabIds::LevelEditorViewport_Clone2, FString()) )
 				.SetDisplayName(NSLOCTEXT("LevelEditorTabs", "LevelEditorViewport_Clone2", "Viewport 3"))
 				.SetTooltipText(ViewportTooltip)
 				.SetGroup( MenuStructure.GetLevelEditorViewportsCategory() )
 				.SetIcon(ViewportIcon);
 
-				LevelEditorTabManager->RegisterTabSpawner( "LevelEditorViewport_Clone3", FOnSpawnTab::CreateSP(this, &SLevelEditor::SpawnLevelEditorTab, FName("LevelEditorViewport_Clone3"), FString()) )
+				LevelEditorTabManager->RegisterTabSpawner(LevelEditorTabIds::LevelEditorViewport_Clone3, FOnSpawnTab::CreateSP(this, &SLevelEditor::SpawnLevelEditorTab, LevelEditorTabIds::LevelEditorViewport_Clone3, FString()) )
 				.SetDisplayName(NSLOCTEXT("LevelEditorTabs", "LevelEditorViewport_Clone3", "Viewport 4"))
 				.SetTooltipText(ViewportTooltip)
 				.SetGroup( MenuStructure.GetLevelEditorViewportsCategory() )
@@ -1110,7 +1109,7 @@ TSharedRef<SWidget> SLevelEditor::RestoreContentArea( const TSharedRef<SDockTab>
 
 		{
 			const FSlateIcon ToolbarIcon(FEditorStyle::GetStyleSetName(), "LevelEditor.Tabs.Toolbar");
-			LevelEditorTabManager->RegisterTabSpawner( "LevelEditorToolBar", FOnSpawnTab::CreateSP<SLevelEditor, FName, FString>(this, &SLevelEditor::SpawnLevelEditorTab, FName("LevelEditorToolBar"), FString()))
+			LevelEditorTabManager->RegisterTabSpawner(LevelEditorTabIds::LevelEditorToolBar, FOnSpawnTab::CreateSP<SLevelEditor, FName, FString>(this, &SLevelEditor::SpawnLevelEditorTab, LevelEditorTabIds::LevelEditorToolBar, FString()))
 				.SetDisplayName(NSLOCTEXT("LevelEditorTabs", "LevelEditorToolBar", "Toolbar"))
 				.SetTooltipText(NSLOCTEXT("LevelEditorTabs", "LevelEditorToolBarTooltipText", "Open the Toolbar tab, which provides access to the most common / important actions."))
 				.SetGroup( MenuStructure.GetLevelEditorCategory() )
@@ -1121,25 +1120,25 @@ TSharedRef<SWidget> SLevelEditor::RestoreContentArea( const TSharedRef<SDockTab>
 			const FText DetailsTooltip = NSLOCTEXT("LevelEditorTabs", "LevelEditorSelectionDetailsTooltip", "Open a Details tab. Use this to view and edit properties of the selected object(s).");
 			const FSlateIcon DetailsIcon(FEditorStyle::GetStyleSetName(), "LevelEditor.Tabs.Details");
 
-			LevelEditorTabManager->RegisterTabSpawner( "LevelEditorSelectionDetails", FOnSpawnTab::CreateSP<SLevelEditor, FName, FString>(this, &SLevelEditor::SpawnLevelEditorTab, FName("LevelEditorSelectionDetails"), FString()) )
+			LevelEditorTabManager->RegisterTabSpawner(LevelEditorTabIds::LevelEditorSelectionDetails, FOnSpawnTab::CreateSP<SLevelEditor, FName, FString>(this, &SLevelEditor::SpawnLevelEditorTab, LevelEditorTabIds::LevelEditorSelectionDetails, FString()) )
 				.SetDisplayName(NSLOCTEXT("LevelEditorTabs", "LevelEditorSelectionDetails", "Details 1"))
 				.SetTooltipText(DetailsTooltip)
 				.SetGroup( MenuStructure.GetLevelEditorDetailsCategory() )
 				.SetIcon( DetailsIcon );
 
-			LevelEditorTabManager->RegisterTabSpawner( "LevelEditorSelectionDetails2", FOnSpawnTab::CreateSP<SLevelEditor, FName, FString>(this, &SLevelEditor::SpawnLevelEditorTab, FName("LevelEditorSelectionDetails2"), FString()) )
+			LevelEditorTabManager->RegisterTabSpawner(LevelEditorTabIds::LevelEditorSelectionDetails2, FOnSpawnTab::CreateSP<SLevelEditor, FName, FString>(this, &SLevelEditor::SpawnLevelEditorTab, LevelEditorTabIds::LevelEditorSelectionDetails2, FString()) )
 				.SetDisplayName(NSLOCTEXT("LevelEditorTabs", "LevelEditorSelectionDetails2", "Details 2"))
 				.SetTooltipText(DetailsTooltip)
 				.SetGroup( MenuStructure.GetLevelEditorDetailsCategory() )
 				.SetIcon( DetailsIcon );
 
-			LevelEditorTabManager->RegisterTabSpawner( "LevelEditorSelectionDetails3", FOnSpawnTab::CreateSP<SLevelEditor, FName, FString>(this, &SLevelEditor::SpawnLevelEditorTab, FName("LevelEditorSelectionDetails3"), FString()) )
+			LevelEditorTabManager->RegisterTabSpawner(LevelEditorTabIds::LevelEditorSelectionDetails3, FOnSpawnTab::CreateSP<SLevelEditor, FName, FString>(this, &SLevelEditor::SpawnLevelEditorTab, LevelEditorTabIds::LevelEditorSelectionDetails3, FString()) )
 				.SetDisplayName(NSLOCTEXT("LevelEditorTabs", "LevelEditorSelectionDetails3", "Details 3"))
 				.SetTooltipText(DetailsTooltip)
 				.SetGroup( MenuStructure.GetLevelEditorDetailsCategory() )
 				.SetIcon( DetailsIcon );
 
-			LevelEditorTabManager->RegisterTabSpawner( "LevelEditorSelectionDetails4", FOnSpawnTab::CreateSP<SLevelEditor, FName, FString>(this, &SLevelEditor::SpawnLevelEditorTab, FName("LevelEditorSelectionDetails4"), FString()) )
+			LevelEditorTabManager->RegisterTabSpawner(LevelEditorTabIds::LevelEditorSelectionDetails4, FOnSpawnTab::CreateSP<SLevelEditor, FName, FString>(this, &SLevelEditor::SpawnLevelEditorTab, LevelEditorTabIds::LevelEditorSelectionDetails4, FString()) )
 				.SetDisplayName(NSLOCTEXT("LevelEditorTabs", "LevelEditorSelectionDetails4", "Details 4"))
 				.SetTooltipText(DetailsTooltip)
 				.SetGroup( MenuStructure.GetLevelEditorDetailsCategory() )
@@ -1150,7 +1149,7 @@ TSharedRef<SWidget> SLevelEditor::RestoreContentArea( const TSharedRef<SDockTab>
 
 			FCanSpawnTab CanSpawnTabDelegate = FCanSpawnTab::CreateSP(this, &SLevelEditor::CanSpawnEditorModeToolboxTab);
 			const FSlateIcon ToolsIcon(FEditorStyle::GetStyleSetName(), "LevelEditor.Tabs.Modes");
-			LevelEditorTabManager->RegisterTabSpawner("LevelEditorToolBox", FOnSpawnTab::CreateSP<SLevelEditor, FName, FString>(this, &SLevelEditor::SpawnLevelEditorTab, FName("LevelEditorToolBox"), FString()), CanSpawnTabDelegate)
+			LevelEditorTabManager->RegisterTabSpawner(LevelEditorTabIds::LevelEditorToolBox, FOnSpawnTab::CreateSP<SLevelEditor, FName, FString>(this, &SLevelEditor::SpawnLevelEditorTab, LevelEditorTabIds::LevelEditorToolBox, FString()), CanSpawnTabDelegate)
 				.SetDisplayName(NSLOCTEXT("LevelEditorTabs", "LevelEditorModesToolboxTab", "Active Mode Toolbox"))
 				.SetTooltipText(NSLOCTEXT("LevelEditorTabs", "LevelEditorModesToolboxTabTooltipText", "Open the Modes tab, which contains the active editor mode's settings."))
 				.SetGroup(MenuStructure.GetLevelEditorModesCategory())
@@ -1169,7 +1168,7 @@ TSharedRef<SWidget> SLevelEditor::RestoreContentArea( const TSharedRef<SDockTab>
 
 		{
 			const FSlateIcon ToolsIcon(FEditorStyle::GetStyleSetName(), "LevelEditor.Tabs.PlacementBrowser");
-			LevelEditorTabManager->RegisterTabSpawner("PlacementBrowser", FOnSpawnTab::CreateSP<SLevelEditor, FName, FString>(this, &SLevelEditor::SpawnLevelEditorTab, FName("PlacementBrowser"), FString()))
+			LevelEditorTabManager->RegisterTabSpawner(LevelEditorTabIds::PlacementBrowser, FOnSpawnTab::CreateSP<SLevelEditor, FName, FString>(this, &SLevelEditor::SpawnLevelEditorTab, LevelEditorTabIds::PlacementBrowser, FString()))
 				.SetDisplayName(NSLOCTEXT("LevelEditorTabs", "PlacementBrowser", "Place Actors"))
 				.SetTooltipText(NSLOCTEXT("LevelEditorTabs", "PlacementBrowserTooltipText", "Actor Placement Browser"))
 				.SetGroup(MenuStructure.GetLevelEditorCategory())
@@ -1178,7 +1177,7 @@ TSharedRef<SWidget> SLevelEditor::RestoreContentArea( const TSharedRef<SDockTab>
 
 		{
 			const FSlateIcon OutlinerIcon(FEditorStyle::GetStyleSetName(), "LevelEditor.Tabs.Outliner");
-		    LevelEditorTabManager->RegisterTabSpawner( "LevelEditorSceneOutliner", FOnSpawnTab::CreateSP<SLevelEditor, FName, FString>(this, &SLevelEditor::SpawnLevelEditorTab, FName("LevelEditorSceneOutliner"), FString()) )
+		    LevelEditorTabManager->RegisterTabSpawner(LevelEditorTabIds::LevelEditorSceneOutliner, FOnSpawnTab::CreateSP<SLevelEditor, FName, FString>(this, &SLevelEditor::SpawnLevelEditorTab, LevelEditorTabIds::LevelEditorSceneOutliner, FString()) )
 				.SetDisplayName(NSLOCTEXT("LevelEditorTabs", "LevelEditorSceneOutliner", "World Outliner"))
 				.SetTooltipText(NSLOCTEXT("LevelEditorTabs", "LevelEditorSceneOutlinerTooltipText", "Open the World Outliner tab, which provides a searchable and filterable list of all actors in the world."))
 				.SetGroup( MenuStructure.GetLevelEditorCategory() )	
@@ -1187,7 +1186,7 @@ TSharedRef<SWidget> SLevelEditor::RestoreContentArea( const TSharedRef<SDockTab>
 
 		{
 			const FSlateIcon LayersIcon(FEditorStyle::GetStyleSetName(), "LevelEditor.Tabs.Layers");
-			LevelEditorTabManager->RegisterTabSpawner( "LevelEditorLayerBrowser", FOnSpawnTab::CreateSP<SLevelEditor, FName, FString>(this, &SLevelEditor::SpawnLevelEditorTab, FName("LevelEditorLayerBrowser"), FString()) )
+			LevelEditorTabManager->RegisterTabSpawner(LevelEditorTabIds::LevelEditorLayerBrowser, FOnSpawnTab::CreateSP<SLevelEditor, FName, FString>(this, &SLevelEditor::SpawnLevelEditorTab, LevelEditorTabIds::LevelEditorLayerBrowser, FString()) )
 				.SetDisplayName(NSLOCTEXT("LevelEditorTabs", "LevelEditorLayerBrowser", "Layers"))
 				.SetTooltipText(NSLOCTEXT("LevelEditorTabs", "LevelEditorLayerBrowserTooltipText", "Open the Layers tab. Use this to manage which actors in the world belong to which layers."))
 				.SetGroup( MenuStructure.GetLevelEditorCategory() )
@@ -1196,7 +1195,7 @@ TSharedRef<SWidget> SLevelEditor::RestoreContentArea( const TSharedRef<SDockTab>
 
 		{
 			const FSlateIcon LayersIcon(FEditorStyle::GetStyleSetName(), "LevelEditor.Tabs.HLOD");
-			LevelEditorTabManager->RegisterTabSpawner("LevelEditorHierarchicalLODOutliner", FOnSpawnTab::CreateSP<SLevelEditor, FName, FString>(this, &SLevelEditor::SpawnLevelEditorTab, FName("LevelEditorHierarchicalLODOutliner"), FString()))
+			LevelEditorTabManager->RegisterTabSpawner(LevelEditorTabIds::LevelEditorHierarchicalLODOutliner, FOnSpawnTab::CreateSP<SLevelEditor, FName, FString>(this, &SLevelEditor::SpawnLevelEditorTab, LevelEditorTabIds::LevelEditorHierarchicalLODOutliner, FString()))
 				.SetDisplayName(NSLOCTEXT("LevelEditorTabs", "LevelEditorHierarchicalLODOutliner", "Hierarchical LOD Outliner"))
 				.SetTooltipText(NSLOCTEXT("LevelEditorTabs", "LevelEditorHierarchicalLODOutlinerTooltipText", "Open the Hierarchical LOD Outliner."))
 				.SetGroup(MenuStructure.GetLevelEditorCategory())
@@ -1204,19 +1203,19 @@ TSharedRef<SWidget> SLevelEditor::RestoreContentArea( const TSharedRef<SDockTab>
 		}
 		
 		{
-			LevelEditorTabManager->RegisterTabSpawner( WorldBrowserHierarchyTab, FOnSpawnTab::CreateSP<SLevelEditor, FName, FString>(this, &SLevelEditor::SpawnLevelEditorTab, WorldBrowserHierarchyTab, FString()) )
+			LevelEditorTabManager->RegisterTabSpawner(LevelEditorTabIds::WorldBrowserHierarchy, FOnSpawnTab::CreateSP<SLevelEditor, FName, FString>(this, &SLevelEditor::SpawnLevelEditorTab, LevelEditorTabIds::WorldBrowserHierarchy, FString()) )
 				.SetDisplayName(NSLOCTEXT("LevelEditorTabs", "WorldBrowserHierarchy", "Levels"))
 				.SetTooltipText(NSLOCTEXT("LevelEditorTabs", "WorldBrowserHierarchyTooltipText", "Open the Levels tab. Use this to manage the levels in the current project."))
 				.SetGroup( WorkspaceMenu::GetMenuStructure().GetLevelEditorCategory() )
 				.SetIcon( FSlateIcon(FEditorStyle::GetStyleSetName(), "LevelEditor.Tabs.WorldBrowser") );
 			
-			LevelEditorTabManager->RegisterTabSpawner( WorldBrowserDetailsTab, FOnSpawnTab::CreateSP<SLevelEditor, FName, FString>(this, &SLevelEditor::SpawnLevelEditorTab, WorldBrowserDetailsTab, FString()) )
+			LevelEditorTabManager->RegisterTabSpawner(LevelEditorTabIds::WorldBrowserDetails, FOnSpawnTab::CreateSP<SLevelEditor, FName, FString>(this, &SLevelEditor::SpawnLevelEditorTab, LevelEditorTabIds::WorldBrowserDetails, FString()) )
 				.SetMenuType( ETabSpawnerMenuType::Hidden )
 				.SetDisplayName(NSLOCTEXT("LevelEditorTabs", "WorldBrowserDetails", "Level Details"))
 				.SetGroup( WorkspaceMenu::GetMenuStructure().GetLevelEditorCategory() )
 				.SetIcon( FSlateIcon(FEditorStyle::GetStyleSetName(), "LevelEditor.Tabs.WorldBrowserDetails") );
 		
-			LevelEditorTabManager->RegisterTabSpawner( WorldBrowserCompositionTab, FOnSpawnTab::CreateSP<SLevelEditor, FName, FString>(this, &SLevelEditor::SpawnLevelEditorTab, WorldBrowserCompositionTab, FString()) )
+			LevelEditorTabManager->RegisterTabSpawner(LevelEditorTabIds::WorldBrowserComposition, FOnSpawnTab::CreateSP<SLevelEditor, FName, FString>(this, &SLevelEditor::SpawnLevelEditorTab, LevelEditorTabIds::WorldBrowserComposition, FString()) )
 				.SetMenuType( ETabSpawnerMenuType::Hidden )
 				.SetDisplayName(NSLOCTEXT("LevelEditorTabs", "WorldBrowserComposition", "World Composition"))
 				.SetGroup( WorkspaceMenu::GetMenuStructure().GetLevelEditorCategory() )
@@ -1225,7 +1224,7 @@ TSharedRef<SWidget> SLevelEditor::RestoreContentArea( const TSharedRef<SDockTab>
 
 		{
 			const FSlateIcon StatsViewerIcon(FEditorStyle::GetStyleSetName(), "LevelEditor.Tabs.StatsViewer");
-			LevelEditorTabManager->RegisterTabSpawner(LevelEditorStatsViewerTab, FOnSpawnTab::CreateSP<SLevelEditor, FName, FString>(this, &SLevelEditor::SpawnLevelEditorTab, LevelEditorStatsViewerTab, FString()))
+			LevelEditorTabManager->RegisterTabSpawner(LevelEditorTabIds::LevelEditorStatsViewer, FOnSpawnTab::CreateSP<SLevelEditor, FName, FString>(this, &SLevelEditor::SpawnLevelEditorTab, LevelEditorTabIds::LevelEditorStatsViewer, FString()))
 				.SetDisplayName(NSLOCTEXT("LevelEditorTabs", "LevelEditorStatsViewer", "Statistics"))
 				.SetTooltipText(NSLOCTEXT("LevelEditorTabs", "LevelEditorStatsViewerTooltipText", "Open the Statistics tab, in order to see data pertaining to lighting, textures and primitives."))
 				.SetGroup(MenuStructure.GetLevelEditorCategory())
@@ -1235,7 +1234,7 @@ TSharedRef<SWidget> SLevelEditor::RestoreContentArea( const TSharedRef<SDockTab>
 		{
 			// @todo remove when world-centric mode is added
 			const FSlateIcon SequencerIcon("LevelSequenceEditorStyle", "LevelSequenceEditor.Tabs.Sequencer" );
-			LevelEditorTabManager->RegisterTabSpawner( "Sequencer", FOnSpawnTab::CreateSP<SLevelEditor, FName, FString>(this, &SLevelEditor::SpawnLevelEditorTab, FName("Sequencer"), FString()) )
+			LevelEditorTabManager->RegisterTabSpawner(LevelEditorTabIds::Sequencer, FOnSpawnTab::CreateSP<SLevelEditor, FName, FString>(this, &SLevelEditor::SpawnLevelEditorTab, LevelEditorTabIds::Sequencer, FString()) )
 				.SetDisplayName(NSLOCTEXT("LevelEditorTabs", "Sequencer", "Sequencer"))
 				.SetGroup( MenuStructure.GetLevelEditorCinematicsCategory() )
 				.SetIcon( SequencerIcon );
@@ -1244,20 +1243,20 @@ TSharedRef<SWidget> SLevelEditor::RestoreContentArea( const TSharedRef<SDockTab>
 		{
 			// @todo remove when world-centric mode is added
 			const FSlateIcon SequencerGraphIcon = FSlateIcon(FEditorStyle::GetStyleSetName(), "GenericCurveEditor.TabIcon");
-			LevelEditorTabManager->RegisterTabSpawner("SequencerGraphEditor", FOnSpawnTab::CreateSP<SLevelEditor, FName, FString>(this, &SLevelEditor::SpawnLevelEditorTab, FName("SequencerGraphEditor"), FString()))
+			LevelEditorTabManager->RegisterTabSpawner(LevelEditorTabIds::SequencerGraphEditor, FOnSpawnTab::CreateSP<SLevelEditor, FName, FString>(this, &SLevelEditor::SpawnLevelEditorTab, LevelEditorTabIds::SequencerGraphEditor, FString()))
 				.SetMenuType(ETabSpawnerMenuType::Type::Hidden);
 		}
 
 		{
 			const FSlateIcon WorldPropertiesIcon(FEditorStyle::GetStyleSetName(), "LevelEditor.WorldProperties.Tab");
-			LevelEditorTabManager->RegisterTabSpawner( "WorldSettingsTab", FOnSpawnTab::CreateSP<SLevelEditor, FName, FString>(this, &SLevelEditor::SpawnLevelEditorTab, FName("WorldSettingsTab"), FString()) )
+			LevelEditorTabManager->RegisterTabSpawner(LevelEditorTabIds::WorldSettings, FOnSpawnTab::CreateSP<SLevelEditor, FName, FString>(this, &SLevelEditor::SpawnLevelEditorTab, LevelEditorTabIds::WorldSettings, FString()) )
 				.SetDisplayName(NSLOCTEXT("LevelEditorTabs", "WorldSettings", "World Settings"))
 				.SetTooltipText(NSLOCTEXT("LevelEditorTabs", "WorldSettingsTooltipText", "Open the World Settings tab, in which global properties of the level can be viewed and edited."))
 				.SetGroup( MenuStructure.GetLevelEditorCategory() )
 				.SetIcon( WorldPropertiesIcon );
 		}
 
-		FTabSpawnerEntry& BuildAndSubmitEntry = LevelEditorTabManager->RegisterTabSpawner(LevelEditorBuildAndSubmitTab, FOnSpawnTab::CreateSP<SLevelEditor, FName, FString>(this, &SLevelEditor::SpawnLevelEditorTab, LevelEditorBuildAndSubmitTab, FString()));
+		FTabSpawnerEntry& BuildAndSubmitEntry = LevelEditorTabManager->RegisterTabSpawner(LevelEditorTabIds::LevelEditorBuildAndSubmit, FOnSpawnTab::CreateSP<SLevelEditor, FName, FString>(this, &SLevelEditor::SpawnLevelEditorTab, LevelEditorTabIds::LevelEditorBuildAndSubmit, FString()));
 		BuildAndSubmitEntry.SetAutoGenerateMenuEntry(false);
 
 		LevelEditorModule.OnRegisterTabs().Broadcast(LevelEditorTabManager);
@@ -1302,9 +1301,9 @@ TSharedRef<SWidget> SLevelEditor::RestoreContentArea( const TSharedRef<SDockTab>
 					(
 						FTabManager::NewStack()
 						->SetSizeCoefficient( 0.3f )
-						->AddTab("PlacementBrowser", ETabState::OpenedTab)
-						->AddTab("LevelEditorToolBox", ETabState::ClosedTab)
-						->SetForegroundTab(FName("PlacementBrowser"))
+						->AddTab(LevelEditorTabIds::PlacementBrowser, ETabState::OpenedTab)
+						->AddTab(LevelEditorTabIds::LevelEditorToolBox, ETabState::ClosedTab)
+						->SetForegroundTab(LevelEditorTabIds::PlacementBrowser)
 					)
 					->Split
 					(
@@ -1315,7 +1314,7 @@ TSharedRef<SWidget> SLevelEditor::RestoreContentArea( const TSharedRef<SDockTab>
 						(
 							FTabManager::NewStack()
 							->SetHideTabWell(true)
-							->AddTab("LevelEditorToolBar", ETabState::OpenedTab)
+							->AddTab(LevelEditorTabIds::LevelEditorToolBar, ETabState::OpenedTab)
 						)
 						->Split
 						(
@@ -1328,7 +1327,7 @@ TSharedRef<SWidget> SLevelEditor::RestoreContentArea( const TSharedRef<SDockTab>
 							FTabManager::NewStack()
 							->SetHideTabWell(true)
 							->SetSizeCoefficient( 1.0f )
-							->AddTab("LevelEditorViewport", ETabState::OpenedTab)
+							->AddTab(LevelEditorTabIds::LevelEditorViewport, ETabState::OpenedTab)
 						)
 					)
 				)
@@ -1337,7 +1336,7 @@ TSharedRef<SWidget> SLevelEditor::RestoreContentArea( const TSharedRef<SDockTab>
 					FTabManager::NewStack()
 					->SetSizeCoefficient(.4)
 					->AddTab("ContentBrowserTab1", ETabState::OpenedTab)
-					->AddTab("OutputLog", ETabState::ClosedTab)
+					->AddTab(LevelEditorTabIds::OutputLog, ETabState::ClosedTab)
 				)
 			)
 			->Split
@@ -1349,16 +1348,16 @@ TSharedRef<SWidget> SLevelEditor::RestoreContentArea( const TSharedRef<SDockTab>
 				(
 					FTabManager::NewStack()
 					->SetSizeCoefficient(0.4f)
-					->AddTab("LevelEditorSceneOutliner", ETabState::OpenedTab)
-					->AddTab("LevelEditorLayerBrowser", ETabState::ClosedTab)
+					->AddTab(LevelEditorTabIds::LevelEditorSceneOutliner, ETabState::OpenedTab)
+					->AddTab(LevelEditorTabIds::LevelEditorLayerBrowser, ETabState::ClosedTab)
 
 				)
 				->Split
 				(
 					FTabManager::NewStack()
-					->AddTab("LevelEditorSelectionDetails", ETabState::OpenedTab)
-					->AddTab("WorldSettingsTab", ETabState::ClosedTab)
-					->SetForegroundTab(FName("LevelEditorSelectionDetails"))
+					->AddTab(LevelEditorTabIds::LevelEditorSelectionDetails, ETabState::OpenedTab)
+					->AddTab(LevelEditorTabIds::WorldSettings, ETabState::ClosedTab)
+					->SetForegroundTab(LevelEditorTabIds::LevelEditorSelectionDetails)
 				)
 			)
 			
@@ -1416,11 +1415,9 @@ void SLevelEditor::ToggleEditorMode( FEditorModeID ModeID )
 	FLevelEditorModule& LevelEditorModule = FModuleManager::GetModuleChecked<FLevelEditorModule>( "LevelEditor" );
 	TSharedPtr<FTabManager> LevelEditorTabManager = LevelEditorModule.GetLevelEditorTabManager();
 
-	static const FTabId ToolboxTabId("LevelEditorToolBox");
-
 	if (GLevelEditorModeTools().IsModeActive(FBuiltinEditorModes::EM_Default) && !GetDefault<ULevelEditorMiscSettings>()->bEnableLegacyEditorModeUI)
 	{
-		TSharedPtr<SDockTab> ToolboxTab = LevelEditorTabManager->FindExistingLiveTab(ToolboxTabId);
+		TSharedPtr<SDockTab> ToolboxTab = LevelEditorTabManager->FindExistingLiveTab(LevelEditorTabIds::LevelEditorToolBox);
 		if (ToolboxTab.IsValid())
 		{
 			ToolboxTab->RequestCloseTab();
@@ -1428,7 +1425,7 @@ void SLevelEditor::ToggleEditorMode( FEditorModeID ModeID )
 	}
 	else
 	{
-		LevelEditorTabManager->InvokeTab(ToolboxTabId);
+		LevelEditorTabManager->InvokeTab(LevelEditorTabIds::LevelEditorToolBox);
 	}
 
 }
diff --git a/Engine/Source/Editor/LevelEditor/Public/LevelEditor.h b/Engine/Source/Editor/LevelEditor/Public/LevelEditor.h
--- a/Engine/Source/Editor/LevelEditor/Public/LevelEditor.h
+++ b/Engine/Source/Editor/LevelEditor/Public/LevelEditor.h
@@ -25,6 +25,35 @@ extern const FName LevelEditorApp;
 
 DECLARE_DELEGATE_RetVal_OneParam(bool, FAreObjectsEditable, const TArray<TWeakObjectPtr<UObject>>&);
 
+class LevelEditorTabIds
+{
+public:
+	static const FName LevelEditorViewport;
+	static const FName LevelEditorViewport_Clone1;
+	static const FName LevelEditorViewport_Clone2;
+	static const FName LevelEditorViewport_Clone3;
+	static const FName LevelEditorViewport_Clone4;
+	static const FName LevelEditorToolBar;
+	static const FName LevelEditorToolBox;
+	static const FName LevelEditorSelectionDetails;
+	static const FName LevelEditorSelectionDetails2;
+	static const FName LevelEditorSelectionDetails3;
+	static const FName LevelEditorSelectionDetails4;
+	static const FName PlacementBrowser;
+	static const FName LevelEditorBuildAndSubmit;
+	static const FName LevelEditorSceneOutliner;
+	static const FName LevelEditorStatsViewer;
+	static const FName LevelEditorLayerBrowser;
+	static const FName Sequencer;
+	static const FName SequencerGraphEditor;
+	static const FName WorldSettings;
+	static const FName WorldBrowserComposition;
+	static const FName WorldBrowserHierarchy;
+	static const FName WorldBrowserDetails;
+	static const FName LevelEditorHierarchicalLODOutliner;
+	static const FName OutputLog;
+};
+
 /**
  * Level editor module
  */
@@ -68,12 +97,6 @@ public:
 	 */
 	virtual void SummonBuildAndSubmit();
 	
-	/**
-	 * Spawns a new level browser tab
-	 * @todo This only works with the first level editor. Fix it.
-	 */
-	virtual void SummonLevelBrowser();
-
 	virtual void SummonWorldBrowserHierarchy();
 	virtual void SummonWorldBrowserDetails();
 	virtual void SummonWorldBrowserComposition();
