Moves bEnableLegacyEditorModeUI from MiscSettings to EditorStyleSettings

diff --git a/Engine/Source/Editor/EditorStyle/Private/EditorStyleClasses.cpp b/Engine/Source/Editor/EditorStyle/Private/EditorStyleClasses.cpp
--- a/Engine/Source/Editor/EditorStyle/Private/EditorStyleClasses.cpp
+++ b/Engine/Source/Editor/EditorStyle/Private/EditorStyleClasses.cpp
@@ -37,6 +37,8 @@ UEditorStyleSettings::UEditorStyleSettings( const FObjectInitializer& ObjectInit
 
 	bShowFriendlyNames = true;
 	LogTimestampMode = ELogTimes::None;
+	
+	bEnableLegacyEditorModeUI = false;
 }
 
 void UEditorStyleSettings::Init()
diff --git a/Engine/Source/Editor/EditorStyle/Public/Classes/EditorStyleSettings.h b/Engine/Source/Editor/EditorStyle/Public/Classes/EditorStyleSettings.h
--- a/Engine/Source/Editor/EditorStyle/Public/Classes/EditorStyleSettings.h
+++ b/Engine/Source/Editor/EditorStyle/Public/Classes/EditorStyleSettings.h
@@ -198,6 +198,10 @@ public:
 	UPROPERTY(EditAnywhere, config, Category=UserInterface)
 	uint32 bEnableColorizedEditorTabs : 1;
 
+	/** If enabled, the modes tab will revert to the previous mode UI */
+	UPROPERTY(EditAnywhere, config, AdvancedDisplay, Category = UserInterface, meta=(ConfigRestartRequired = true))
+	uint32 bEnableLegacyEditorModeUI : 1;
+
 public:
 
 	/**
diff --git a/Engine/Source/Editor/LevelEditor/Private/LevelEditor.cpp b/Engine/Source/Editor/LevelEditor/Private/LevelEditor.cpp
--- a/Engine/Source/Editor/LevelEditor/Private/LevelEditor.cpp
+++ b/Engine/Source/Editor/LevelEditor/Private/LevelEditor.cpp
@@ -39,7 +39,7 @@
 #include "Interfaces/IMainFrameModule.h"
 #include "Framework/Commands/GenericCommands.h"
 #include "Misc/EngineBuildSettings.h"
-#include "Settings/LevelEditorMiscSettings.h"
+#include "Classes/EditorStyleSettings.h"
 
 #define LOCTEXT_NAMESPACE "LevelEditor"
 
@@ -199,7 +199,7 @@ TSharedRef<SDockTab> FLevelEditorModule::SpawnLevelEditor( const FSpawnTabArgs&
 		SetLevelEditorInstance(LevelEditorTmp);
 		LevelEditorTmp->Initialize( LevelEditorTab, OwnerWindow.ToSharedRef() );
 
-		if (GetDefault<ULevelEditorMiscSettings>()->bEnableLegacyEditorModeUI)
+		if (GetDefault<UEditorStyleSettings>()->bEnableLegacyEditorModeUI)
 		{
 			GLevelEditorModeTools().RemoveDefaultMode(FBuiltinEditorModes::EM_Default);
 			GLevelEditorModeTools().AddDefaultMode(FBuiltinEditorModes::EM_Placement);
@@ -208,7 +208,7 @@ TSharedRef<SDockTab> FLevelEditorModule::SpawnLevelEditor( const FSpawnTabArgs&
 		GLevelEditorModeTools().DeactivateAllModes();
 		GLevelEditorModeTools().ActivateDefaultMode();
 
-		if (GetDefault<ULevelEditorMiscSettings>()->bEnableLegacyEditorModeUI)
+		if (GetDefault<UEditorStyleSettings>()->bEnableLegacyEditorModeUI)
 		{
 			// In legacy mode this toolbox should always be open
 			static const FTabId ToolboxTabId("LevelEditorToolBox");
diff --git a/Engine/Source/Editor/LevelEditor/Private/SLevelEditor.cpp b/Engine/Source/Editor/LevelEditor/Private/SLevelEditor.cpp
--- a/Engine/Source/Editor/LevelEditor/Private/SLevelEditor.cpp
+++ b/Engine/Source/Editor/LevelEditor/Private/SLevelEditor.cpp
@@ -51,7 +51,7 @@
 #include "HierarchicalLODOutlinerModule.h"
 #include "EditorViewportCommands.h"
 #include "IPlacementModeModule.h"
-#include "Settings/LevelEditorMiscSettings.h"
+#include "Classes/EditorStyleSettings.h"
 
 static const FName MainFrameModuleName("MainFrame");
 static const FName LevelEditorModuleName("LevelEditor");
@@ -682,7 +682,7 @@ TSharedRef<SDockTab> SLevelEditor::SpawnLevelEditorTab( const FSpawnTabArgs& Arg
 	}
 	else if (TabIdentifier == LevelEditorTabIds::PlacementBrowser)
 	{
-		if(!GetDefault<ULevelEditorMiscSettings>()->bEnableLegacyEditorModeUI)
+		if(!GetDefault<UEditorStyleSettings>()->bEnableLegacyEditorModeUI)
 		{
 			return
 				SNew(SDockTab)
@@ -1415,7 +1415,7 @@ void SLevelEditor::ToggleEditorMode( FEditorModeID ModeID )
 	FLevelEditorModule& LevelEditorModule = FModuleManager::GetModuleChecked<FLevelEditorModule>( "LevelEditor" );
 	TSharedPtr<FTabManager> LevelEditorTabManager = LevelEditorModule.GetLevelEditorTabManager();
 
-	if (GLevelEditorModeTools().IsModeActive(FBuiltinEditorModes::EM_Default) && !GetDefault<ULevelEditorMiscSettings>()->bEnableLegacyEditorModeUI)
+	if (GLevelEditorModeTools().IsModeActive(FBuiltinEditorModes::EM_Default) && !GetDefault<UEditorStyleSettings>()->bEnableLegacyEditorModeUI)
 	{
 		TSharedPtr<SDockTab> ToolboxTab = LevelEditorTabManager->FindExistingLiveTab(LevelEditorTabIds::LevelEditorToolBox);
 		if (ToolboxTab.IsValid())
diff --git a/Engine/Source/Editor/LevelEditor/Private/SLevelEditorToolBox.cpp b/Engine/Source/Editor/LevelEditor/Private/SLevelEditorToolBox.cpp
--- a/Engine/Source/Editor/LevelEditor/Private/SLevelEditorToolBox.cpp
+++ b/Engine/Source/Editor/LevelEditor/Private/SLevelEditorToolBox.cpp
@@ -13,7 +13,7 @@
 #include "LevelEditor.h"
 #include "LevelEditorActions.h"
 #include "LevelEditorModesActions.h"
-#include "Settings/LevelEditorMiscSettings.h"
+#include "Classes/EditorStyleSettings.h"
 #include "EdMode.h"
 #include "Widgets/Docking/SDockTab.h"
 
@@ -136,7 +136,7 @@ void SLevelEditorToolBox::UpdateModeLegacyToolBar()
 		}
 	}
 
-	if (GetDefault<ULevelEditorMiscSettings>()->bEnableLegacyEditorModeUI)
+	if (GetDefault<UEditorStyleSettings>()->bEnableLegacyEditorModeUI)
 	{
 		ModeToolBarContainer->SetContent(EditorModeTools.MakeWidget());
 	}
diff --git a/Engine/Source/Editor/PlacementMode/Private/PlacementMode.cpp b/Engine/Source/Editor/PlacementMode/Private/PlacementMode.cpp
--- a/Engine/Source/Editor/PlacementMode/Private/PlacementMode.cpp
+++ b/Engine/Source/Editor/PlacementMode/Private/PlacementMode.cpp
@@ -14,7 +14,7 @@
 
 #include "Toolkits/ToolkitManager.h"
 #include "IAssetViewport.h"
-#include "Settings/LevelEditorMiscSettings.h"
+#include "Classes/EditorStyleSettings.h"
 
 FPlacementMode::FPlacementMode()
 	: AssetsToPlace()
@@ -35,7 +35,7 @@ void FPlacementMode::Initialize()
 
 bool FPlacementMode::UsesToolkits() const
 {
-	return GetDefault<ULevelEditorMiscSettings>()->bEnableLegacyEditorModeUI;
+	return GetDefault<UEditorStyleSettings>()->bEnableLegacyEditorModeUI;
 }
 
 void FPlacementMode::Enter()
@@ -44,7 +44,7 @@ void FPlacementMode::Enter()
 	FEdMode::Enter();
 
 	
-	if(!Toolkit.IsValid() && GetDefault<ULevelEditorMiscSettings>()->bEnableLegacyEditorModeUI)
+	if(!Toolkit.IsValid() && GetDefault<UEditorStyleSettings>()->bEnableLegacyEditorModeUI)
 	{
 		Toolkit = MakeShareable( new FPlacementModeToolkit );
 		Toolkit->Init(Owner->GetToolkitHost());
diff --git a/Engine/Source/Editor/PlacementMode/Private/PlacementModeModule.cpp b/Engine/Source/Editor/PlacementMode/Private/PlacementModeModule.cpp
--- a/Engine/Source/Editor/PlacementMode/Private/PlacementModeModule.cpp
+++ b/Engine/Source/Editor/PlacementMode/Private/PlacementModeModule.cpp
@@ -49,7 +49,7 @@
 #include "Kismet2/KismetEditorUtilities.h"
 #include "ActorFactories/ActorFactoryPlanarReflection.h"
 #include "SPlacementModeTools.h"
-#include "Settings/LevelEditorMiscSettings.h"
+#include "Classes/EditorStyleSettings.h"
 
 
 TOptional<FLinearColor> GetBasicShapeColorOverride()
@@ -87,7 +87,7 @@ void FPlacementModeModule::StartupModule()
 		FBuiltinEditorModes::EM_Placement,
 		NSLOCTEXT("PlacementMode", "DisplayName", "Place"),
 		FSlateIcon(FEditorStyle::GetStyleSetName(), "LevelEditor.PlacementMode", "LevelEditor.PlacementMode.Small"),
-		GetDefault<ULevelEditorMiscSettings>()->bEnableLegacyEditorModeUI, 0);
+		GetDefault<UEditorStyleSettings>()->bEnableLegacyEditorModeUI, 0);
 
 	FAssetRegistryModule& AssetRegistryModule = FModuleManager::LoadModuleChecked<FAssetRegistryModule>(TEXT("AssetRegistry"));
 	AssetRegistryModule.Get().OnAssetRemoved().AddRaw(this, &FPlacementModeModule::OnAssetRemoved);
diff --git a/Engine/Source/Editor/UnrealEd/Classes/Settings/LevelEditorMiscSettings.h b/Engine/Source/Editor/UnrealEd/Classes/Settings/LevelEditorMiscSettings.h
--- a/Engine/Source/Editor/UnrealEd/Classes/Settings/LevelEditorMiscSettings.h
+++ b/Engine/Source/Editor/UnrealEd/Classes/Settings/LevelEditorMiscSettings.h
@@ -43,9 +43,6 @@ public:
 	UPROPERTY(EditAnywhere, config, Category=Editing, meta=( DisplayName = "Preserve Actor Scale on Replace" ))
 	uint32 bReplaceRespectsScale:1;
 
-	/** If enabled, the modes tab will revert to the previous mode UI */
-	UPROPERTY(EditAnywhere, config, AdvancedDisplay, Category = Editing, meta=(ConfigRestartRequired = true))
-	uint32 bEnableLegacyEditorModeUI : 1;
 
 	/** If enabled, the modes tab will revert to the original mesh paint */
 	UPROPERTY(EditAnywhere, config, AdvancedDisplay, Category = Editing, meta = (ConfigRestartRequired = true))
diff --git a/Engine/Source/Editor/UnrealEd/Private/EditorModeRegistry.cpp b/Engine/Source/Editor/UnrealEd/Private/EditorModeRegistry.cpp
--- a/Engine/Source/Editor/UnrealEd/Private/EditorModeRegistry.cpp
+++ b/Engine/Source/Editor/UnrealEd/Private/EditorModeRegistry.cpp
@@ -15,7 +15,7 @@
 #include "Editor/FoliageEdit/Public/FoliageEditModule.h"
 #include "Editor/VirtualTexturingEditor/Public/VirtualTexturingEditorModule.h"
 #include "Tools/UEdMode.h"
-#include "Settings/LevelEditorMiscSettings.h"
+#include "Classes/EditorStyleSettings.h"
 
 FEditorModeInfo::FEditorModeInfo()
 	: ID(NAME_None)
@@ -50,7 +50,7 @@ void FEditorModeRegistry::Initialize()
 {
 	Get();
 
-	if(!GetDefault<ULevelEditorMiscSettings>()->bEnableLegacyEditorModeUI)
+	if(!GetDefault<UEditorStyleSettings>()->bEnableLegacyEditorModeUI)
 	{
 		// Add default editor modes
 		FEditorModeRegistry::Get().RegisterMode<FEdModeDefault>(
diff --git a/Engine/Source/Editor/UnrealEd/Private/Settings/SettingsClasses.cpp b/Engine/Source/Editor/UnrealEd/Private/Settings/SettingsClasses.cpp
--- a/Engine/Source/Editor/UnrealEd/Private/Settings/SettingsClasses.cpp
+++ b/Engine/Source/Editor/UnrealEd/Private/Settings/SettingsClasses.cpp
@@ -430,7 +430,6 @@ ULevelEditorMiscSettings::ULevelEditorMiscSettings( const FObjectInitializer& Ob
 	PercentageThresholdForPrompt = 20.0f;
 	MinimumBoundsForCheckingSize = FVector(500.0f, 500.0f, 50.0f);
 	bCreateNewAudioDeviceForPlayInEditor = true;
-	bEnableLegacyEditorModeUI = false;
 	bEnableLegacyMeshPaintMode = true;
 }
 
