Raise minimum slot height from 20 to 80 for actor details panel to reduce overlap

diff --git a/Engine/Source/Editor/LevelEditor/Private/SActorDetails.cpp b/Engine/Source/Editor/LevelEditor/Private/SActorDetails.cpp
--- a/Engine/Source/Editor/LevelEditor/Private/SActorDetails.cpp
+++ b/Engine/Source/Editor/LevelEditor/Private/SActorDetails.cpp
@@ -152,6 +152,7 @@ void SActorDetails::Construct(const FArguments& InArgs, const FName TabIdentifie
 		+SVerticalBox::Slot()
 		[
 			SAssignNew(DetailsSplitter, SSplitter)
+			.MinimumSlotHeight(80.0f)
 			.Orientation(Orient_Vertical)
 			.Style(FEditorStyle::Get(), "SplitterDark")
 			.PhysicalSplitterHandleSize(2.0f)
