Fix several bugs related to sequencer and editor modes being open that cause the mode tabs to not open or to be missing content

diff --git a/Engine/Source/Editor/FoliageEdit/Private/FoliageEdMode.h b/Engine/Source/Editor/FoliageEdit/Private/FoliageEdMode.h
--- a/Engine/Source/Editor/FoliageEdit/Private/FoliageEdMode.h
+++ b/Engine/Source/Editor/FoliageEdit/Private/FoliageEdMode.h
@@ -325,6 +325,8 @@ public:
 	/** FEdMode: Called after an Undo operation */
 	virtual void PostUndo() override;
 
+	virtual bool UsesToolkits() const override { return true; }
+
 	/** Called when the current level changes */
 	void NotifyNewCurrentLevel();
 	void NotifyLevelAddedToWorld(ULevel* InLevel, UWorld* InWorld);
diff --git a/Engine/Source/Editor/LevelEditor/Private/SLevelEditor.cpp b/Engine/Source/Editor/LevelEditor/Private/SLevelEditor.cpp
--- a/Engine/Source/Editor/LevelEditor/Private/SLevelEditor.cpp
+++ b/Engine/Source/Editor/LevelEditor/Private/SLevelEditor.cpp
@@ -1449,13 +1449,12 @@ void SLevelEditor::EditorModeCommandsChanged()
 
 void SLevelEditor::OnEditorModeIdChanged(const FEditorModeID& ModeChangedID, bool bIsEnteringMode)
 {
-	FEdMode* EdMode  = GLevelEditorModeTools().GetActiveMode(ModeChangedID);
-	if (EdMode && EdMode->UsesToolkits() && EdMode->GetToolkit().IsValid())
+	if(bIsEnteringMode)
 	{
 		FLevelEditorModule& LevelEditorModule = FModuleManager::GetModuleChecked<FLevelEditorModule>("LevelEditor");
 		TSharedPtr<FTabManager> LevelEditorTabManager = LevelEditorModule.GetLevelEditorTabManager();
 
-		if ((ModeChangedID == FBuiltinEditorModes::EM_Default || !GLevelEditorModeTools().ShouldShowModeToolbox()) && !GetDefault<UEditorStyleSettings>()->bEnableLegacyEditorModeUI)
+		if (!GLevelEditorModeTools().ShouldShowModeToolbox() && !GetDefault<UEditorStyleSettings>()->bEnableLegacyEditorModeUI)
 		{
 			TSharedPtr<SDockTab> ToolboxTab = LevelEditorTabManager->FindExistingLiveTab(LevelEditorTabIds::LevelEditorToolBox);
 			if (ToolboxTab.IsValid())
diff --git a/Engine/Source/Editor/LevelEditor/Private/SLevelEditorToolBox.cpp b/Engine/Source/Editor/LevelEditor/Private/SLevelEditorToolBox.cpp
--- a/Engine/Source/Editor/LevelEditor/Private/SLevelEditorToolBox.cpp
+++ b/Engine/Source/Editor/LevelEditor/Private/SLevelEditorToolBox.cpp
@@ -147,7 +147,6 @@ void SLevelEditorToolBox::UpdateModeLegacyToolBar()
 	for(const TSharedPtr<IToolkit>& HostedToolkitIt : HostedToolkits)
 	{
 		UpdateInlineContent(HostedToolkitIt, HostedToolkitIt->GetInlineContent());
-		break;
 	}
 }
 
@@ -163,45 +162,48 @@ EVisibility SLevelEditorToolBox::GetNoToolSelectedTextVisibility() const
 
 void SLevelEditorToolBox::UpdateInlineContent(const TSharedPtr<IToolkit>& Toolkit, TSharedPtr<SWidget> InlineContent) 
 {
-
-	TabName = NSLOCTEXT("LevelEditor", "ToolsTabTitle", "Toolbox");
-	TabIcon = FEditorStyle::Get().GetBrush("LevelEditor.Tabs.Modes");
-
 	if (Toolkit.IsValid())
 	{
-		TabName = Toolkit->GetEditorModeDisplayName();
-		TabIcon = Toolkit->GetEditorModeIcon().GetSmallIcon();
-
-		TSharedPtr<FModeToolkit> ModeToolkit = StaticCastSharedPtr<FModeToolkit>(Toolkit);
-
-		ModeToolHeader->SetContent(
-			SNew(SExpandableArea)	
-			.HeaderPadding(FMargin(2.0f))
-			.Padding(FMargin(10.f))
-			.BorderImage(FEditorStyle::Get().GetBrush("DetailsView.CategoryTop"))
-			.BorderBackgroundColor(FLinearColor( .6,.6,.6, 1.0f ) )
-			.BodyBorderBackgroundColor (FLinearColor::Transparent )
-			.AreaTitleFont(FEditorStyle::Get().GetFontStyle("EditorModesPanel.CategoryFontStyle"))
-			.Visibility_Lambda( [ModeToolkit] { return ModeToolkit->GetActiveToolDisplayName().IsEmpty() && ModeToolkit->GetActiveToolMessage().IsEmpty() ? EVisibility::Collapsed : EVisibility::Visible;  })
-			.BodyContent()
-			[
-				SNew( STextBlock )
-				.Text(ModeToolkit.Get(), &FModeToolkit::GetActiveToolMessage)
-				.Font( FEditorStyle::Get().GetFontStyle("EditorModesPanel.ToolDescriptionFont") )
-				.AutoWrapText(true)
-			]
-			.HeaderContent()
-			[
-				SNew( STextBlock )
-				.Text(ModeToolkit.Get(), &FModeToolkit::GetActiveToolDisplayName)
-				.Justification(ETextJustify::Center)
-				.Font(FEditorStyle::Get().GetFontStyle("EditorModesPanel.CategoryFontStyle"))
-			]
-		);
+		if(Toolkit->GetEditorMode() || Toolkit->GetScriptableEditorMode())
+		{
+			TabName = Toolkit->GetEditorModeDisplayName();
+			TabIcon = Toolkit->GetEditorModeIcon().GetSmallIcon();
+
+			TSharedPtr<FModeToolkit> ModeToolkit = StaticCastSharedPtr<FModeToolkit>(Toolkit);
+
+			ModeToolHeader->SetContent(
+				SNew(SExpandableArea)
+				.HeaderPadding(FMargin(2.0f))
+				.Padding(FMargin(10.f))
+				.BorderImage(FEditorStyle::Get().GetBrush("DetailsView.CategoryTop"))
+				.BorderBackgroundColor(FLinearColor(.6, .6, .6, 1.0f))
+				.BodyBorderBackgroundColor(FLinearColor::Transparent)
+				.AreaTitleFont(FEditorStyle::Get().GetFontStyle("EditorModesPanel.CategoryFontStyle"))
+				.Visibility_Lambda([ModeToolkit] { return ModeToolkit->GetActiveToolDisplayName().IsEmpty() && ModeToolkit->GetActiveToolMessage().IsEmpty() ? EVisibility::Collapsed : EVisibility::Visible;  })
+				.BodyContent()
+				[
+					SNew(STextBlock)
+					.Text(ModeToolkit.Get(), &FModeToolkit::GetActiveToolMessage)
+					.Font(FEditorStyle::Get().GetFontStyle("EditorModesPanel.ToolDescriptionFont"))
+					.AutoWrapText(true)
+				]
+				.HeaderContent()
+				[
+					SNew(STextBlock)
+					.Text(ModeToolkit.Get(), &FModeToolkit::GetActiveToolDisplayName)
+					.Justification(ETextJustify::Center)
+					.Font(FEditorStyle::Get().GetFontStyle("EditorModesPanel.CategoryFontStyle"))
+				]
+			);
+		}
 
 	}
 	else
 	{
+
+		TabName = NSLOCTEXT("LevelEditor", "ToolsTabTitle", "Toolbox");
+		TabIcon = FEditorStyle::Get().GetBrush("LevelEditor.Tabs.Modes");
+
 		ModeToolHeader->SetContent(SNullWidget::NullWidget);
 	}
 
diff --git a/Engine/Source/Editor/UnrealEd/Private/EdMode.cpp b/Engine/Source/Editor/UnrealEd/Private/EdMode.cpp
--- a/Engine/Source/Editor/UnrealEd/Private/EdMode.cpp
+++ b/Engine/Source/Editor/UnrealEd/Private/EdMode.cpp
@@ -590,14 +590,10 @@ void FEdMode::Enter()
 	bPendingDeletion = false;
 
 	FEditorDelegates::EditorModeIDEnter.Broadcast( GetID() );
-	const bool bIsEnteringMode = true;
-	Owner->BroadcastEditorModeIDChanged( GetID(), bIsEnteringMode );
 }
 
 void FEdMode::Exit()
 {
-	const bool bIsEnteringMode = false;
-	Owner->BroadcastEditorModeIDChanged( GetID(), bIsEnteringMode );
 	FEditorDelegates::EditorModeIDExit.Broadcast(GetID());
 }
 
diff --git a/Engine/Source/Editor/UnrealEd/Private/EditorModeManager.cpp b/Engine/Source/Editor/UnrealEd/Private/EditorModeManager.cpp
--- a/Engine/Source/Editor/UnrealEd/Private/EditorModeManager.cpp
+++ b/Engine/Source/Editor/UnrealEd/Private/EditorModeManager.cpp
@@ -671,6 +671,9 @@ void FEditorModeTools::DeactivateModeAtIndex(int32 InIndex)
 
 	Mode->Exit();
 
+	const bool bIsEnteringMode = false;
+	BroadcastEditorModeIDChanged(Mode->GetID(), bIsEnteringMode);
+
 	// Remove the toolbar widget
 	ActiveToolBarRows.RemoveAll(
 		[&Mode](FEdModeToolbarRow& Row)
@@ -693,6 +696,9 @@ void FEditorModeTools::DeactivateScriptableModeAtIndex(int32 InIndex)
 
 	Mode->Exit();
 
+	const bool bIsEnteringMode = false;
+	BroadcastEditorModeIDChanged(Mode->GetID(), bIsEnteringMode);
+
 	// Remove the toolbar widget
 	ActiveToolBarRows.RemoveAll(
 		[&Mode](FEdModeToolbarRow& Row)
@@ -956,7 +962,7 @@ bool FEditorModeTools::ShouldShowModeToolbox() const
 	bool bHasVisibleModes = false;
 	for (const TSharedPtr<FEdMode>& Mode : ActiveModes)
 	{
-		if (Mode->GetModeInfo().bVisible)
+		if (Mode->GetModeInfo().bVisible && Mode->UsesToolkits())
 		{
 			bHasVisibleModes = true;
 			break;
@@ -967,7 +973,7 @@ bool FEditorModeTools::ShouldShowModeToolbox() const
 	{
 		for (const UEdMode* Mode : ActiveScriptableModes)
 		{
-			if (Mode->GetModeInfo().bVisible)
+			if (Mode->GetModeInfo().bVisible && Mode->UsesToolkits())
 			{
 				bHasVisibleModes = true;
 				break;
@@ -1061,6 +1067,9 @@ void FEditorModeTools::ActivateMode(FEditorModeID InID, bool bToggle)
 			// Enter the new mode
 			Mode->Enter();
 
+			const bool bIsEnteringMode = true;
+			BroadcastEditorModeIDChanged(Mode->GetID(), bIsEnteringMode);
+
 			// Ask the mode to build the toolbar.
 			TSharedPtr<FUICommandList> CommandList;
 			const TSharedPtr<FModeToolkit> Toolkit = Mode->GetToolkit();
@@ -1134,6 +1143,9 @@ void FEditorModeTools::ActivateMode(FEditorModeID InID, bool bToggle)
 			// Enter the new mode
 			ScriptableMode->Enter();
 
+			const bool bIsEnteringMode = true;
+			BroadcastEditorModeIDChanged(ScriptableMode->GetID(), bIsEnteringMode);
+
 			// Ask the mode to build the toolbar.
 			TSharedPtr<FUICommandList> CommandList;
 			const TSharedPtr<FModeToolkit> Toolkit = ScriptableMode->GetToolkit();
diff --git a/Engine/Source/Editor/UnrealEd/Private/Toolkits/AssetEditorToolkit.cpp b/Engine/Source/Editor/UnrealEd/Private/Toolkits/AssetEditorToolkit.cpp
--- a/Engine/Source/Editor/UnrealEd/Private/Toolkits/AssetEditorToolkit.cpp
+++ b/Engine/Source/Editor/UnrealEd/Private/Toolkits/AssetEditorToolkit.cpp
@@ -406,6 +406,11 @@ class FEdMode* FAssetEditorToolkit::GetEditorMode() const
 	return nullptr;
 }
 
+class UEdMode* FAssetEditorToolkit::GetScriptableEditorMode() const
+{
+	return nullptr;
+}
+
 FText FAssetEditorToolkit::GetEditorModeDisplayName() const
 {
 	return FText::GetEmpty();
diff --git a/Engine/Source/Editor/UnrealEd/Private/Tools/UEdMode.cpp b/Engine/Source/Editor/UnrealEd/Private/Tools/UEdMode.cpp
--- a/Engine/Source/Editor/UnrealEd/Private/Tools/UEdMode.cpp
+++ b/Engine/Source/Editor/UnrealEd/Private/Tools/UEdMode.cpp
@@ -801,8 +801,6 @@ void UEdMode::Enter()
 	bInvalidationPending = false;
 
 	FEditorDelegates::EditorModeIDEnter.Broadcast(GetID());
-	const bool bIsEnteringMode = true;
-	Owner->BroadcastEditorModeIDChanged(GetID(), bIsEnteringMode);
 
 }
 
@@ -849,11 +847,10 @@ void UEdMode::Exit()
 		Toolkit.Reset();
 	}
 
+
 	OnToolNotificationMessage.Clear();
 	OnToolWarningMessage.Clear();
 
-	const bool bIsEnteringMode = false;
-	Owner->BroadcastEditorModeIDChanged(GetID(), bIsEnteringMode);
 	FEditorDelegates::EditorModeIDExit.Broadcast(GetID());
 	if (ToolsContext != nullptr)
 	{
diff --git a/Engine/Source/Editor/UnrealEd/Public/Toolkits/AssetEditorToolkit.h b/Engine/Source/Editor/UnrealEd/Public/Toolkits/AssetEditorToolkit.h
--- a/Engine/Source/Editor/UnrealEd/Public/Toolkits/AssetEditorToolkit.h
+++ b/Engine/Source/Editor/UnrealEd/Public/Toolkits/AssetEditorToolkit.h
@@ -128,6 +128,7 @@ public:
 	virtual FText GetToolkitToolTipText() const override;
 	virtual FString GetWorldCentricTabPrefix() const override = 0;	// Must implement in derived class!
 	virtual class FEdMode* GetEditorMode() const override;
+	virtual class UEdMode* GetScriptableEditorMode() const override;
 	virtual FText GetEditorModeDisplayName() const override;
 	virtual FSlateIcon GetEditorModeIcon() const override;
 
diff --git a/Engine/Source/Editor/UnrealEd/Public/Toolkits/IToolkit.h b/Engine/Source/Editor/UnrealEd/Public/Toolkits/IToolkit.h
--- a/Engine/Source/Editor/UnrealEd/Public/Toolkits/IToolkit.h
+++ b/Engine/Source/Editor/UnrealEd/Public/Toolkits/IToolkit.h
@@ -7,6 +7,7 @@
 #include "Framework/Docking/TabManager.h"
 
 class FEdMode;
+class UEdMode;
 class IToolkitHost;
 class SDockableTab;
 
@@ -125,6 +126,9 @@ public:
 	/** @returns the editor mode this toolkit is used for, or null if not relevant. */
 	virtual FEdMode* GetEditorMode() const = 0;
 
+	/** @returns the scriptable editor mode this toolkit is used for, or null if not relevant. */
+	virtual UEdMode* GetScriptableEditorMode() const = 0;
+
 	/** @return the editor mode display name if there is an editor mode associated with this toolkit.  An empty FText is returned otherwise */
 	virtual FText GetEditorModeDisplayName() const = 0;
 
diff --git a/Engine/Source/Runtime/Slate/Private/Framework/Docking/TabManager.cpp b/Engine/Source/Runtime/Slate/Private/Framework/Docking/TabManager.cpp
--- a/Engine/Source/Runtime/Slate/Private/Framework/Docking/TabManager.cpp
+++ b/Engine/Source/Runtime/Slate/Private/Framework/Docking/TabManager.cpp
@@ -1653,7 +1653,7 @@ TSharedPtr<SDockTab> FTabManager::SpawnTab(const FTabId& TabId, const TSharedPtr
 			bSpawningAllowedBySpawner = Spawner->CanSpawnTab.Execute(FSpawnTabArgs(ParentWindow, TabId));
 		}
 
-		if(bSpawningAllowedBySpawner)
+		if (bSpawningAllowedBySpawner && !Spawner->SpawnedTabPtr.IsValid())
 		{
 			NewTabWidget = Spawner->OnSpawnTab.Execute(FSpawnTabArgs(ParentWindow, TabId));
 			NewTabWidget->SetLayoutIdentifier(TabId);
@@ -1662,6 +1662,11 @@ TSharedPtr<SDockTab> FTabManager::SpawnTab(const FTabId& TabId, const TSharedPtr
 
 			// The spawner tracks that last tab it spawned
 			Spawner->SpawnedTabPtr = NewTabWidget;
+		} 
+		else
+		{
+			// If we got here, somehow there is two entries spawning the same tab.  This is now allowed so just ignore it.
+			bSpawningAllowedBySpawner = false;
 		}
 	}
 
