Adds CVar to hide components section of details panel.

diff --git a/Engine/Source/Editor/LevelEditor/Private/SActorDetails.cpp b/Engine/Source/Editor/LevelEditor/Private/SActorDetails.cpp
--- a/Engine/Source/Editor/LevelEditor/Private/SActorDetails.cpp
+++ b/Engine/Source/Editor/LevelEditor/Private/SActorDetails.cpp
@@ -30,6 +30,10 @@
 #include "SourceCodeNavigation.h"
 #include "Widgets/Docking/SDockTab.h"
 
+
+IConsoleVariable* SActorDetails::ShowComponents = IConsoleManager::Get().RegisterConsoleVariable(TEXT("ShowFlag.DetailsPanelComponents"), true, TEXT("Show components in editor details panel."), ECVF_Cheat);
+
+
 class SActorDetailsUneditableComponentWarning : public SCompoundWidget
 {
 public:
@@ -257,7 +261,8 @@ void SActorDetails::SetObjects(const TArray<UObject*>& InObjects, bool bForceRef
 			}
 		}
 
-		ComponentsBox->SetVisibility(bShowingComponents ? EVisibility::Visible : EVisibility::Collapsed);
+		bool ShowComponentsInDetailsPanel = ShowComponents->GetBool();
+ 		ComponentsBox->SetVisibility((bShowingComponents & ShowComponentsInDetailsPanel) ? EVisibility::Visible : EVisibility::Collapsed);
 
 		if(DetailsView->GetHostTabManager().IsValid())
 		{
diff --git a/Engine/Source/Editor/LevelEditor/Private/SActorDetails.h b/Engine/Source/Editor/LevelEditor/Private/SActorDetails.h
--- a/Engine/Source/Editor/LevelEditor/Private/SActorDetails.h
+++ b/Engine/Source/Editor/LevelEditor/Private/SActorDetails.h
@@ -44,6 +44,8 @@ public:
 	virtual void PostUndo(bool bSuccess) override;
 	virtual void PostRedo(bool bSuccess) override;
 
+	static IConsoleVariable* ShowComponents;
+
 private:
 	AActor* GetSelectedActorInEditor() const;
 	AActor* GetActorContext() const;
