Export a leveleditor static

#rnx

diff --git a/Engine/Source/Editor/LevelEditor/Private/SActorDetails.h b/Engine/Source/Editor/LevelEditor/Private/SActorDetails.h
--- a/Engine/Source/Editor/LevelEditor/Private/SActorDetails.h
+++ b/Engine/Source/Editor/LevelEditor/Private/SActorDetails.h
@@ -44,7 +44,7 @@ public:
 	virtual void PostUndo(bool bSuccess) override;
 	virtual void PostRedo(bool bSuccess) override;
 
-	static IConsoleVariable* ShowComponents;
+	static LEVELEDITOR_API IConsoleVariable* ShowComponents;
 
 private:
 	AActor* GetSelectedActorInEditor() const;
