Non-unity  fixes

diff --git a/Engine/Source/Developer/SettingsEditor/Private/Widgets/SSettingsSectionHeader.cpp b/Engine/Source/Developer/SettingsEditor/Private/Widgets/SSettingsSectionHeader.cpp
--- a/Engine/Source/Developer/SettingsEditor/Private/Widgets/SSettingsSectionHeader.cpp
+++ b/Engine/Source/Developer/SettingsEditor/Private/Widgets/SSettingsSectionHeader.cpp
@@ -15,6 +15,8 @@
 #include "SSettingsEditorCheckoutNotice.h"
 #include "DesktopPlatformModule.h"
 #include "IDetailsView.h"
+#include "Widgets/Views/SExpanderArrow.h"
+#include "Widgets/Views/STableRow.h"
 
 #define LOCTEXT_NAMESPACE "SSettingsEditor"
 
diff --git a/Engine/Source/Editor/LevelEditor/Private/SActorDetails.cpp b/Engine/Source/Editor/LevelEditor/Private/SActorDetails.cpp
--- a/Engine/Source/Editor/LevelEditor/Private/SActorDetails.cpp
+++ b/Engine/Source/Editor/LevelEditor/Private/SActorDetails.cpp
@@ -30,6 +30,7 @@
 #include "SourceCodeNavigation.h"
 #include "Widgets/Docking/SDockTab.h"
 #include "Subsystems/PanelExtensionSubsystem.h"
+#include "DetailsViewObjectFilter.h"
 
 
 IConsoleVariable* SActorDetails::ShowComponents = IConsoleManager::Get().RegisterConsoleVariable(TEXT("ShowFlag.DetailsPanelComponents"), true, TEXT("Show components in editor details panel."), ECVF_Cheat);
diff --git a/Engine/Source/Editor/LevelEditor/Private/SActorDetails.h b/Engine/Source/Editor/LevelEditor/Private/SActorDetails.h
--- a/Engine/Source/Editor/LevelEditor/Private/SActorDetails.h
+++ b/Engine/Source/Editor/LevelEditor/Private/SActorDetails.h
@@ -19,7 +19,7 @@ class SBox;
 class SSCSEditor;
 class SSplitter;
 class UBlueprint;
-
+class FDetailsViewObjectFilter;
 
 /**
  * Wraps a details panel customized for viewing actors
