#UE4 Nonunity fix

diff --git a/Engine/Source/Editor/LevelEditor/Private/SActorDetails.cpp b/Engine/Source/Editor/LevelEditor/Private/SActorDetails.cpp
--- a/Engine/Source/Editor/LevelEditor/Private/SActorDetails.cpp
+++ b/Engine/Source/Editor/LevelEditor/Private/SActorDetails.cpp
@@ -31,7 +31,7 @@
 #include "Widgets/Docking/SDockTab.h"
 #include "Subsystems/PanelExtensionSubsystem.h"
 #include "DetailsViewObjectFilter.h"
-
+#include "IDetailRootObjectCustomization.h"
 
 IConsoleVariable* SActorDetails::ShowComponents = IConsoleManager::Get().RegisterConsoleVariable(TEXT("ShowFlag.DetailsPanelComponents"), true, TEXT("Show components in editor details panel."), ECVF_Cheat);
 
diff --git a/Engine/Source/Editor/LevelEditor/Private/SActorDetails.h b/Engine/Source/Editor/LevelEditor/Private/SActorDetails.h
--- a/Engine/Source/Editor/LevelEditor/Private/SActorDetails.h
+++ b/Engine/Source/Editor/LevelEditor/Private/SActorDetails.h
@@ -51,7 +51,7 @@ public:
 	 *
 	 * @param InFilter	The filter to use or nullptr to remove the active filter
 	 */
-	void SetActorDetailsRootCustomization(TSharedPtr<FDetailsViewObjectFilter> ActorDetailsObjectFilter, TSharedPtr<IDetailRootObjectCustomization> ActorDetailsRootCustomization);
+	void SetActorDetailsRootCustomization(TSharedPtr<FDetailsViewObjectFilter> ActorDetailsObjectFilter, TSharedPtr<class IDetailRootObjectCustomization> ActorDetailsRootCustomization);
 private:
 	AActor* GetSelectedActorInEditor() const;
 	AActor* GetActorContext() const;
