Level editor viewport's toolbar converted into tool menu
#rnx

diff --git a/Engine/Source/Developer/ToolMenus/Private/ToolMenuEntry.cpp b/Engine/Source/Developer/ToolMenus/Private/ToolMenuEntry.cpp
--- a/Engine/Source/Developer/ToolMenus/Private/ToolMenuEntry.cpp
+++ b/Engine/Source/Developer/ToolMenus/Private/ToolMenuEntry.cpp
@@ -58,10 +58,10 @@ void FToolMenuEntry::SetCommandList(const TSharedPtr<const FUICommandList>& InCo
 	CommandList = InCommandList;
 }
 
-void FToolMenuEntry::SetCommand(const TSharedPtr<const FUICommandInfo>& InCommand, FName InName, const TAttribute<FText>& InLabel, const TAttribute<FText>& InToolTip, const TAttribute<FSlateIcon>& InIcon)
+void FToolMenuEntry::SetCommand(const TSharedPtr<const FUICommandInfo>& InCommand, TOptional<FName> InName, const TAttribute<FText>& InLabel, const TAttribute<FText>& InToolTip, const TAttribute<FSlateIcon>& InIcon)
 {
 	Command = InCommand;
-	Name = InName != NAME_None ? InName : InCommand->GetCommandName();
+	Name = InName.IsSet() ? InName.GetValue() : InCommand->GetCommandName();
 	Label = InLabel.IsSet() ? InLabel : InCommand->GetLabel();
 	ToolTip = InToolTip.IsSet() ? InToolTip : InCommand->GetDescription();
 	Icon = InIcon.IsSet() ? InIcon : InCommand->GetIcon();
@@ -79,15 +79,24 @@ FToolMenuEntry FToolMenuEntry::InitMenuEntry(const FName InName, const TAttribut
 	return Entry;
 }
 
-FToolMenuEntry FToolMenuEntry::InitMenuEntry(const TSharedPtr< const FUICommandInfo >& InCommand, const TAttribute<FText>& InLabel, const TAttribute<FText>& InToolTip, const TAttribute<FSlateIcon>& InIcon, const FName InTutorialHighlightName, const FName InName)
+FToolMenuEntry FToolMenuEntry::InitMenuEntry(const TSharedPtr< const FUICommandInfo >& InCommand, const TAttribute<FText>& InLabel, const TAttribute<FText>& InToolTip, const TAttribute<FSlateIcon>& InIcon, const FName InTutorialHighlightName, const TOptional<FName> InName)
 {
-	FToolMenuEntry Entry(UToolMenus::Get()->CurrentOwner(), InName, EMultiBlockType::MenuEntry);
+	FToolMenuEntry Entry(UToolMenus::Get()->CurrentOwner(), InName.IsSet() ? InName.GetValue() : NAME_None, EMultiBlockType::MenuEntry);
 	Entry.TutorialHighlightName = InTutorialHighlightName;
 	Entry.SetCommand(InCommand, InName, InLabel, InToolTip, InIcon);
 	Entry.CommandList.Reset();
 	return Entry;
 }
 
+FToolMenuEntry FToolMenuEntry::InitMenuEntry(const FName InNameOverride, const TSharedPtr< const FUICommandInfo >& InCommand, const TAttribute<FText>& InLabel, const TAttribute<FText>& InToolTip, const TAttribute<FSlateIcon>& InIcon, const FName InTutorialHighlightName)
+{
+	FToolMenuEntry Entry(UToolMenus::Get()->CurrentOwner(), InNameOverride, EMultiBlockType::MenuEntry);
+	Entry.TutorialHighlightName = InTutorialHighlightName;
+	Entry.SetCommand(InCommand, InNameOverride, InLabel, InToolTip, InIcon);
+	Entry.CommandList.Reset();
+	return Entry;
+}
+
 FToolMenuEntry FToolMenuEntry::InitMenuEntryWithCommandList(const TSharedPtr< const FUICommandInfo >& InCommand, const TSharedPtr< const FUICommandList >& InCommandList, const TAttribute<FText>& InLabel, const TAttribute<FText>& InToolTip, const TAttribute<FSlateIcon>& InIcon, const FName InTutorialHighlightName, const FName InName)
 {
 	FToolMenuEntry Entry(UToolMenus::Get()->CurrentOwner(), InName, EMultiBlockType::MenuEntry);
diff --git a/Engine/Source/Developer/ToolMenus/Private/ToolMenuSection.cpp b/Engine/Source/Developer/ToolMenus/Private/ToolMenuSection.cpp
--- a/Engine/Source/Developer/ToolMenus/Private/ToolMenuSection.cpp
+++ b/Engine/Source/Developer/ToolMenus/Private/ToolMenuSection.cpp
@@ -86,7 +86,12 @@ FToolMenuEntry& FToolMenuSection::AddMenuEntry(const FName InName, const TAttrib
 	return AddEntry(FToolMenuEntry::InitMenuEntry(InName, InLabel, InToolTip, InIcon, InAction, InUserInterfaceActionType, InTutorialHighlightName));
 }
 
-FToolMenuEntry& FToolMenuSection::AddMenuEntry(const TSharedPtr< const FUICommandInfo >& InCommand, const TAttribute<FText>& InLabelOverride, const TAttribute<FText>& InToolTipOverride, const TAttribute<FSlateIcon>& InIconOverride, const FName InTutorialHighlightName, const FName InNameOverride)
+FToolMenuEntry& FToolMenuSection::AddMenuEntry(const TSharedPtr< const FUICommandInfo >& InCommand, const TAttribute<FText>& InLabelOverride, const TAttribute<FText>& InToolTipOverride, const TAttribute<FSlateIcon>& InIconOverride, const FName InTutorialHighlightName, const TOptional<FName> InNameOverride)
+{
+	return AddEntry(FToolMenuEntry::InitMenuEntry(InCommand, InLabelOverride, InToolTipOverride, InIconOverride, InTutorialHighlightName, InNameOverride));
+}
+
+FToolMenuEntry& FToolMenuSection::AddMenuEntry(const FName InNameOverride, const TSharedPtr< const FUICommandInfo >& InCommand, const TAttribute<FText>& InLabelOverride, const TAttribute<FText>& InToolTipOverride, const TAttribute<FSlateIcon>& InIconOverride, const FName InTutorialHighlightName)
 {
 	return AddEntry(FToolMenuEntry::InitMenuEntry(InCommand, InLabelOverride, InToolTipOverride, InIconOverride, InTutorialHighlightName, InNameOverride));
 }
diff --git a/Engine/Source/Developer/ToolMenus/Public/ToolMenuEntry.h b/Engine/Source/Developer/ToolMenus/Public/ToolMenuEntry.h
--- a/Engine/Source/Developer/ToolMenus/Public/ToolMenuEntry.h
+++ b/Engine/Source/Developer/ToolMenus/Public/ToolMenuEntry.h
@@ -85,7 +85,8 @@ struct TOOLMENUS_API FToolMenuEntry
 	FToolMenuEntry(const FToolMenuOwner InOwner, const FName InName, EMultiBlockType InType);
 
 	static FToolMenuEntry InitMenuEntry(const FName InName, const TAttribute<FText>& InLabel, const TAttribute<FText>& InToolTip, const TAttribute<FSlateIcon>& InIcon, const FToolUIActionChoice& InAction, const EUserInterfaceActionType UserInterfaceActionType = EUserInterfaceActionType::Button, const FName InTutorialHighlightName = NAME_None);
-	static FToolMenuEntry InitMenuEntry(const TSharedPtr< const FUICommandInfo >& InCommand, const TAttribute<FText>& InLabelOverride = TAttribute<FText>(), const TAttribute<FText>& InToolTipOverride = TAttribute<FText>(), const TAttribute<FSlateIcon>& InIconOverride = TAttribute<FSlateIcon>(), const FName InTutorialHighlightName = NAME_None, const FName InNameOverride = NAME_None);
+	static FToolMenuEntry InitMenuEntry(const TSharedPtr< const FUICommandInfo >& InCommand, const TAttribute<FText>& InLabelOverride = TAttribute<FText>(), const TAttribute<FText>& InToolTipOverride = TAttribute<FText>(), const TAttribute<FSlateIcon>& InIconOverride = TAttribute<FSlateIcon>(), const FName InTutorialHighlightName = NAME_None, const TOptional<FName> InNameOverride = TOptional<FName>());
+	static FToolMenuEntry InitMenuEntry(const FName InNameOverride, const TSharedPtr< const FUICommandInfo >& InCommand, const TAttribute<FText>& InLabelOverride = TAttribute<FText>(), const TAttribute<FText>& InToolTipOverride = TAttribute<FText>(), const TAttribute<FSlateIcon>& InIconOverride = TAttribute<FSlateIcon>(), const FName InTutorialHighlightName = NAME_None);
 	static FToolMenuEntry InitMenuEntryWithCommandList(const TSharedPtr<const FUICommandInfo>& InCommand, const TSharedPtr<const FUICommandList>& InCommandList, const TAttribute<FText>& InLabelOverride = TAttribute<FText>(), const TAttribute<FText>& InToolTipOverride = TAttribute<FText>(), const TAttribute<FSlateIcon>& InIconOverride = TAttribute<FSlateIcon>(), const FName InTutorialHighlightName = NAME_None, const FName InNameOverride = NAME_None);
 	static FToolMenuEntry InitMenuEntry(const FName InName, const FToolUIActionChoice& InAction, const TSharedRef<SWidget>& Widget);
 
@@ -115,7 +116,7 @@ struct TOOLMENUS_API FToolMenuEntry
 
 private:
 
-	void SetCommand(const TSharedPtr< const FUICommandInfo >& InCommand, FName InName, const TAttribute<FText>& InLabel, const TAttribute<FText>& InToolTip, const TAttribute<FSlateIcon>& InIcon);
+	void SetCommand(const TSharedPtr< const FUICommandInfo >& InCommand, TOptional<FName> InName, const TAttribute<FText>& InLabel, const TAttribute<FText>& InToolTip, const TAttribute<FSlateIcon>& InIcon);
 
 	void ResetActions();
 
diff --git a/Engine/Source/Developer/ToolMenus/Public/ToolMenuSection.h b/Engine/Source/Developer/ToolMenus/Public/ToolMenuSection.h
--- a/Engine/Source/Developer/ToolMenus/Public/ToolMenuSection.h
+++ b/Engine/Source/Developer/ToolMenus/Public/ToolMenuSection.h
@@ -35,7 +35,8 @@ public:
 	FToolMenuEntry& AddEntryObject(UToolMenuEntryScript* InObject);
 	FToolMenuEntry& AddEntry(const FToolMenuEntry& Args);
 	FToolMenuEntry& AddMenuEntry(const FName InName, const TAttribute<FText>& InLabel, const TAttribute<FText>& InToolTip, const TAttribute<FSlateIcon>& InIcon, const FToolUIActionChoice& InAction, const EUserInterfaceActionType UserInterfaceActionType = EUserInterfaceActionType::Button, const FName InTutorialHighlightName = NAME_None);
-	FToolMenuEntry& AddMenuEntry(const TSharedPtr< const FUICommandInfo >& InCommand, const TAttribute<FText>& InLabelOverride = TAttribute<FText>(), const TAttribute<FText>& InToolTipOverride = TAttribute<FText>(), const TAttribute<FSlateIcon>& InIconOverride = TAttribute<FSlateIcon>(), const FName InTutorialHighlightName = NAME_None, const FName InNameOverride = NAME_None);
+	FToolMenuEntry& AddMenuEntry(const TSharedPtr< const FUICommandInfo >& InCommand, const TAttribute<FText>& InLabelOverride = TAttribute<FText>(), const TAttribute<FText>& InToolTipOverride = TAttribute<FText>(), const TAttribute<FSlateIcon>& InIconOverride = TAttribute<FSlateIcon>(), const FName InTutorialHighlightName = NAME_None, const TOptional<FName> InNameOverride = TOptional<FName>());
+	FToolMenuEntry& AddMenuEntry(const FName InNameOverride, const TSharedPtr< const FUICommandInfo >& InCommand, const TAttribute<FText>& InLabelOverride = TAttribute<FText>(), const TAttribute<FText>& InToolTipOverride = TAttribute<FText>(), const TAttribute<FSlateIcon>& InIconOverride = TAttribute<FSlateIcon>(), const FName InTutorialHighlightName = NAME_None);
 	FToolMenuEntry& AddMenuEntryWithCommandList(const TSharedPtr<const FUICommandInfo>& InCommand, const TSharedPtr<const FUICommandList>& InCommandList, const TAttribute<FText>& InLabelOverride = TAttribute<FText>(), const TAttribute<FText>& InToolTipOverride = TAttribute<FText>(), const TAttribute<FSlateIcon>& InIconOverride = TAttribute<FSlateIcon>(), const FName InTutorialHighlightName = NAME_None, const FName InNameOverride = NAME_None);
 
 	FToolMenuEntry& AddDynamicEntry(const FName InName, const FNewToolMenuSectionDelegate& InConstruct);
diff --git a/Engine/Source/Editor/BlueprintGraph/Private/EdGraphSchema_K2.cpp b/Engine/Source/Editor/BlueprintGraph/Private/EdGraphSchema_K2.cpp
--- a/Engine/Source/Editor/BlueprintGraph/Private/EdGraphSchema_K2.cpp
+++ b/Engine/Source/Editor/BlueprintGraph/Private/EdGraphSchema_K2.cpp
@@ -1959,7 +1959,7 @@ void UEdGraphSchema_K2::GetNonExistentVariableMenu(FToolMenuSection& Section, co
 		// delete this node
 		{			
 			const FText Desc = FText::Format( LOCTEXT("DeleteNonExistentVarToolTip", "Referenced variable '{0}' does not exist, delete this node?"), Variable->GetVarNameText());
-			Section.AddMenuEntry(FGenericCommands::Get().Delete, FGenericCommands::Get().Delete->GetLabel(), Desc, FSlateIcon(), NAME_None, "DeleteNonExistentVar");
+			Section.AddMenuEntry("DeleteNonExistentVar", FGenericCommands::Get().Delete, FGenericCommands::Get().Delete->GetLabel(), Desc, FSlateIcon());
 		}
 
 		GetReplaceVariableMenu(Section, InGraphNode, OwnerBlueprint);
diff --git a/Engine/Source/Editor/LevelEditor/Private/SLevelViewportToolBar.cpp b/Engine/Source/Editor/LevelEditor/Private/SLevelViewportToolBar.cpp
--- a/Engine/Source/Editor/LevelEditor/Private/SLevelViewportToolBar.cpp
+++ b/Engine/Source/Editor/LevelEditor/Private/SLevelViewportToolBar.cpp
@@ -5,6 +5,8 @@
 #include "Framework/MultiBox/MultiBoxDefs.h"
 #include "Framework/MultiBox/MultiBoxExtender.h"
 #include "Framework/MultiBox/MultiBoxBuilder.h"
+#include "ToolMenus.h"
+#include "LevelEditorMenuContext.h"
 #include "Modules/ModuleManager.h"
 #include "Widgets/Layout/SBorder.h"
 #include "Widgets/Layout/SBox.h"
@@ -47,23 +49,32 @@ public:
 	}
 };
 
-static void FillShowMenu(class FMenuBuilder& MenuBuilder, TArray< FLevelViewportCommands::FShowMenuCommand > MenuCommands, int32 EntryOffset)
+static void FillShowMenuStatic(UToolMenu* Menu, TArray< FLevelViewportCommands::FShowMenuCommand > MenuCommands, int32 EntryOffset)
 {
+	FToolMenuSection& Section = Menu->AddSection("Section");
+
 	// Generate entries for the standard show flags
 	// Assumption: the first 'n' entries types like 'Show All' and 'Hide All' buttons, so insert a separator after them
 	for (int32 EntryIndex = 0; EntryIndex < MenuCommands.Num(); ++EntryIndex)
 	{
-		MenuBuilder.AddMenuEntry(MenuCommands[EntryIndex].ShowMenuItem, NAME_None, MenuCommands[EntryIndex].LabelOverride);
+		Section.AddMenuEntry(
+			NAME_None,
+			MenuCommands[EntryIndex].ShowMenuItem,
+			MenuCommands[EntryIndex].LabelOverride
+		);
+
 		if (EntryIndex == EntryOffset - 1)
 		{
-			MenuBuilder.AddMenuSeparator();
+			Section.AddMenuSeparator(NAME_None);
 		}
 	}
 }
 
-static void FillShowStatsSubMenus(class FMenuBuilder& MenuBuilder, TArray< FLevelViewportCommands::FShowMenuCommand > MenuCommands, TMap< FString, TArray< FLevelViewportCommands::FShowMenuCommand > > StatCatCommands)
+static void FillShowStatsSubMenus(UToolMenu* Menu, TArray< FLevelViewportCommands::FShowMenuCommand > MenuCommands, TMap< FString, TArray< FLevelViewportCommands::FShowMenuCommand > > StatCatCommands)
 {
-	FillShowMenu(MenuBuilder, MenuCommands, 1);
+	FillShowMenuStatic(Menu, MenuCommands, 1);
+
+	FToolMenuSection& Section = Menu->FindOrAddSection("Section");
 
 	// Separate out stats into two list, those with and without submenus
 	TArray< FLevelViewportCommands::FShowMenuCommand > SingleStatCommands;
@@ -94,7 +105,11 @@ static void FillShowStatsSubMenus(class FMenuBuilder& MenuBuilder, TArray< FLeve
 	for (auto StatCatIt = SingleStatCommands.CreateConstIterator(); StatCatIt; ++StatCatIt)
 	{
 		const FLevelViewportCommands::FShowMenuCommand& StatCommand = *StatCatIt;
-		MenuBuilder.AddMenuEntry(StatCommand.ShowMenuItem, NAME_None, StatCommand.LabelOverride);
+		Section.AddMenuEntry(
+			NAME_None,
+			StatCommand.ShowMenuItem,
+			StatCommand.LabelOverride
+		);
 	}
 
 	// Now add all the stats that have sub menus
@@ -107,8 +122,8 @@ static void FillShowStatsSubMenus(class FMenuBuilder& MenuBuilder, TArray< FLeve
 		Args.Add(TEXT("StatCat"), CategoryName);
 		const FText CategoryDescription = FText::Format(NSLOCTEXT("UICommands", "StatShowCatName", "Show {StatCat} stats"), Args);
 
-		MenuBuilder.AddSubMenu(CategoryName, CategoryDescription,
-			FNewMenuDelegate::CreateStatic(&FillShowMenu, StatCommands, 0));
+		Section.AddSubMenu(NAME_None, CategoryName, CategoryDescription,
+			FNewToolMenuDelegate::CreateStatic(&FillShowMenuStatic, StatCommands, 0));
 	}
 }
 
@@ -348,8 +363,10 @@ bool SLevelViewportToolBar::IsPerspectiveViewport() const
 /**
  * Called to generate the set bookmark submenu
  */
-static void OnGenerateSetBookmarkMenu( FMenuBuilder& MenuBuilder , TWeakPtr<class SLevelViewport> Viewport )
+static void OnGenerateSetBookmarkMenu(UToolMenu* Menu, TWeakPtr<class SLevelViewport> Viewport)
 {
+	FToolMenuSection& Section = Menu->AddSection("Section");
+
 	// Add a menu entry for each bookmark
 	FEditorModeTools& Tools = GLevelEditorModeTools();
 	TSharedPtr<SLevelViewport> SharedViewport = Viewport.Pin();
@@ -360,7 +377,11 @@ static void OnGenerateSetBookmarkMenu( FMenuBuilder& MenuBuilder , TWeakPtr<clas
 
 	for( int32 BookmarkIndex = 0; BookmarkIndex < NumberOfMappedBookmarks; ++BookmarkIndex )
 	{
-		MenuBuilder.AddMenuEntry( FLevelViewportCommands::Get().SetBookmarkCommands[ BookmarkIndex ], NAME_None, FBookmarkUI::GetPlainLabel(BookmarkIndex) );
+		Section.AddMenuEntry(
+			NAME_None,
+			FLevelViewportCommands::Get().SetBookmarkCommands[BookmarkIndex],
+			FBookmarkUI::GetPlainLabel(BookmarkIndex)
+		);
 	}
 
 	// Only mapped bookmarks will have predefined actions.
@@ -370,7 +391,8 @@ static void OnGenerateSetBookmarkMenu( FMenuBuilder& MenuBuilder , TWeakPtr<clas
 		FUIAction Action;
 		Action.ExecuteAction.BindSP(SharedViewport.ToSharedRef(), &SLevelViewport::OnSetBookmark, BookmarkIndex);
 
-		MenuBuilder.AddMenuEntry(
+		Section.AddMenuEntry(
+			NAME_None,
 			FBookmarkUI::GetPlainLabel(BookmarkIndex),
 			FBookmarkUI::GetSetTooltip(BookmarkIndex),
 			FBookmarkUI::GetDefaultIcon(),
@@ -381,8 +403,10 @@ static void OnGenerateSetBookmarkMenu( FMenuBuilder& MenuBuilder , TWeakPtr<clas
 /**
  * Called to generate the clear bookmark submenu
  */
-static void OnGenerateClearBookmarkMenu( FMenuBuilder& MenuBuilder , TWeakPtr<class SLevelViewport> Viewport )
+static void OnGenerateClearBookmarkMenu(UToolMenu* Menu, TWeakPtr<class SLevelViewport> Viewport)
 {
+	FToolMenuSection& Section = Menu->AddSection("Section");
+
 	// Add a menu entry for each bookmark
 	FEditorModeTools& Tools = GLevelEditorModeTools();
 	TSharedPtr<SLevelViewport> SharedViewport = Viewport.Pin();
@@ -395,7 +419,11 @@ static void OnGenerateClearBookmarkMenu( FMenuBuilder& MenuBuilder , TWeakPtr<cl
 	{
 		if ( Tools.CheckBookmark( BookmarkIndex , &ViewportClient ) )
 		{
-			MenuBuilder.AddMenuEntry( FLevelViewportCommands::Get().ClearBookmarkCommands[ BookmarkIndex ], NAME_None, FBookmarkUI::GetPlainLabel(BookmarkIndex) );
+			Section.AddMenuEntry(
+				NAME_None,
+				FLevelViewportCommands::Get().ClearBookmarkCommands[ BookmarkIndex ],
+				FBookmarkUI::GetPlainLabel(BookmarkIndex)
+			);
 		}
 	}
 
@@ -406,7 +434,8 @@ static void OnGenerateClearBookmarkMenu( FMenuBuilder& MenuBuilder , TWeakPtr<cl
 			FUIAction Action;
 			Action.ExecuteAction.BindSP(SharedViewport.ToSharedRef(), &SLevelViewport::OnClearBookmark, BookmarkIndex);
 			
-			MenuBuilder.AddMenuEntry(
+			Section.AddMenuEntry(
+				NAME_None,
 				FBookmarkUI::GetPlainLabel(BookmarkIndex),
 				FBookmarkUI::GetClearTooltip(BookmarkIndex),
 				FBookmarkUI::GetDefaultIcon(),
@@ -418,8 +447,10 @@ static void OnGenerateClearBookmarkMenu( FMenuBuilder& MenuBuilder , TWeakPtr<cl
 /**
  * Called to generate the jump to bookmark menu.
  */
-static bool GenerateJumpToBookmarkMenu( FMenuBuilder& MenuBuilder , TWeakPtr<class SLevelViewport> Viewport )
+static bool GenerateJumpToBookmarkMenu(UToolMenu* Menu, TWeakPtr<class SLevelViewport> Viewport)
 {
+	FToolMenuSection& Section = Menu->AddSection("Section");
+
 	// Add a menu entry for each bookmark
 	
 	FEditorModeTools& Tools = GLevelEditorModeTools();
@@ -436,7 +467,10 @@ static bool GenerateJumpToBookmarkMenu( FMenuBuilder& MenuBuilder , TWeakPtr<cla
 		if ( Tools.CheckBookmark( BookmarkIndex , &ViewportClient ) )
 		{
 			bFoundAnyBookmarks = true;
-			MenuBuilder.AddMenuEntry( FLevelViewportCommands::Get().JumpToBookmarkCommands[ BookmarkIndex ] );
+			Section.AddMenuEntry(
+				NAME_None,
+				FLevelViewportCommands::Get().JumpToBookmarkCommands[BookmarkIndex]				
+			);
 		}
 	}
 
@@ -449,7 +483,8 @@ static bool GenerateJumpToBookmarkMenu( FMenuBuilder& MenuBuilder , TWeakPtr<cla
 			FUIAction Action;
 			Action.ExecuteAction.BindSP(SharedViewport.ToSharedRef(), &SLevelViewport::OnJumpToBookmark, BookmarkIndex);
 			
-			MenuBuilder.AddMenuEntry(
+			Section.AddMenuEntry(
+				NAME_None,
 				FBookmarkUI::GetJumpToLabel(BookmarkIndex),
 				FBookmarkUI::GetJumpToTooltip(BookmarkIndex),
 				FBookmarkUI::GetDefaultIcon(),
@@ -463,7 +498,7 @@ static bool GenerateJumpToBookmarkMenu( FMenuBuilder& MenuBuilder , TWeakPtr<cla
 /**
  * Called to generate the bookmark submenu
  */
-static void OnGenerateBookmarkMenu( FMenuBuilder& MenuBuilder , TWeakPtr<class SLevelViewport> Viewport )
+static void OnGenerateBookmarkMenu(UToolMenu* Menu, TWeakPtr<class SLevelViewport> Viewport)
 {
 	FEditorModeTools& Tools = GLevelEditorModeTools();
 
@@ -473,39 +508,51 @@ static void OnGenerateBookmarkMenu( FMenuBuilder& MenuBuilder , TWeakPtr<class S
 	// Get the viewport client to pass down to the CheckBookmark function
 	FLevelEditorViewportClient& ViewportClient = Viewport.Pin()->GetLevelViewportClient();
 
-	MenuBuilder.BeginSection("LevelViewportActiveBoookmarks", LOCTEXT("JumpToBookmarkHeader", "Active Bookmarks") );
-
-	const bool bFoundBookmarks = GenerateJumpToBookmarkMenu( MenuBuilder, Viewport );
-
-	MenuBuilder.EndSection();
+	bool bFoundBookmarks = false;
+	{
+		FToolMenuSection& Section = Menu->AddSection("LevelViewportActiveBoookmarks", LOCTEXT("JumpToBookmarkHeader", "Active Bookmarks"));
+		bFoundBookmarks = GenerateJumpToBookmarkMenu(Menu, Viewport);
+	}
 
-	MenuBuilder.BeginSection("LevelViewportBookmarkSubmenus");
 	{
-		MenuBuilder.AddSubMenu( 
+		FToolMenuSection& Section = Menu->AddSection("LevelViewportBookmarkSubmenus");
+		Section.AddSubMenu(
+			"SetBookmark",
 			LOCTEXT("SetBookmarkSubMenu", "Set Bookmark"),
 			LOCTEXT("SetBookmarkSubMenu_ToolTip", "Set viewport bookmarks"),
-			FNewMenuDelegate::CreateStatic( &OnGenerateSetBookmarkMenu, Viewport )
+			FNewToolMenuDelegate::CreateStatic( &OnGenerateSetBookmarkMenu, Viewport )
 			);
 
 		const FLevelViewportCommands& Actions = FLevelViewportCommands::Get();
-		MenuBuilder.AddMenuEntry( Actions.CompactBookmarks );
+		Section.AddMenuEntry( Actions.CompactBookmarks );
 
 		if( bFoundBookmarks )
 		{
-			MenuBuilder.AddSubMenu( 
+			Section.AddSubMenu(
+				"ClearBookmark",
 				LOCTEXT("ClearBookmarkSubMenu", "Clear Bookmark"),
 				LOCTEXT("ClearBookmarkSubMenu_ToolTip", "Clear viewport bookmarks"),
-				FNewMenuDelegate::CreateStatic( &OnGenerateClearBookmarkMenu , Viewport )
+				FNewToolMenuDelegate::CreateStatic( &OnGenerateClearBookmarkMenu, Viewport )
 				);
 
-			MenuBuilder.AddMenuEntry( Actions.ClearAllBookmarks );
+			Section.AddMenuEntry( Actions.ClearAllBookmarks );
 		}
 	}
-	MenuBuilder.EndSection();
 }
 
 TSharedRef<SWidget> SLevelViewportToolBar::GenerateOptionsMenu() 
 {
+	static const FName MenuName("LevelEditor.LevelViewportToolBar.Options");
+	if (!UToolMenus::Get()->IsMenuRegistered(MenuName))
+	{
+		UToolMenu* Menu = UToolMenus::Get()->RegisterMenu(MenuName);
+		Menu->AddDynamicSection("DynamicSection", FNewToolMenuDelegate::CreateLambda([](UToolMenu* InMenu)
+		{
+			ULevelViewportToolBarContext* Context = InMenu->FindContext<ULevelViewportToolBarContext>();
+			Context->LevelViewportToolBarWidget.Pin()->FillOptionsMenu(InMenu);
+		}));
+	}
+
 	Viewport.Pin()->OnFloatingButtonClicked();
 
 	const FLevelViewportCommands& LevelViewportActions = FLevelViewportCommands::Get();
@@ -515,13 +562,22 @@ TSharedRef<SWidget> SLevelViewportToolBar::GenerateOptionsMenu()
 	FLevelEditorModule& LevelEditorModule = FModuleManager::GetModuleChecked<FLevelEditorModule>( TEXT("LevelEditor") );
 	TSharedPtr<FExtender> MenuExtender = LevelEditorModule.AssembleExtenders(CommandList, LevelEditorModule.GetAllLevelViewportOptionsMenuExtenders());
 
+	ULevelViewportToolBarContext* ContextObject = NewObject<ULevelViewportToolBarContext>();
+	ContextObject->LevelViewportToolBarWidget = SharedThis(this);
+
+	FToolMenuContext MenuContext(CommandList, MenuExtender, ContextObject);
+	return UToolMenus::Get()->GenerateWidget(MenuName, MenuContext);
+}
+
+void SLevelViewportToolBar::FillOptionsMenu(UToolMenu* Menu)
+{
+	const FLevelViewportCommands& LevelViewportActions = FLevelViewportCommands::Get();
 	const bool bIsPerspective = Viewport.Pin()->GetLevelViewportClient().IsPerspective();
-	const bool bInShouldCloseWindowAfterMenuSelection = true;
-	FMenuBuilder OptionsMenuBuilder( bInShouldCloseWindowAfterMenuSelection, Viewport.Pin()->GetCommandList(), MenuExtender );
+
 	{
-		OptionsMenuBuilder.BeginSection("LevelViewportViewportOptions", LOCTEXT("OptionsMenuHeader", "Viewport Options") );
 		{
-			OptionsMenuBuilder.AddMenuEntry( FEditorViewportCommands::Get().ToggleRealTime );
+			FToolMenuSection& Section = Menu->AddSection("LevelViewportViewportOptions", LOCTEXT("OptionsMenuHeader", "Viewport Options"));
+			Section.AddMenuEntry(FEditorViewportCommands::Get().ToggleRealTime);
 
 			// Add an option to disable the temporary override if there is one
 			{
@@ -531,19 +587,20 @@ TSharedRef<SWidget> SLevelViewportToolBar::GenerateOptionsMenu()
 
 				TAttribute<FText> Tooltip(this, &SLevelViewportToolBar::GetRealtimeOverrideTooltip);
 
-				OptionsMenuBuilder.AddMenuEntry(
+				Section.AddMenuEntry(
+					"DisableRealtimeOverride",
 					LOCTEXT("DisableRealtimeOverride", "Disable Realtime Override"),
 					Tooltip,
 					FSlateIcon(),
 					Action);
 
-				OptionsMenuBuilder.AddMenuSeparator();
+				Section.AddMenuSeparator("DisableRealtimeOverrideSeparator");
 					
 			}
-		
-			OptionsMenuBuilder.AddMenuEntry( FEditorViewportCommands::Get().ToggleStats );
-			OptionsMenuBuilder.AddMenuEntry( FEditorViewportCommands::Get().ToggleFPS );
-			OptionsMenuBuilder.AddMenuEntry( LevelViewportActions.ToggleViewportToolbar );
+
+			Section.AddMenuEntry(FEditorViewportCommands::Get().ToggleStats);
+			Section.AddMenuEntry(FEditorViewportCommands::Get().ToggleFPS);
+			Section.AddMenuEntry(LevelViewportActions.ToggleViewportToolbar);
 
 			FText HideAllLabel = LOCTEXT("HideAllLabel", "Hide All");
 			TArray< FLevelViewportCommands::FShowMenuCommand > HideStatsMenu;
@@ -551,103 +608,116 @@ TSharedRef<SWidget> SLevelViewportToolBar::GenerateOptionsMenu()
 			// 'Hide All' button
 			HideStatsMenu.Add(FLevelViewportCommands::FShowMenuCommand(LevelViewportActions.HideAllStats, HideAllLabel));
 
-			OptionsMenuBuilder.AddSubMenu(
+			Section.AddSubMenu(
+				"ShowStatsMenu",
 				LOCTEXT("ShowStatsMenu", "Stat"), 
 				LOCTEXT("ShowStatsMenu_ToolTip", "Show Stat commands"),
-				FNewMenuDelegate::CreateStatic(&FillShowStatsSubMenus, HideStatsMenu, LevelViewportActions.ShowStatCatCommands));
+				FNewToolMenuDelegate::CreateStatic(&FillShowStatsSubMenus, HideStatsMenu, LevelViewportActions.ShowStatCatCommands));
 
 			if( bIsPerspective )
 			{
-				OptionsMenuBuilder.AddWidget( GenerateFOVMenu(), LOCTEXT("FOVAngle", "Field of View (H)") );
-				OptionsMenuBuilder.AddWidget( GenerateFarViewPlaneMenu(), LOCTEXT("FarViewPlane", "Far View Plane") );
+				Section.AddEntry(FToolMenuEntry::InitWidget("FOVAngle", GenerateFOVMenu(), LOCTEXT("FOVAngle", "Field of View (H)")));
+				Section.AddEntry(FToolMenuEntry::InitWidget("FarViewPlane", GenerateFarViewPlaneMenu(), LOCTEXT("FarViewPlane", "Far View Plane")));
 			}
 
-			OptionsMenuBuilder.AddWidget(GenerateScreenPercentageMenu(), LOCTEXT("ScreenPercentage", "Screen Percentage"));
+			Section.AddEntry(FToolMenuEntry::InitWidget("ScreenPercentage", GenerateScreenPercentageMenu(), LOCTEXT("ScreenPercentage", "Screen Percentage")));
 		}
-		OptionsMenuBuilder.EndSection();
 
-		OptionsMenuBuilder.BeginSection("LevelViewportViewportOptions2");
 		{
+			FToolMenuSection& Section = Menu->AddSection("LevelViewportViewportOptions2");
+
 			if( bIsPerspective )
 			{
 				// Cinematic preview only applies to perspective
-				OptionsMenuBuilder.AddMenuEntry( LevelViewportActions.ToggleCinematicPreview );
+				Section.AddMenuEntry( LevelViewportActions.ToggleCinematicPreview );
 			}
 
-			OptionsMenuBuilder.AddMenuEntry( LevelViewportActions.ToggleGameView );
-			OptionsMenuBuilder.AddMenuEntry( LevelViewportActions.ToggleImmersive );
+			Section.AddMenuEntry( LevelViewportActions.ToggleGameView );
+			Section.AddMenuEntry( LevelViewportActions.ToggleImmersive );
 		}
-		OptionsMenuBuilder.EndSection();
 
-		if( bIsPerspective )
-		{
-			// Bookmarks only work in perspective viewports so only show the menu option if this toolbar is in one
 
-			OptionsMenuBuilder.BeginSection("LevelViewportBookmarks");
+		{
+			FToolMenuSection& Section = Menu->AddSection("LevelViewportBookmarks");
+			if( bIsPerspective )
 			{
-				OptionsMenuBuilder.AddSubMenu( 
-					LOCTEXT("BookmarkSubMenu", "Bookmarks"), 
-					LOCTEXT("BookmarkSubMenu_ToolTip", "Viewport location bookmarking"), 
-					FNewMenuDelegate::CreateStatic( &OnGenerateBookmarkMenu , Viewport )
-					);
+				// Bookmarks only work in perspective viewports so only show the menu option if this toolbar is in one
+
+				Section.AddSubMenu(
+					"Bookmark",
+					LOCTEXT("BookmarkSubMenu", "Bookmarks"),
+					LOCTEXT("BookmarkSubMenu_ToolTip", "Viewport location bookmarking"),
+					FNewToolMenuDelegate::CreateStatic(&OnGenerateBookmarkMenu, Viewport)
+				);
+
+				Section.AddSubMenu(
+					"Camera",
+					LOCTEXT("CameraSubMeun", "Create Camera Here"),
+					LOCTEXT("CameraSubMenu_ToolTip", "Select a camera type to create at current viewport's location"),
+					FNewToolMenuDelegate::CreateSP(this, &SLevelViewportToolBar::GenerateCameraSpawnMenu)
+				);
 			}
-			OptionsMenuBuilder.EndSection();
-		
-			OptionsMenuBuilder.AddSubMenu(
-				LOCTEXT("CameraSubMeun", "Create Camera Here"),
-				LOCTEXT("CameraSubMenu_ToolTip", "Select a camera type to create at current viewport's location"),
-				FNewMenuDelegate::CreateSP(this, &SLevelViewportToolBar::GenerateCameraSpawnMenu)
-			);
-			
+
+			Section.AddMenuEntry(LevelViewportActions.HighResScreenshot);
 		}
 
-		OptionsMenuBuilder.AddMenuEntry( LevelViewportActions.HighResScreenshot );
 
-		OptionsMenuBuilder.BeginSection("LevelViewportLayouts");
 		{
-			OptionsMenuBuilder.AddSubMenu(
+			FToolMenuSection& Section = Menu->AddSection("LevelViewportLayouts");
+			Section.AddSubMenu(
+				"Configs",
 				LOCTEXT("ConfigsSubMenu", "Layouts"),
 				FText::GetEmpty(),
-				FNewMenuDelegate::CreateSP(this, &SLevelViewportToolBar::GenerateViewportConfigsMenu));
+				FNewToolMenuDelegate::CreateSP(this, &SLevelViewportToolBar::GenerateViewportConfigsMenu));
 		}
-		OptionsMenuBuilder.EndSection();
 
-		OptionsMenuBuilder.BeginSection("LevelViewportSettings");
 		{
-			OptionsMenuBuilder.AddMenuEntry( LevelViewportActions.AdvancedSettings );
+			FToolMenuSection& Section = Menu->AddSection("LevelViewportSettings");
+			Section.AddMenuEntry(LevelViewportActions.AdvancedSettings);
 		}
-		OptionsMenuBuilder.EndSection();
 	}
-
-	return OptionsMenuBuilder.MakeWidget();
 }
 
 
 TSharedRef<SWidget> SLevelViewportToolBar::GenerateDevicePreviewMenu() const
 {
-	IDeviceProfileServicesModule& ScreenDeviceProfileUIServices = FModuleManager::LoadModuleChecked<IDeviceProfileServicesModule>(TEXT( "DeviceProfileServices"));
-	IDeviceProfileServicesUIManagerPtr UIManager = ScreenDeviceProfileUIServices.GetProfileServicesManager();
+	static const FName MenuName("LevelEditor.LevelViewportToolBar.DevicePreview");
+	if (!UToolMenus::Get()->IsMenuRegistered(MenuName))
+	{
+		UToolMenu* Menu = UToolMenus::Get()->RegisterMenu(MenuName);
+		Menu->AddDynamicSection("DynamicSection", FNewToolMenuDelegate::CreateLambda([](UToolMenu* InMenu)
+		{
+			ULevelViewportToolBarContext* Context = InMenu->FindContext<ULevelViewportToolBarContext>();
+			Context->LevelViewportToolBarWidgetConst.Pin()->FillDevicePreviewMenu(InMenu);
+		}));
+	}
+
+	ULevelViewportToolBarContext* ContextObject = NewObject<ULevelViewportToolBarContext>();
+	ContextObject->LevelViewportToolBarWidgetConst = SharedThis(this);
 
-	// Create the menu
-	const bool bInShouldCloseWindowAfterMenuSelection = true;
-	FMenuBuilder DeviceMenuBuilder( bInShouldCloseWindowAfterMenuSelection, Viewport.Pin()->GetCommandList() );
+	FToolMenuContext MenuContext(Viewport.Pin()->GetCommandList(), TSharedPtr<FExtender>(), ContextObject);
+	return UToolMenus::Get()->GenerateWidget(MenuName, MenuContext);
+}
 
-	DeviceMenuBuilder.BeginSection( "DevicePreview", LOCTEXT("DevicePreviewMenuTitle", "Device Preview") );
+void SLevelViewportToolBar::FillDevicePreviewMenu(UToolMenu* Menu) const
+{
+	IDeviceProfileServicesModule& ScreenDeviceProfileUIServices = FModuleManager::LoadModuleChecked<IDeviceProfileServicesModule>(TEXT( "DeviceProfileServices"));
+	IDeviceProfileServicesUIManagerPtr UIManager = ScreenDeviceProfileUIServices.GetProfileServicesManager();
 
 	TSharedRef<SLevelViewport> ViewportRef = Viewport.Pin().ToSharedRef();
 
 	// Default menu - clear all settings
 	{
+		FToolMenuSection& Section = Menu->AddSection("DevicePreview", LOCTEXT("DevicePreviewMenuTitle", "Device Preview"));
 		FUIAction Action( FExecuteAction::CreateSP( const_cast<SLevelViewportToolBar*>(this), &SLevelViewportToolBar::SetLevelProfile, FString( "Default" ) ),
 			FCanExecuteAction(),
 			FIsActionChecked::CreateSP( ViewportRef, &SLevelViewport::IsDeviceProfileStringSet, FString( "Default" ) ) );
-		DeviceMenuBuilder.AddMenuEntry( LOCTEXT("DevicePreviewMenuClear", "Off"), FText::GetEmpty(), FSlateIcon(), Action, NAME_None, EUserInterfaceActionType::Button );
+		Section.AddMenuEntry("DevicePreviewMenuClear", LOCTEXT("DevicePreviewMenuClear", "Off"), FText::GetEmpty(), FSlateIcon(), Action, EUserInterfaceActionType::Button);
 		}
 
-	DeviceMenuBuilder.EndSection();
-
 	// Recent Device Profiles	
-	DeviceMenuBuilder.BeginSection( "Recent", LOCTEXT("RecentMenuHeading", "Recent") );
+	{
+	FToolMenuSection& Section = Menu->AddSection("Recent", LOCTEXT("RecentMenuHeading", "Recent"));
 
 	const FString INISection = "SelectedProfile";
 	const FString INIKeyBase = "ProfileItem";
@@ -664,14 +734,14 @@ TSharedRef<SWidget> SLevelViewportToolBar::GenerateDevicePreviewMenu() const
 			FUIAction Action( FExecuteAction::CreateSP( const_cast<SLevelViewportToolBar*>(this), &SLevelViewportToolBar::SetLevelProfile, CurItem ),
 				FCanExecuteAction(),
 				FIsActionChecked::CreateSP( ViewportRef, &SLevelViewport::IsDeviceProfileStringSet, CurItem ) );
-			DeviceMenuBuilder.AddMenuEntry( FText::FromString( CurItem ), FText(), FSlateIcon(FEditorStyle::GetStyleSetName(), PlatformIcon), Action, NAME_None, EUserInterfaceActionType::Button );
+			Section.AddMenuEntry(NAME_None, FText::FromString(CurItem), FText(), FSlateIcon(FEditorStyle::GetStyleSetName(), PlatformIcon), Action, EUserInterfaceActionType::Button );
 		}
 	}
-
-	DeviceMenuBuilder.EndSection();
+	}
 
 	// Device List
-	DeviceMenuBuilder.BeginSection( "Devices", LOCTEXT("DevicesMenuHeading", "Devices") );
+	{
+	FToolMenuSection& Section = Menu->AddSection("Devices", LOCTEXT("DevicesMenuHeading", "Devices"));
 
 	const TArray<TSharedPtr<FString> > PlatformList = UIManager->GetPlatformList();
 	for ( int32 Index = 0; Index < PlatformList.Num(); Index++ )
@@ -682,23 +752,23 @@ TSharedRef<SWidget> SLevelViewportToolBar::GenerateDevicePreviewMenu() const
 		{
 			const FString PlatformNameStr = DeviceProfiles[0]->DeviceType;
 			const FName PlatformIcon =  UIManager->GetPlatformIconName( PlatformNameStr );
-			DeviceMenuBuilder.AddSubMenu(
+			Section.AddSubMenu(
+				NAME_None,
 				FText::FromString( PlatformNameStr ),
 				FText::GetEmpty(),
-				FNewMenuDelegate::CreateRaw( const_cast<SLevelViewportToolBar*>(this), &SLevelViewportToolBar::MakeDevicePreviewSubMenu, DeviceProfiles ),
+				FNewToolMenuDelegate::CreateRaw( const_cast<SLevelViewportToolBar*>(this), &SLevelViewportToolBar::MakeDevicePreviewSubMenu, DeviceProfiles ),
 				false,
 				FSlateIcon(FEditorStyle::GetStyleSetName(), PlatformIcon)
 				);
 		}
 	}
-	DeviceMenuBuilder.EndSection();
-
-	return DeviceMenuBuilder.MakeWidget();
+	}
 }
 
-void SLevelViewportToolBar::MakeDevicePreviewSubMenu(FMenuBuilder& MenuBuilder, TArray< class UDeviceProfile* > InProfiles )
+void SLevelViewportToolBar::MakeDevicePreviewSubMenu(UToolMenu* Menu, TArray< class UDeviceProfile* > InProfiles)
 {
 	TSharedRef<SLevelViewport> ViewportRef = Viewport.Pin().ToSharedRef();
+	FToolMenuSection& Section = Menu->AddSection("Section");
 
 	for ( int32 Index = 0; Index < InProfiles.Num(); Index++ )
 	{
@@ -706,7 +776,7 @@ void SLevelViewportToolBar::MakeDevicePreviewSubMenu(FMenuBuilder& MenuBuilder,
 			FCanExecuteAction(),
 			FIsActionChecked::CreateSP( ViewportRef, &SLevelViewport::IsDeviceProfileStringSet, InProfiles[ Index ]->GetName() ) );
 
-		MenuBuilder.AddMenuEntry( FText::FromString( InProfiles[ Index ]->GetName() ), FText(), FSlateIcon(), Action, NAME_None, EUserInterfaceActionType::RadioButton );
+		Section.AddMenuEntry(NAME_None, FText::FromString( InProfiles[ Index ]->GetName() ), FText(), FSlateIcon(), Action, EUserInterfaceActionType::RadioButton );
 	}
 }
 
@@ -720,7 +790,7 @@ void SLevelViewportToolBar::SetLevelProfile( FString DeviceProfileName )
 	UIManager->SetProfile( DeviceProfileName );
 }
 
-void SLevelViewportToolBar::GeneratePlacedCameraMenuEntries( FMenuBuilder& Builder, TArray<ACameraActor*> Cameras ) const
+void SLevelViewportToolBar::GeneratePlacedCameraMenuEntries(FToolMenuSection& Section, TArray<ACameraActor*> Cameras) const
 {
 	FSlateIcon CameraIcon( FEditorStyle::GetStyleSetName(), "ClassIcon.CameraComponent" );
 
@@ -736,49 +806,83 @@ void SLevelViewportToolBar::GeneratePlacedCameraMenuEntries( FMenuBuilder& Build
 			FIsActionChecked::CreateSP(Viewport.Pin().ToSharedRef(), &SLevelViewport::IsActorLocked, MakeWeakObjectPtr(GenericActor))
 			);
 
-		Builder.AddMenuEntry( ActorDisplayName, FText::Format(LOCTEXT("LookThroughCameraActor_ToolTip", "Look through and pilot {0}"), ActorDisplayName), CameraIcon, LookThroughCameraAction, NAME_None, EUserInterfaceActionType::RadioButton );
+		Section.AddMenuEntry( NAME_None, ActorDisplayName, FText::Format(LOCTEXT("LookThroughCameraActor_ToolTip", "Look through and pilot {0}"), ActorDisplayName), CameraIcon, LookThroughCameraAction, EUserInterfaceActionType::RadioButton );
 	}
 }
 
-void SLevelViewportToolBar::GenerateViewportTypeMenu( FMenuBuilder& Builder ) const
+void SLevelViewportToolBar::GeneratePlacedCameraMenuEntries(UToolMenu* Menu, TArray<ACameraActor*> Cameras) const
+{
+	FToolMenuSection& Section = Menu->AddSection("Section");
+	GeneratePlacedCameraMenuEntries(Section, Cameras);
+}
+
+void SLevelViewportToolBar::GenerateViewportTypeMenu(FToolMenuSection& Section) const
 {
 	FLevelEditorModule& LevelEditorModule = FModuleManager::GetModuleChecked<FLevelEditorModule>(TEXT("LevelEditor"));
 	LevelEditorModule.IterateViewportTypes([&](FName, const FViewportTypeDefinition& InDefinition) {
 		if (InDefinition.ActivationCommand.IsValid())
 		{
-			Builder.AddMenuEntry(InDefinition.ActivationCommand);
+			Section.AddMenuEntry(NAME_None, InDefinition.ActivationCommand);
 		}
 	});
 }
 
-void SLevelViewportToolBar::GenerateCameraSpawnMenu(FMenuBuilder & Builder) const
+void SLevelViewportToolBar::GenerateViewportTypeMenu(UToolMenu* Menu) const
+{
+	FToolMenuSection& Section = Menu->AddSection("Section");
+	GenerateViewportTypeMenu(Section);
+}
+
+void SLevelViewportToolBar::GenerateCameraSpawnMenu(UToolMenu* Menu) const
 {
+	FToolMenuSection& Section = Menu->AddSection("Section");
 	const FLevelViewportCommands& Actions = FLevelViewportCommands::Get();
 
 	for (TSharedPtr<FUICommandInfo> Camera : Actions.CreateCameras)
 	{
-		Builder.AddMenuEntry(Camera);
+		Section.AddMenuEntry(NAME_None, Camera);
 	}
 }
 
 TSharedRef<SWidget> SLevelViewportToolBar::GenerateCameraMenu() const
 {
+	static const FName MenuName("LevelEditor.LevelViewportToolBar.Camera");
+	if (!UToolMenus::Get()->IsMenuRegistered(MenuName))
+	{
+		UToolMenu* Menu = UToolMenus::Get()->RegisterMenu(MenuName);
+		Menu->AddDynamicSection("DynamicSection", FNewToolMenuDelegate::CreateLambda([](UToolMenu* InMenu)
+		{
+			ULevelViewportToolBarContext* Context = InMenu->FindContext<ULevelViewportToolBarContext>();
+			Context->LevelViewportToolBarWidgetConst.Pin()->FillCameraMenu(InMenu);
+		}));
+	}
+
 	Viewport.Pin()->OnFloatingButtonClicked();
 
-	const bool bInShouldCloseWindowAfterMenuSelection = true;
-	FMenuBuilder CameraMenuBuilder( bInShouldCloseWindowAfterMenuSelection, Viewport.Pin()->GetCommandList() );
+	ULevelViewportToolBarContext* ContextObject = NewObject<ULevelViewportToolBarContext>();
+	ContextObject->LevelViewportToolBarWidgetConst = SharedThis(this);
+
+	FToolMenuContext MenuContext(Viewport.Pin()->GetCommandList(), TSharedPtr<FExtender>(), ContextObject);
+	return UToolMenus::Get()->GenerateWidget(MenuName, MenuContext);
+}
 
+void SLevelViewportToolBar::FillCameraMenu(UToolMenu* Menu) const
+{
 	// Camera types
-	CameraMenuBuilder.AddMenuEntry( FEditorViewportCommands::Get().Perspective );
+	{
+		FToolMenuSection& Section = Menu->AddSection("CameraTypes");
+		Section.AddMenuEntry(FEditorViewportCommands::Get().Perspective);
+	}
 
-	CameraMenuBuilder.BeginSection("LevelViewportCameraType_Ortho", LOCTEXT("CameraTypeHeader_Ortho", "Orthographic") );
-		CameraMenuBuilder.AddMenuEntry(FEditorViewportCommands::Get().Top);
-		CameraMenuBuilder.AddMenuEntry(FEditorViewportCommands::Get().Bottom);
-		CameraMenuBuilder.AddMenuEntry(FEditorViewportCommands::Get().Left);
-		CameraMenuBuilder.AddMenuEntry(FEditorViewportCommands::Get().Right);
-		CameraMenuBuilder.AddMenuEntry(FEditorViewportCommands::Get().Front);
-		CameraMenuBuilder.AddMenuEntry(FEditorViewportCommands::Get().Back);
-	CameraMenuBuilder.EndSection();
+	{
+		FToolMenuSection& Section = Menu->AddSection("LevelViewportCameraType_Ortho", LOCTEXT("CameraTypeHeader_Ortho", "Orthographic"));
+		Section.AddMenuEntry(FEditorViewportCommands::Get().Top);
+		Section.AddMenuEntry(FEditorViewportCommands::Get().Bottom);
+		Section.AddMenuEntry(FEditorViewportCommands::Get().Left);
+		Section.AddMenuEntry(FEditorViewportCommands::Get().Right);
+		Section.AddMenuEntry(FEditorViewportCommands::Get().Front);
+		Section.AddMenuEntry(FEditorViewportCommands::Get().Back);
+	}
 
 	TArray<ACameraActor*> Cameras;
 
@@ -793,15 +897,13 @@ TSharedRef<SWidget> SLevelViewportToolBar::GenerateCameraMenu() const
 	const uint32 MaxCamerasInTopLevelMenu = 10;
 	if( Cameras.Num() > MaxCamerasInTopLevelMenu )
 	{
-		CameraMenuBuilder.BeginSection("CameraActors");
-			CameraMenuBuilder.AddSubMenu( CameraActorsHeading, LOCTEXT("LookThroughPlacedCameras_ToolTip", "Look through and pilot placed cameras"), FNewMenuDelegate::CreateSP(this, &SLevelViewportToolBar::GeneratePlacedCameraMenuEntries, Cameras ) );
-		CameraMenuBuilder.EndSection();
+		FToolMenuSection& Section = Menu->AddSection("CameraActors");
+		Section.AddSubMenu("CameraActors", CameraActorsHeading, LOCTEXT("LookThroughPlacedCameras_ToolTip", "Look through and pilot placed cameras"), FNewToolMenuDelegate::CreateSP(this, &SLevelViewportToolBar::GeneratePlacedCameraMenuEntries, Cameras ) );
 	}
 	else
 	{
-		CameraMenuBuilder.BeginSection("CameraActors", CameraActorsHeading );
-			GeneratePlacedCameraMenuEntries( CameraMenuBuilder, Cameras );
-		CameraMenuBuilder.EndSection();
+		FToolMenuSection& Section = Menu->AddSection("CameraActors", CameraActorsHeading);
+		GeneratePlacedCameraMenuEntries(Section, Cameras);
 	}
 
 	{
@@ -814,37 +916,33 @@ TSharedRef<SWidget> SLevelViewportToolBar::GenerateCameraMenu() const
 		const uint32 MaxViewportTypesInTopLevelMenu = 4;
 		if( NumCustomViewportTypes > MaxViewportTypesInTopLevelMenu )
 		{
-			CameraMenuBuilder.BeginSection("ViewportTypes");
-			CameraMenuBuilder.AddSubMenu( ViewportTypesHeading, FText(), FNewMenuDelegate::CreateSP(this, &SLevelViewportToolBar::GenerateViewportTypeMenu) );
-			CameraMenuBuilder.EndSection();
+			FToolMenuSection& Section = Menu->AddSection("ViewportTypes");
+			Section.AddSubMenu("ViewportTypes", ViewportTypesHeading, FText(), FNewToolMenuDelegate::CreateSP(this, &SLevelViewportToolBar::GenerateViewportTypeMenu) );
 		}
 		else
 		{
-			CameraMenuBuilder.BeginSection("ViewportTypes", ViewportTypesHeading);
-			GenerateViewportTypeMenu(CameraMenuBuilder);
-			CameraMenuBuilder.EndSection();
-
+			FToolMenuSection& Section = Menu->AddSection("ViewportTypes", ViewportTypesHeading);
+			GenerateViewportTypeMenu(Section);
 		}
-		CameraMenuBuilder.EndSection();
 	}
-
-	return CameraMenuBuilder.MakeWidget();
 }
 
-void SLevelViewportToolBar::GenerateViewportConfigsMenu(FMenuBuilder& MenuBuilder) const
+void SLevelViewportToolBar::GenerateViewportConfigsMenu(UToolMenu* Menu) const
 {
 	check (Viewport.IsValid());
 	TSharedPtr<FUICommandList> CommandList = Viewport.Pin()->GetCommandList();
 
-	MenuBuilder.BeginSection("LevelViewportOnePaneConfigs", LOCTEXT("OnePaneConfigHeader", "One Pane"));
 	{
+		FToolMenuSection& Section = Menu->AddSection("LevelViewportOnePaneConfigs", LOCTEXT("OnePaneConfigHeader", "One Pane"));
+
 		FToolBarBuilder OnePaneButton(CommandList, FMultiBoxCustomization::None);
 		OnePaneButton.SetLabelVisibility(EVisibility::Collapsed);
 		OnePaneButton.SetStyle(&FEditorStyle::Get(), "ViewportLayoutToolbar");
 
 		OnePaneButton.AddToolBarButton(FLevelViewportCommands::Get().ViewportConfig_OnePane);
 
-		MenuBuilder.AddWidget(
+		Section.AddEntry(FToolMenuEntry::InitWidget(
+			"LevelViewportOnePaneConfigs",
 			SNew(SHorizontalBox)
 			+ SHorizontalBox::Slot()
 			.AutoWidth()
@@ -857,12 +955,11 @@ void SLevelViewportToolBar::GenerateViewportConfigsMenu(FMenuBuilder& MenuBuilde
 				SNullWidget::NullWidget
 			],
 			FText::GetEmpty(), true
-			);
+			));
 	}
-	MenuBuilder.EndSection();
 
-	MenuBuilder.BeginSection("LevelViewportTwoPaneConfigs", LOCTEXT("TwoPaneConfigHeader", "Two Panes") );
 	{
+		FToolMenuSection& Section = Menu->AddSection("LevelViewportTwoPaneConfigs", LOCTEXT("TwoPaneConfigHeader", "Two Panes"));
 		FToolBarBuilder TwoPaneButtons(CommandList, FMultiBoxCustomization::None);
 		TwoPaneButtons.SetLabelVisibility(EVisibility::Collapsed);
 		TwoPaneButtons.SetStyle(&FEditorStyle::Get(), "ViewportLayoutToolbar");
@@ -870,7 +967,8 @@ void SLevelViewportToolBar::GenerateViewportConfigsMenu(FMenuBuilder& MenuBuilde
 		TwoPaneButtons.AddToolBarButton(FLevelViewportCommands::Get().ViewportConfig_TwoPanesH, NAME_None, FText());
 		TwoPaneButtons.AddToolBarButton(FLevelViewportCommands::Get().ViewportConfig_TwoPanesV, NAME_None, FText());
 
-		MenuBuilder.AddWidget(
+		Section.AddEntry(FToolMenuEntry::InitWidget(
+			"LevelViewportTwoPaneConfigs",
 			SNew(SHorizontalBox)
 			+SHorizontalBox::Slot()
 			.AutoWidth()
@@ -883,12 +981,11 @@ void SLevelViewportToolBar::GenerateViewportConfigsMenu(FMenuBuilder& MenuBuilde
 				SNullWidget::NullWidget
 			],
 			FText::GetEmpty(), true
-			);
+			));
 	}
-	MenuBuilder.EndSection();
 
-	MenuBuilder.BeginSection("LevelViewportThreePaneConfigs", LOCTEXT("ThreePaneConfigHeader", "Three Panes") );
 	{
+		FToolMenuSection& Section = Menu->AddSection("LevelViewportThreePaneConfigs", LOCTEXT("ThreePaneConfigHeader", "Three Panes"));
 		FToolBarBuilder ThreePaneButtons(CommandList, FMultiBoxCustomization::None);
 		ThreePaneButtons.SetLabelVisibility(EVisibility::Collapsed);
 		ThreePaneButtons.SetStyle(&FEditorStyle::Get(), "ViewportLayoutToolbar");
@@ -898,7 +995,8 @@ void SLevelViewportToolBar::GenerateViewportConfigsMenu(FMenuBuilder& MenuBuilde
 		ThreePaneButtons.AddToolBarButton(FLevelViewportCommands::Get().ViewportConfig_ThreePanesTop, NAME_None, FText());
 		ThreePaneButtons.AddToolBarButton(FLevelViewportCommands::Get().ViewportConfig_ThreePanesBottom, NAME_None, FText());
 
-		MenuBuilder.AddWidget(
+		Section.AddEntry(FToolMenuEntry::InitWidget(
+			"LevelViewportThreePaneConfigs",
 			SNew(SHorizontalBox)
 			+SHorizontalBox::Slot()
 			.AutoWidth()
@@ -911,12 +1009,11 @@ void SLevelViewportToolBar::GenerateViewportConfigsMenu(FMenuBuilder& MenuBuilde
 				SNullWidget::NullWidget
 			],
 			FText::GetEmpty(), true
-			);
+			));
 	}
-	MenuBuilder.EndSection();
 
-	MenuBuilder.BeginSection("LevelViewportFourPaneConfigs", LOCTEXT("FourPaneConfigHeader", "Four Panes") );
 	{
+		FToolMenuSection& Section = Menu->AddSection("LevelViewportFourPaneConfigs", LOCTEXT("FourPaneConfigHeader", "Four Panes"));
 		FToolBarBuilder FourPaneButtons(CommandList, FMultiBoxCustomization::None);
 		FourPaneButtons.SetLabelVisibility(EVisibility::Collapsed);
 		FourPaneButtons.SetStyle(&FEditorStyle::Get(), "ViewportLayoutToolbar");
@@ -927,7 +1024,8 @@ void SLevelViewportToolBar::GenerateViewportConfigsMenu(FMenuBuilder& MenuBuilde
 		FourPaneButtons.AddToolBarButton(FLevelViewportCommands::Get().ViewportConfig_FourPanesTop, NAME_None, FText());
 		FourPaneButtons.AddToolBarButton(FLevelViewportCommands::Get().ViewportConfig_FourPanesBottom, NAME_None, FText());	
 
-		MenuBuilder.AddWidget(
+		Section.AddEntry(FToolMenuEntry::InitWidget(
+			"LevelViewportFourPaneConfigs",
 			SNew(SHorizontalBox)
 			+SHorizontalBox::Slot()
 			.AutoWidth()
@@ -940,33 +1038,57 @@ void SLevelViewportToolBar::GenerateViewportConfigsMenu(FMenuBuilder& MenuBuilde
 				SNullWidget::NullWidget
 			],
 			FText::GetEmpty(), true
-			);
+			));
 	}
-	MenuBuilder.EndSection();
 }
 
 TSharedRef<SWidget> SLevelViewportToolBar::GenerateShowMenu() const
 {
+	static const FName MenuName("LevelEditor.LevelViewportToolbar.Show");
+	if (!UToolMenus::Get()->IsMenuRegistered(MenuName))
+	{
+		UToolMenu* Menu = UToolMenus::Get()->RegisterMenu(MenuName);
+		Menu->AddDynamicSection("DynamicSection", FNewToolMenuDelegate::CreateLambda([](UToolMenu* InMenu)
+		{
+			ULevelViewportToolBarContext* Context = InMenu->FindContext<ULevelViewportToolBarContext>();
+			Context->LevelViewportToolBarWidgetConst.Pin()->FillShowMenu(InMenu);
+		}));
+	}
+
 	Viewport.Pin()->OnFloatingButtonClicked();
 
-	const FLevelViewportCommands& Actions = FLevelViewportCommands::Get();
 	TSharedRef<FUICommandList> CommandList = Viewport.Pin()->GetCommandList().ToSharedRef();
 
 	// Get all menu extenders for this context menu from the level editor module
 	FLevelEditorModule& LevelEditorModule = FModuleManager::GetModuleChecked<FLevelEditorModule>(TEXT("LevelEditor"));
 	TSharedPtr<FExtender> MenuExtender = LevelEditorModule.AssembleExtenders(CommandList, LevelEditorModule.GetAllLevelViewportShowMenuExtenders());
 
-	const bool bInShouldCloseWindowAfterMenuSelection = true;
-	FMenuBuilder ShowMenuBuilder( bInShouldCloseWindowAfterMenuSelection, Viewport.Pin()->GetCommandList(), MenuExtender);
+	ULevelViewportToolBarContext* ContextObject = NewObject<ULevelViewportToolBarContext>();
+	ContextObject->LevelViewportToolBarWidgetConst = SharedThis(this);
+
+	FToolMenuContext MenuContext(CommandList, MenuExtender, ContextObject);
+	return UToolMenus::Get()->GenerateWidget(MenuName, MenuContext);
+}
+
+void SLevelViewportToolBar::FillShowMenu(UToolMenu* Menu) const
+{
+	const FLevelViewportCommands& Actions = FLevelViewportCommands::Get();
 	{
-		ShowMenuBuilder.AddMenuEntry(Actions.UseDefaultShowFlags);
-		FShowFlagMenuCommands::Get().BuildShowFlagsMenu(ShowMenuBuilder);
+		{
+			FToolMenuSection& Section = Menu->AddSection("UseDefaultShowFlags");
+			Section.AddMenuEntry(Actions.UseDefaultShowFlags);
+		}
+
+		Menu->AddDynamicSection("ShowFlagsMenu", FNewToolMenuDelegateLegacy::CreateLambda([](FMenuBuilder& InMenuBuilder, UToolMenu* InMenu)
+		{
+			FShowFlagMenuCommands::Get().BuildShowFlagsMenu(InMenuBuilder);
+		}));
 
 		FText ShowAllLabel = LOCTEXT("ShowAllLabel", "Show All");
 		FText HideAllLabel = LOCTEXT("HideAllLabel", "Hide All");
 
-		ShowMenuBuilder.BeginSection("LevelViewportEditorShow", LOCTEXT("EditorShowHeader", "Editor"));
 		{
+			FToolMenuSection& Section = Menu->AddSection("LevelViewportEditorShow", LOCTEXT("EditorShowHeader", "Editor"));
 			// Show Volumes sub-menu
 			{
 				TArray< FLevelViewportCommands::FShowMenuCommand > ShowVolumesMenu;
@@ -978,14 +1100,14 @@ TSharedRef<SWidget> SLevelViewportToolBar::GenerateShowMenu() const
 				// Get each show flag command and put them in their corresponding groups
 				ShowVolumesMenu += Actions.ShowVolumeCommands;
 
-				ShowMenuBuilder.AddSubMenu(LOCTEXT("ShowVolumesMenu", "Volumes"), LOCTEXT("ShowVolumesMenu_ToolTip", "Show volumes flags"),
-					FNewMenuDelegate::CreateStatic(&FillShowMenu, ShowVolumesMenu, 2));
+				Section.AddSubMenu("ShowVolumes", LOCTEXT("ShowVolumesMenu", "Volumes"), LOCTEXT("ShowVolumesMenu_ToolTip", "Show volumes flags"),
+					FNewToolMenuDelegate::CreateStatic(&FillShowMenuStatic, ShowVolumesMenu, 2));
 			}
 
 			// Show Layers sub-menu is dynamically generated when the user enters 'show' menu
 			{
-				ShowMenuBuilder.AddSubMenu(LOCTEXT("ShowLayersMenu", "Layers"), LOCTEXT("ShowLayersMenu_ToolTip", "Show layers flags"),
-					FNewMenuDelegate::CreateStatic(&SLevelViewportToolBar::FillShowLayersMenu, Viewport));
+				Section.AddSubMenu("ShowLayers", LOCTEXT("ShowLayersMenu", "Layers"), LOCTEXT("ShowLayersMenu_ToolTip", "Show layers flags"),
+					FNewToolMenuDelegate::CreateStatic(&SLevelViewportToolBar::FillShowLayersMenu, Viewport));
 			}
 
 			// Show Sprites sub-menu
@@ -999,21 +1121,18 @@ TSharedRef<SWidget> SLevelViewportToolBar::GenerateShowMenu() const
 				// Get each show flag command and put them in their corresponding groups
 				ShowSpritesMenu += Actions.ShowSpriteCommands;
 
-				ShowMenuBuilder.AddSubMenu(LOCTEXT("ShowSpritesMenu", "Sprites"), LOCTEXT("ShowSpritesMenu_ToolTip", "Show sprites flags"),
-					FNewMenuDelegate::CreateStatic(&FillShowMenu, ShowSpritesMenu, 2));
+				Section.AddSubMenu("ShowSprites", LOCTEXT("ShowSpritesMenu", "Sprites"), LOCTEXT("ShowSpritesMenu_ToolTip", "Show sprites flags"),
+					FNewToolMenuDelegate::CreateStatic(&FillShowMenuStatic, ShowSpritesMenu, 2));
 			}
 
 			// Show 'Foliage types' sub-menu is dynamically generated when the user enters 'show' menu
 			{
-				ShowMenuBuilder.AddSubMenu(LOCTEXT("ShowFoliageTypesMenu", "Foliage Types"), LOCTEXT("ShowFoliageTypesMenu_ToolTip", "Show/hide specific foliage types"),
-					FNewMenuDelegate::CreateStatic(&SLevelViewportToolBar::FillShowFoliageTypesMenu, Viewport));
+				Section.AddSubMenu("ShowFoliageTypes", LOCTEXT("ShowFoliageTypesMenu", "Foliage Types"), LOCTEXT("ShowFoliageTypesMenu_ToolTip", "Show/hide specific foliage types"),
+					FNewToolMenuDelegate::CreateStatic(&SLevelViewportToolBar::FillShowFoliageTypesMenu, Viewport));
 			}
 
 		}
-		ShowMenuBuilder.EndSection();
 	}
-
-	return ShowMenuBuilder.MakeWidget();
 }
 
 EVisibility SLevelViewportToolBar::GetViewModeOptionsVisibility() const
@@ -1168,20 +1287,19 @@ void SLevelViewportToolBar::OnFarViewPlaneValueChanged( float NewValue )
 	Viewport.Pin()->GetLevelViewportClient().OverrideFarClipPlane(NewValue);
 }
 
-void SLevelViewportToolBar::FillShowLayersMenu( FMenuBuilder& MenuBuilder, TWeakPtr<class SLevelViewport> Viewport )
+void SLevelViewportToolBar::FillShowLayersMenu(UToolMenu* Menu, TWeakPtr<class SLevelViewport> Viewport)
 {	
-	MenuBuilder.BeginSection("LevelViewportLayers");
 	{
-		MenuBuilder.AddMenuEntry( FLevelViewportCommands::Get().ShowAllLayers, NAME_None, LOCTEXT("ShowAllLabel", "Show All"));
-		MenuBuilder.AddMenuEntry( FLevelViewportCommands::Get().HideAllLayers, NAME_None, LOCTEXT("HideAllLabel", "Hide All") );
+		FToolMenuSection& Section = Menu->AddSection("LevelViewportLayers");
+		Section.AddMenuEntry(FLevelViewportCommands::Get().ShowAllLayers, LOCTEXT("ShowAllLabel", "Show All"));
+		Section.AddMenuEntry(FLevelViewportCommands::Get().HideAllLayers, LOCTEXT("HideAllLabel", "Hide All"));
 	}
-	MenuBuilder.EndSection();
 
 	if( Viewport.IsValid() )
 	{
 		TSharedRef<SLevelViewport> ViewportRef = Viewport.Pin().ToSharedRef();
-		MenuBuilder.BeginSection("LevelViewportLayers2");
 		{
+			FToolMenuSection& Section = Menu->AddSection("LevelViewportLayers2");
 			// Get all the layers and create an entry for each of them
 			TArray<FName> AllLayerNames;
 			ULayersSubsystem* Layers = GEditor->GetEditorSubsystem<ULayersSubsystem>();
@@ -1196,10 +1314,9 @@ void SLevelViewportToolBar::FillShowLayersMenu( FMenuBuilder& MenuBuilder, TWeak
 						FCanExecuteAction(),
 						FIsActionChecked::CreateSP( ViewportRef, &SLevelViewport::IsLayerVisible, LayerName ) );
 
-				MenuBuilder.AddMenuEntry( FText::FromName( LayerName ), FText::GetEmpty(), FSlateIcon(), Action, NAME_None, EUserInterfaceActionType::ToggleButton );
+				Section.AddMenuEntry(NAME_None, FText::FromName( LayerName ), FText::GetEmpty(), FSlateIcon(), Action, EUserInterfaceActionType::ToggleButton);
 			}
 		}
-		MenuBuilder.EndSection();
 	}
 }
 
@@ -1224,7 +1341,7 @@ static TMap<FName, TArray<UFoliageType*>> GroupFoliageByOuter(const TArray<UFoli
 	return Result;
 }
 
-void SLevelViewportToolBar::FillShowFoliageTypesMenu(class FMenuBuilder& MenuBuilder, TWeakPtr<class SLevelViewport> Viewport)
+void SLevelViewportToolBar::FillShowFoliageTypesMenu(UToolMenu* Menu, TWeakPtr<class SLevelViewport> Viewport)
 {
 	auto ViewportPtr = Viewport.Pin();
 	if (!ViewportPtr.IsValid())
@@ -1232,16 +1349,15 @@ void SLevelViewportToolBar::FillShowFoliageTypesMenu(class FMenuBuilder& MenuBui
 		return;
 	}
 		
-	MenuBuilder.BeginSection("LevelViewportFoliageMeshes");
 	{
+		FToolMenuSection& Section = Menu->AddSection("LevelViewportFoliageMeshes");
 		// Map 'Show All' and 'Hide All' commands
 		FUIAction ShowAllFoliage(FExecuteAction::CreateSP(ViewportPtr.ToSharedRef(), &SLevelViewport::ToggleAllFoliageTypes, true));
 		FUIAction HideAllFoliage(FExecuteAction::CreateSP(ViewportPtr.ToSharedRef(), &SLevelViewport::ToggleAllFoliageTypes, false));
 		
-		MenuBuilder.AddMenuEntry(LOCTEXT("ShowAllLabel", "Show All"), FText::GetEmpty(), FSlateIcon(), ShowAllFoliage);
-		MenuBuilder.AddMenuEntry(LOCTEXT("HideAllLabel", "Hide All"), FText::GetEmpty(), FSlateIcon(), HideAllFoliage);
+		Section.AddMenuEntry("ShowAll", LOCTEXT("ShowAllLabel", "Show All"), FText::GetEmpty(), FSlateIcon(), ShowAllFoliage);
+		Section.AddMenuEntry("HideAll", LOCTEXT("HideAllLabel", "Hide All"), FText::GetEmpty(), FSlateIcon(), HideAllFoliage);
 	}
-	MenuBuilder.EndSection();
 	
 	// Gather all foliage types used in this world and group them by sub-levels
 	auto AllFoliageMap = GroupFoliageByOuter(GEditor->GetFoliageTypesInWorld(ViewportPtr->GetWorld()));
@@ -1250,7 +1366,7 @@ void SLevelViewportToolBar::FillShowFoliageTypesMenu(class FMenuBuilder& MenuBui
 	{
 		// Name foliage group by an outer sub-level name, or empty if foliage type is an asset
 		FText EntryName = (FoliagePair.Key == NAME_None ? FText::GetEmpty() : FText::FromName(FPackageName::GetShortFName(FoliagePair.Key)));
-		MenuBuilder.BeginSection(NAME_None, EntryName);
+		FToolMenuSection& Section = Menu->AddSection(NAME_None, EntryName);
 
 		TArray<UFoliageType*>& FoliageList = FoliagePair.Value;
 		for (UFoliageType* FoliageType : FoliageList)
@@ -1263,10 +1379,8 @@ void SLevelViewportToolBar::FillShowFoliageTypesMenu(class FMenuBuilder& MenuBui
 				FCanExecuteAction(),
 				FIsActionChecked::CreateSP(ViewportPtr.ToSharedRef(), &SLevelViewport::IsFoliageTypeVisible, FoliageTypePtr));
 
-			MenuBuilder.AddMenuEntry(FText::FromName(MeshName), FText::GetEmpty(), FSlateIcon(), Action, NAME_None, EUserInterfaceActionType::ToggleButton);
+			Section.AddMenuEntry(NAME_None, FText::FromName(MeshName), FText::GetEmpty(), FSlateIcon(), Action, EUserInterfaceActionType::ToggleButton);
 		}
-
-		MenuBuilder.EndSection();
 	}
 }
 
diff --git a/Engine/Source/Editor/LevelEditor/Private/SLevelViewportToolBar.h b/Engine/Source/Editor/LevelEditor/Private/SLevelViewportToolBar.h
--- a/Engine/Source/Editor/LevelEditor/Private/SLevelViewportToolBar.h
+++ b/Engine/Source/Editor/LevelEditor/Private/SLevelViewportToolBar.h
@@ -12,6 +12,8 @@
 class ACameraActor;
 class FExtender;
 class FMenuBuilder;
+class UToolMenu;
+struct FToolMenuSection;
 
 /**
  * A level viewport toolbar widget that is placed in a viewport
@@ -77,20 +79,23 @@ private:
 	/** @return Returns true if this viewport is the perspective viewport */
 	bool IsPerspectiveViewport() const;
 
-		/**
+	/**
 	 * Generates the toolbar device profile simulation menu content .
 	 *
 	 * @return The widget containing the options menu content.
 	 */
 	TSharedRef<SWidget> GenerateDevicePreviewMenu() const;
 
+	/** Fills device preview menu */
+	void FillDevicePreviewMenu(UToolMenu* Menu) const;
+
 	/**
 	 * Generates the sub  menu for different device profile previews.
 	 *
-	 * @param MenuBuilder - the parent menu.
+	 * @param Menu - The menu.
 	 * @param InDeviceProfiles - The array of device profiles.
 	 */
-	void MakeDevicePreviewSubMenu( FMenuBuilder& MenuBuilder, TArray< class UDeviceProfile* > Profiles );
+	void MakeDevicePreviewSubMenu(UToolMenu* Menu, TArray< class UDeviceProfile* > Profiles);
 
 	/**
 	 * Set the level profile, and save the selection to an .ini file.
@@ -106,6 +111,9 @@ private:
 	 */
 	TSharedRef<SWidget> GenerateOptionsMenu();
 
+	/** Fills options menu */
+	void FillOptionsMenu(UToolMenu* Menu);
+
 	/**
 	 * Generates the toolbar camera menu content 
 	 *
@@ -113,27 +121,45 @@ private:
 	 */
 	TSharedRef<SWidget> GenerateCameraMenu() const;
 
+	/** Fills camera menu */
+	void FillCameraMenu(UToolMenu* Menu) const;
+
+	/**
+	 * Generates menu entries for placed cameras (e.g CameraActors
+	 *
+	 * @param Menu	The menu to add menu entries to
+	 * @param Cameras	The list of cameras to add
+	 */
+	void GeneratePlacedCameraMenuEntries(UToolMenu* Menu, TArray<ACameraActor*> Cameras) const;
+
 	/**
 	 * Generates menu entries for placed cameras (e.g CameraActors
 	 *
-	 * @param Builder	The menu builder to add menu entries to
+	 * @param Section	The menu section to add menu entries to
 	 * @param Cameras	The list of cameras to add
 	 */
-	void GeneratePlacedCameraMenuEntries( FMenuBuilder& Builder, TArray<ACameraActor*> Cameras ) const;
+	void GeneratePlacedCameraMenuEntries(FToolMenuSection& Section, TArray<ACameraActor*> Cameras) const;
+
+	/**
+	 * Generates menu entries for changing the type of the viewport
+	 *
+	 * @param Menu	The menu to add menu entries to
+	 */
+	void GenerateViewportTypeMenu(UToolMenu* Menu) const;
 
 	/**
 	 * Generates menu entries for changing the type of the viewport
 	 *
-	 * @param Builder	The menu builder to add menu entries to
+	 * @param Section	The menu section to add menu entries to
 	 */
-	void GenerateViewportTypeMenu( FMenuBuilder& Builder ) const;
+	void GenerateViewportTypeMenu(FToolMenuSection& Section) const;
 
 	/**
 	 * Generates menu entries for spawning cameras at the current viewport
 	 *
-	 * @param Builder	The menu builder to add menu entries to
+	 * @param Menu	The menu to add menu entries to
 	 */
-	void GenerateCameraSpawnMenu(FMenuBuilder& Builder) const;
+	void GenerateCameraSpawnMenu(UToolMenu* Menu) const;
 
 	/**
 	 * Generates the toolbar view menu content 
@@ -149,6 +175,9 @@ private:
 	 */
 	TSharedRef<SWidget> GenerateShowMenu() const;
 
+	/** Fills the toolbar show menu content */
+	void FillShowMenu(UToolMenu* Menu) const;
+
 	/**
 	 * Returns the initial visibility of the view mode options widget 
 	 *
@@ -209,20 +238,20 @@ private:
 	/**
 	 * Generates the toolbar show layers menu content 
 	 *
-	 * @param MenuBuilder menu builder
+	 * @param Menu	The tool menu
 	 */
-	static void FillShowLayersMenu( class FMenuBuilder& MenuBuilder, TWeakPtr<class SLevelViewport> Viewport );
+	static void FillShowLayersMenu(UToolMenu* Menu, TWeakPtr<class SLevelViewport> Viewport );
 
 	/**
 	 * Generates 'show foliage types' menu content for a viewport
 	 *
-	 * @param MenuBuilder	menu builder
+	 * @param Menu	The tool menu
 	 * @param Viewport		target vieport
 	 */
-	static void FillShowFoliageTypesMenu(class FMenuBuilder& MenuBuilder, TWeakPtr<class SLevelViewport> Viewport);
+	static void FillShowFoliageTypesMenu(UToolMenu* Menu, TWeakPtr<class SLevelViewport> Viewport);
 
 	/** Generates the layout sub-menu content */
-	void GenerateViewportConfigsMenu(FMenuBuilder& MenuBuilder) const;
+	void GenerateViewportConfigsMenu(UToolMenu* Menu) const;
 
 	/** Gets the world we are editing */
 	TWeakObjectPtr<UWorld> GetWorld() const;
diff --git a/Engine/Source/Editor/LevelEditor/Public/LevelEditorMenuContext.h b/Engine/Source/Editor/LevelEditor/Public/LevelEditorMenuContext.h
--- a/Engine/Source/Editor/LevelEditor/Public/LevelEditorMenuContext.h
+++ b/Engine/Source/Editor/LevelEditor/Public/LevelEditorMenuContext.h
@@ -10,6 +10,8 @@
 
 class SLevelEditor;
 class UActorComponent;
+class SLevelViewport;
+class SLevelViewportToolBar;
 
 UCLASS()
 class LEVELEDITOR_API ULevelEditorMenuContext : public UObject
@@ -39,3 +41,11 @@ public:
 	ELevelEditorMenuContext ContextType;
 };
 
+UCLASS()
+class LEVELEDITOR_API ULevelViewportToolBarContext : public UObject
+{
+	GENERATED_BODY()
+public:
+	TWeakPtr<SLevelViewportToolBar> LevelViewportToolBarWidget;
+	TWeakPtr<const SLevelViewportToolBar> LevelViewportToolBarWidgetConst;
+};
