Added realtime and scalability indicators in the editor viewports so that users understand when realtime is off and non-default scalability is affecting their scene.

diff --git a/Engine/Source/Editor/EditorStyle/Private/SlateEditorStyle.cpp b/Engine/Source/Editor/EditorStyle/Private/SlateEditorStyle.cpp
--- a/Engine/Source/Editor/EditorStyle/Private/SlateEditorStyle.cpp
+++ b/Engine/Source/Editor/EditorStyle/Private/SlateEditorStyle.cpp
@@ -5396,6 +5396,12 @@ void FSlateEditorStyle::FStyle::SetupLevelEditorStyle()
 			.SetPressed(BOX_BRUSH( "Common/SmallRoundedButton", FMargin(7.f/16.f) ))
 		);
 
+		Set("EditorViewportToolBar.MenuButtonWarning", FButtonStyle(Button)
+			.SetNormal(BOX_BRUSH("Common/SmallRoundedButton", FMargin(7.f / 16.f), LogColor_Warning_LinearRef))
+			.SetHovered(BOX_BRUSH("Common/SmallRoundedButton", FMargin(7.f / 16.f), *LogColor_Warning_LinearRef*FLinearColor(1.1, 1.1, 1.1, 1)))
+			.SetPressed(BOX_BRUSH("Common/SmallRoundedButton", FMargin(7.f / 16.f), LogColor_Warning_LinearRef))
+		);
+
 		Set( "EditorViewportToolBar.Button", HoverHintOnly );
 
 		/* Set style structure for "EditorViewportToolBar.Button" ... */
diff --git a/Engine/Source/Editor/LevelEditor/Private/SLevelViewportToolBar.cpp b/Engine/Source/Editor/LevelEditor/Private/SLevelViewportToolBar.cpp
--- a/Engine/Source/Editor/LevelEditor/Private/SLevelViewportToolBar.cpp
+++ b/Engine/Source/Editor/LevelEditor/Private/SLevelViewportToolBar.cpp
@@ -33,6 +33,9 @@
 #include "FoliageType.h"
 #include "ShowFlagMenuCommands.h"
 #include "Bookmarks/BookmarkUI.h"
+#include "Scalability.h"
+#include "SScalabilitySettings.h"
+#include "Editor/EditorPerformanceSettings.h"
 
 #define LOCTEXT_NAMESPACE "LevelViewportToolBar"
 
@@ -224,6 +227,40 @@ void SLevelViewportToolBar::Construct( const FArguments& InArgs )
 					.Visibility(EVisibility::Collapsed)
 				]
 				+ SHorizontalBox::Slot()
+				.AutoWidth()
+				.Padding(ToolbarSlotPadding)
+				[
+					// Button to show that realtime is off
+					SNew(SEditorViewportToolBarButton)	
+					.Cursor(EMouseCursor::Default)
+					.ButtonType(EUserInterfaceActionType::Button)
+					.ButtonStyle(&FEditorStyle::Get().GetWidgetStyle<FButtonStyle>("EditorViewportToolBar.MenuButtonWarning"))
+					.OnClicked(this, &SLevelViewportToolBar::OnRealtimeWarningClicked)
+					.Visibility(this, &SLevelViewportToolBar::GetRealtimeWarningVisibility)
+					.ToolTipText(LOCTEXT("RealtimeOff_ToolTip", "This viewport is not updating in realtime.  Click to turn on realtime mode."))
+					.Content()
+					[
+						SNew(STextBlock)
+						.Font(FEditorStyle::GetFontStyle("EditorViewportToolBar.Font"))
+						.Text(LOCTEXT("RealtimeOff", "Realtime: Off"))
+						.ColorAndOpacity(FLinearColor::Black)
+					]
+				]
+				+ SHorizontalBox::Slot()
+				.AutoWidth()
+				.Padding(ToolbarSlotPadding)
+				[
+					// Button to show scalability warnings
+					SNew(SEditorViewportToolbarMenu)
+					.ParentToolBar(SharedThis(this))
+					.Cursor(EMouseCursor::Default)
+					.Label(this, &SLevelViewportToolBar::GetScalabilityWarningLabel)
+					.MenuStyle(FEditorStyle::Get(), "EditorViewportToolBar.MenuButtonWarning")
+					.OnGetMenuContent(this, &SLevelViewportToolBar::GetScalabilityWarningMenuContent)
+					.Visibility(this, &SLevelViewportToolBar::GetScalabilityWarningVisibility)
+					.ToolTipText(LOCTEXT("ScalabilityWarning_ToolTip", "Non-default scalability settings could be affecting what is shown in this viewport.\nFor example you may experience lower visual quality, reduced particle counts, and other artifacts that don't match what the scene would look like when running outside of the editor. Click to make changes."))
+				]
+				+ SHorizontalBox::Slot()
 				.Padding( ToolbarSlotPadding )
 				.HAlign( HAlign_Right )
 				[
@@ -1503,4 +1540,46 @@ void SLevelViewportToolBar::OnLandscapeLODChanged(int32 NewValue)
 }
 
 
+FReply SLevelViewportToolBar::OnRealtimeWarningClicked()
+{
+	FLevelEditorViewportClient& ViewportClient = Viewport.Pin()->GetLevelViewportClient();
+	ViewportClient.SetRealtime(true);
+
+	return FReply::Handled();
+}
+
+EVisibility SLevelViewportToolBar::GetRealtimeWarningVisibility() const
+{
+	FLevelEditorViewportClient& ViewportClient = Viewport.Pin()->GetLevelViewportClient();
+	// If the viewport is not realtime and there is no override then realtime is off
+	return !ViewportClient.IsRealtime() && !ViewportClient.IsRealtimeOverrideSet() ? EVisibility::Visible : EVisibility::Collapsed;
+}
+
+FText SLevelViewportToolBar::GetScalabilityWarningLabel() const
+{
+	const int32 QualityLevel = Scalability::GetQualityLevels().GetMinQualityLevel();
+	if (QualityLevel >= 0)
+	{
+		return FText::Format(LOCTEXT("ScalabilityWarning", "Scalability: {0}"), Scalability::GetScalabilityNameFromQualityLevel(QualityLevel));
+	}
+	
+	return FText::GetEmpty();
+}
+
+EVisibility SLevelViewportToolBar::GetScalabilityWarningVisibility() const
+{
+	//This method returns magic numbers. 3 means epic
+	return GetDefault<UEditorPerformanceSettings>()->bEnableScalabilityWarningIndicator && Scalability::GetQualityLevels().GetMinQualityLevel() < 3 ? EVisibility::Visible : EVisibility::Collapsed;
+}
+
+TSharedRef<SWidget> SLevelViewportToolBar::GetScalabilityWarningMenuContent() const
+{
+	return
+		SNew(SBorder)
+		.BorderImage(FEditorStyle::GetBrush("Menu.Background"))
+		[
+			SNew(SScalabilitySettings)
+		];
+}
+
 #undef LOCTEXT_NAMESPACE
diff --git a/Engine/Source/Editor/LevelEditor/Private/SLevelViewportToolBar.h b/Engine/Source/Editor/LevelEditor/Private/SLevelViewportToolBar.h
--- a/Engine/Source/Editor/LevelEditor/Private/SLevelViewportToolBar.h
+++ b/Engine/Source/Editor/LevelEditor/Private/SLevelViewportToolBar.h
@@ -234,6 +234,13 @@ private:
 	bool IsLandscapeLODSettingChecked(int32 Value) const;
 	void OnLandscapeLODChanged(int32 NewValue);
 
+	FReply OnRealtimeWarningClicked();
+	EVisibility GetRealtimeWarningVisibility() const;
+
+	FText GetScalabilityWarningLabel() const;
+	EVisibility GetScalabilityWarningVisibility() const;
+	TSharedRef<SWidget> GetScalabilityWarningMenuContent() const;
+
 private:
 	/**
 	 * Generates the toolbar show layers menu content 
diff --git a/Engine/Source/Editor/UnrealEd/Classes/Editor/EditorPerformanceSettings.h b/Engine/Source/Editor/UnrealEd/Classes/Editor/EditorPerformanceSettings.h
--- a/Engine/Source/Editor/UnrealEd/Classes/Editor/EditorPerformanceSettings.h
+++ b/Engine/Source/Editor/UnrealEd/Classes/Editor/EditorPerformanceSettings.h
@@ -37,6 +37,10 @@ class UEditorPerformanceSettings : public UDeveloperSettings
 	UPROPERTY(EditAnywhere, config, Category = EditorPerformance, meta = (DisplayName = "Enable Shared Data Cache Performance Notifications"))
 	uint32 bEnableSharedDDCPerformanceNotifications : 1;
 
+	/** When enabled, a warning will appear in the viewport when your editors scalability settings are non-default and you may be viewing a low quality scene */
+	UPROPERTY(EditAnywhere, config, Category = EditorPerformance, meta = (DisplayName = "Enable Scalability Warning Indicator"))
+	uint32 bEnableScalabilityWarningIndicator : 1;
+
 public:
 	/** UObject interface */
 	virtual void PostInitProperties() override;
diff --git a/Engine/Source/Editor/UnrealEd/Private/EditorPerformanceSettings.cpp b/Engine/Source/Editor/UnrealEd/Private/EditorPerformanceSettings.cpp
--- a/Engine/Source/Editor/UnrealEd/Private/EditorPerformanceSettings.cpp
+++ b/Engine/Source/Editor/UnrealEd/Private/EditorPerformanceSettings.cpp
@@ -14,9 +14,10 @@ UEditorPerformanceSettings::UEditorPerformanceSettings(const FObjectInitializer&
 	: Super(ObjectInitializer)
 	, bShowFrameRateAndMemory(false)
 	, bThrottleCPUWhenNotForeground(true)
-	, bMonitorEditorPerformance(true)
+	, bMonitorEditorPerformance(false)
 	, bOverrideDPIBasedEditorViewportScaling(false)
 	, bEnableSharedDDCPerformanceNotifications(true)
+	, bEnableScalabilityWarningIndicator(true)
 {
 
 }
diff --git a/Engine/Source/Editor/UnrealEd/Private/SEditorViewportToolBarButton.cpp b/Engine/Source/Editor/UnrealEd/Private/SEditorViewportToolBarButton.cpp
--- a/Engine/Source/Editor/UnrealEd/Private/SEditorViewportToolBarButton.cpp
+++ b/Engine/Source/Editor/UnrealEd/Private/SEditorViewportToolBarButton.cpp
@@ -10,7 +10,7 @@
 #include "EditorStyleSet.h"
 
 
-void SEditorViewportToolBarButton::Construct( const FArguments& Declaration )
+void SEditorViewportToolBarButton::Construct( const FArguments& Declaration)
 {
 	OnClickedDelegate = Declaration._OnClicked;
 	IsChecked = Declaration._IsChecked;
@@ -31,7 +31,7 @@ void SEditorViewportToolBarButton::Construct( const FArguments& Declaration )
 
 		ButtonWidget =
 			SNew( SButton )
-			.ButtonStyle( FEditorStyle::Get(), "EditorViewportToolBar.Button" ) 
+			.ButtonStyle(Declaration._ButtonStyle)
 			.OnClicked( OnClickedDelegate )
 			.HAlign( HAlign_Center )
 			.VAlign( VAlign_Center )
@@ -57,7 +57,7 @@ void SEditorViewportToolBarButton::Construct( const FArguments& Declaration )
 
 		ButtonWidget = 
 			SNew( SCheckBox )
-			.Style(FEditorStyle::Get(), "LevelViewportToolBar.CheckBoxButton" )
+			.Style(Declaration._CheckBoxStyle)
 			.OnCheckStateChanged( this, &SEditorViewportToolBarButton::OnCheckStateChanged )
 			.IsChecked( this, &SEditorViewportToolBarButton::OnIsChecked )
 			[
diff --git a/Engine/Source/Editor/UnrealEd/Private/SEditorViewportToolBarMenu.cpp b/Engine/Source/Editor/UnrealEd/Private/SEditorViewportToolBarMenu.cpp
--- a/Engine/Source/Editor/UnrealEd/Private/SEditorViewportToolBarMenu.cpp
+++ b/Engine/Source/Editor/UnrealEd/Private/SEditorViewportToolBarMenu.cpp
@@ -102,7 +102,7 @@ void SEditorViewportToolbarMenu::Construct( const FArguments& Declaration )
 			.ClickMethod( EButtonClickMethod::MouseDown )
 			.ContentPadding( FMargin( 5.0f, 2.0f ) )
 			.VAlign( VAlign_Center )
-			.ButtonStyle( FEditorStyle::Get(), "EditorViewportToolBar.MenuButton" )
+			.ButtonStyle(Declaration._MenuStyle)
 			.OnClicked( this, &SEditorViewportToolbarMenu::OnMenuClicked )
 			[
 				ButtonContent.ToSharedRef()
diff --git a/Engine/Source/Editor/UnrealEd/Public/SEditorViewportToolBarButton.h b/Engine/Source/Editor/UnrealEd/Public/SEditorViewportToolBarButton.h
--- a/Engine/Source/Editor/UnrealEd/Public/SEditorViewportToolBarButton.h
+++ b/Engine/Source/Editor/UnrealEd/Public/SEditorViewportToolBarButton.h
@@ -20,18 +20,29 @@ class UNREALED_API SEditorViewportToolBarButton : public SCompoundWidget
 {
 public:
 	SLATE_BEGIN_ARGS( SEditorViewportToolBarButton )
-		: _ButtonType( EUserInterfaceActionType::Button )
+		: _ButtonType(EUserInterfaceActionType::Button)
+		, _ButtonStyle(&FEditorStyle::Get().GetWidgetStyle<FButtonStyle>("EditorViewportToolBar.Button"))
+		, _CheckBoxStyle(&FEditorStyle::Get().GetWidgetStyle<FCheckBoxStyle>("LevelViewportToolBar.CheckBoxButton"))
 		, _IsChecked(false)
 		{}
 
 		/** Called when the button is clicked */
-		SLATE_EVENT( FOnClicked, OnClicked )
+	SLATE_EVENT(FOnClicked, OnClicked)
+
 		/** The button type to use */
-		SLATE_ARGUMENT( EUserInterfaceActionType, ButtonType )
+		SLATE_ARGUMENT(EUserInterfaceActionType, ButtonType)
+
+		/** Style to use when this is a regular button type */
+		SLATE_STYLE_ARGUMENT(FButtonStyle, ButtonStyle)
+
+		/** Style to use when this is a check box type */
+		SLATE_STYLE_ARGUMENT(FCheckBoxStyle, CheckBoxStyle)
+
 		/** Checked state of the button */
-		SLATE_ATTRIBUTE( bool, IsChecked )
+		SLATE_ATTRIBUTE(bool, IsChecked)
+
 		/** Style name of an image to use. Simple two state images are supported.  An image can be different depending on checked/unchecked state */
-		SLATE_ATTRIBUTE( FName, Image )
+		SLATE_ATTRIBUTE(FName, Image)
 		/** Any custom content to show in the button in place of other content */
 		SLATE_DEFAULT_SLOT( FArguments, Content )
 	SLATE_END_ARGS()
diff --git a/Engine/Source/Editor/UnrealEd/Public/SEditorViewportToolBarMenu.h b/Engine/Source/Editor/UnrealEd/Public/SEditorViewportToolBarMenu.h
--- a/Engine/Source/Editor/UnrealEd/Public/SEditorViewportToolBarMenu.h
+++ b/Engine/Source/Editor/UnrealEd/Public/SEditorViewportToolBarMenu.h
@@ -32,9 +32,13 @@ namespace EMenuItemType
 class UNREALED_API SEditorViewportToolbarMenu : public SCompoundWidget
 {
 public:
-	SLATE_BEGIN_ARGS( SEditorViewportToolbarMenu ){}
+	SLATE_BEGIN_ARGS( SEditorViewportToolbarMenu )
+		: _MenuStyle(&FEditorStyle::Get().GetWidgetStyle<FButtonStyle>("EditorViewportToolBar.MenuButton"))
+	{}
 		/** We need to know about the toolbar we are in */
 		SLATE_ARGUMENT( TSharedPtr<class SViewportToolBar>, ParentToolBar );
+		/** Style to use */
+		SLATE_STYLE_ARGUMENT(FButtonStyle, MenuStyle)
 		/** The label to show in the menu */
 		SLATE_ATTRIBUTE( FText, Label )
 		/** Optional icon to display next to the label */
diff --git a/Engine/Source/Runtime/Engine/Private/Scalability.cpp b/Engine/Source/Runtime/Engine/Private/Scalability.cpp
--- a/Engine/Source/Runtime/Engine/Private/Scalability.cpp
+++ b/Engine/Source/Runtime/Engine/Private/Scalability.cpp
@@ -401,6 +401,35 @@ float GetResolutionScreenPercentage()
 	return CVar->GetFloat();
 }
 
+FText GetScalabilityNameFromQualityLevel(int32 QualityLevel)
+{
+#define LOCTEXT_NAMESPACE "EngineScalabiltySettings"
+	static const FText NamesLow(LOCTEXT("QualityLowLabel", "Low"));
+	static const FText NamesMedium(LOCTEXT("QualityMediumLabel", "Medium"));
+	static const FText NamesHigh(LOCTEXT("QualityHighLabel", "High"));
+	static const FText NamesEpic(LOCTEXT("QualityEpicLabel", "Epic"));
+	static const FText NamesCine(LOCTEXT("QualityCineLabel", "Cinematic"));
+
+	switch (QualityLevel)
+	{
+	case 0:
+		return NamesLow;
+	case 1:
+		return NamesMedium;
+	case 2:
+		return NamesHigh;
+	case 3: 
+		return NamesEpic;
+	case 4:
+		return NamesCine;
+	default:
+		ensureMsgf(false, TEXT("Scalability Level %d needs a display name"), QualityLevel);
+		return FText::GetEmpty();
+	}
+
+#undef LOCTEXT_NAMESPACE
+}
+
 static void SetResolutionQualityLevel(float InResolutionQualityLevel)
 {
 	InResolutionQualityLevel = FMath::Clamp(InResolutionQualityLevel, Scalability::MinResolutionScale, Scalability::MaxResolutionScale);
@@ -823,6 +852,21 @@ int32 FQualityLevels::GetSingleQualityLevel() const
 	return -1;
 }
 
+int32 FQualityLevels::GetMinQualityLevel() const
+{
+	int32 Level = ViewDistanceQuality;
+
+	Level = FMath::Min(Level, AntiAliasingQuality);
+	Level = FMath::Min(Level, ShadowQuality);
+	Level = FMath::Min(Level, PostProcessQuality);
+	Level = FMath::Min(Level, TextureQuality);
+	Level = FMath::Min(Level, EffectsQuality);
+	Level = FMath::Min(Level, FoliageQuality);
+	Level = FMath::Min(Level, ShadingQuality);
+
+	return Level;
+}
+
 void FQualityLevels::SetViewDistanceQuality(int32 Value)
 {
 	ViewDistanceQuality = FMath::Clamp(Value, 0, CVarViewDistanceQuality_NumLevels->GetInt() - 1);
diff --git a/Engine/Source/Runtime/Engine/Public/Scalability.h b/Engine/Source/Runtime/Engine/Public/Scalability.h
--- a/Engine/Source/Runtime/Engine/Public/Scalability.h
+++ b/Engine/Source/Runtime/Engine/Public/Scalability.h
@@ -73,6 +73,10 @@ namespace Scalability
 		// @param Value -1:custom, 0:low, 1:medium, 2:high, 3:epic, 4:cinematic
 		int32 GetSingleQualityLevel() const;
 
+		// Returns the minimum set quality level from all settings
+		// @param Value -1:custom, 0:low, 1:medium, 2:high, 3:epic, 4:cinematic
+		int32 GetMinQualityLevel() const;
+
 		// Sets view distance quality
 		// @param Value 0:low, 1:medium, 2:high, 3:epic, 4:cinematic (gets clamped if needed)
 		void SetViewDistanceQuality(int32 Value);
@@ -166,6 +170,9 @@ namespace Scalability
 	/** Returns the current screen percentage */
 	ENGINE_API float GetResolutionScreenPercentage();
 
+	/** Returns a human readable name for a scalability quality level */
+	ENGINE_API FText GetScalabilityNameFromQualityLevel(int32 QualityLevel);
+
 #if WITH_EDITOR
 	/** Set an Editor preview scalability platform */
 	void ENGINE_API ChangeScalabilityPreviewPlatform(FName NewPlatformScalabilityName);
