Improve error message for failing to be able to run Live Coding.

diff --git a/Engine/Source/Editor/LevelEditor/Private/LevelEditorActions.cpp b/Engine/Source/Editor/LevelEditor/Private/LevelEditorActions.cpp
--- a/Engine/Source/Editor/LevelEditor/Private/LevelEditorActions.cpp
+++ b/Engine/Source/Editor/LevelEditor/Private/LevelEditorActions.cpp
@@ -1039,7 +1039,7 @@ void FLevelEditorActionCallbacks::RecompileGameCode_Clicked()
 		}
 		else
 		{
-			FMessageDialog::Open(EAppMsgType::Ok, LOCTEXT("NoLiveCodingCompileAfterHotReload", "Live Coding cannot be enabled after hot-reload has been used. Please restart the editor."));
+			FMessageDialog::Open(EAppMsgType::Ok, LOCTEXT("NoLiveCodingCompileAfterHotReload", "Live Coding cannot be enabled while hot-reloaded modules are active. Please close the editor and build from your IDE before restarting."));
 		}
 		return;
 	}
@@ -1086,7 +1086,7 @@ void FLevelEditorActionCallbacks::LiveCoding_ToggleEnabled()
 
 		if (LiveCoding->IsEnabledByDefault() && !LiveCoding->IsEnabledForSession())
 		{
-			FMessageDialog::Open(EAppMsgType::Ok, LOCTEXT("NoEnableLiveCodingAfterHotReload", "Live Coding cannot be enabled after hot-reload has been used. Please restart the editor."));
+			FMessageDialog::Open(EAppMsgType::Ok, LOCTEXT("NoEnableLiveCodingAfterHotReload", "Live Coding cannot be enabled while hot-reloaded modules are active. Please close the editor and build from your IDE before restarting."));
 		}
 	}
 }
