Disable input in Level Editor Viewports when a map is being torn down.

# Speculative fix - the callstack in the bug looks like we think we have a valid actor pointer, but it's probably getting garbage collected. The log is outputting level unload messages right before the crash.

diff --git a/Engine/Source/Editor/LevelEditor/Private/SLevelViewport.cpp b/Engine/Source/Editor/LevelEditor/Private/SLevelViewport.cpp
--- a/Engine/Source/Editor/LevelEditor/Private/SLevelViewport.cpp
+++ b/Engine/Source/Editor/LevelEditor/Private/SLevelViewport.cpp
@@ -1108,6 +1108,10 @@ TSharedRef< SWidget > SLevelViewport::BuildViewportDragDropContextMenu()
 
 void SLevelViewport::OnMapChanged( UWorld* World, EMapChangeType MapChangeType )
 {
+	// If the level is being unloaded, disable input in the viewport. Since click actions are deferred,
+	// It's possible to get into a state where the selected actors have been cleaned up by the time the click message is handled
+	LevelViewportClient->bDisableInput = (MapChangeType == EMapChangeType::TearDownWorld);
+
 	if( World && ( ( World == GetWorld() ) || ( World->EditorViews[LevelViewportClient->ViewportType].CamUpdated ) ) )
 	{
 		if( MapChangeType == EMapChangeType::LoadMap )
