Fixed incomplete toolltip for building static lighting, that was giving misleading information on the reasons building lighting could be disabled.

diff --git a/Engine/Source/Editor/LevelEditor/Private/LevelEditorActions.cpp b/Engine/Source/Editor/LevelEditor/Private/LevelEditorActions.cpp
--- a/Engine/Source/Editor/LevelEditor/Private/LevelEditorActions.cpp
+++ b/Engine/Source/Editor/LevelEditor/Private/LevelEditorActions.cpp
@@ -3083,10 +3083,10 @@ void FLevelEditorCommands::RegisterCommands()
 	UI_COMMAND( ExportAll, "Export All...", "Exports the entire level to a file on disk (multiple formats are supported.)", EUserInterfaceActionType::Button, FInputChord() );
 	UI_COMMAND( ExportSelected, "Export Selected...", "Exports currently-selected objects to a file on disk (multiple formats are supported.)", EUserInterfaceActionType::Button, FInputChord() );
 
-	UI_COMMAND( Build, "Build All Levels", "Builds all levels (precomputes lighting data and visibility data, generates navigation networks and updates brush models.)\nThis action is not available while Play in Editor is active, or when previewing less than Shader Model 5", EUserInterfaceActionType::Button, FInputChord() );
+	UI_COMMAND( Build, "Build All Levels", "Builds all levels (precomputes lighting data and visibility data, generates navigation networks and updates brush models.)\nThis action is not available while Play in Editor is active, static lighting is disabled in the project settings, or when previewing less than Shader Model 5", EUserInterfaceActionType::Button, FInputChord() );
 	UI_COMMAND( BuildAndSubmitToSourceControl, "Build and Submit...", "Displays a window that allows you to build all levels and submit them to source control", EUserInterfaceActionType::Button, FInputChord() );
-	UI_COMMAND( BuildLightingOnly, "Build Lighting", "Only precomputes lighting (all levels.)\nThis action is not available while Play in Editor is active, or when previewing less than Shader Model 5", EUserInterfaceActionType::Button, FInputChord(EModifierKey::Control|EModifierKey::Shift, EKeys::Semicolon) );
-	UI_COMMAND( BuildReflectionCapturesOnly, "Build Reflection Captures", "Updates Reflection Captures and stores their data in the BuildData package.\nThis action is not available while Play in Editor is active, or when previewing less than Shader Model 5", EUserInterfaceActionType::Button, FInputChord() );
+	UI_COMMAND( BuildLightingOnly, "Build Lighting", "Only precomputes lighting (all levels.)\nThis action is not available while Play in Editor is active, static lighting is disabled in the project settings, or when previewing less than Shader Model 5", EUserInterfaceActionType::Button, FInputChord(EModifierKey::Control|EModifierKey::Shift, EKeys::Semicolon) );
+	UI_COMMAND( BuildReflectionCapturesOnly, "Build Reflection Captures", "Updates Reflection Captures and stores their data in the BuildData package.\nThis action is not available while Play in Editor is active, static lighting is disabled in the project settings, or when previewing less than Shader Model 5", EUserInterfaceActionType::Button, FInputChord() );
 	UI_COMMAND( BuildLightingOnly_VisibilityOnly, "Precompute Static Visibility", "Only precomputes static visibility data (all levels.)", EUserInterfaceActionType::Button, FInputChord() );
 	UI_COMMAND( LightingBuildOptions_UseErrorColoring, "Use Error Coloring", "When enabled, errors during lighting precomputation will be baked as colors into light map data", EUserInterfaceActionType::ToggleButton, FInputChord() );
 	UI_COMMAND( LightingBuildOptions_ShowLightingStats, "Show Lighting Stats", "When enabled, a window containing metrics about lighting performance and memory will be displayed after a successful build.", EUserInterfaceActionType::ToggleButton, FInputChord() );
