Add forward declaration of a class
#rnx

diff --git a/Engine/Source/Editor/LevelEditor/Public/LevelEditorMenuContext.h b/Engine/Source/Editor/LevelEditor/Public/LevelEditorMenuContext.h
--- a/Engine/Source/Editor/LevelEditor/Public/LevelEditorMenuContext.h
+++ b/Engine/Source/Editor/LevelEditor/Public/LevelEditorMenuContext.h
@@ -12,6 +12,7 @@ class SLevelEditor;
 class UActorComponent;
 class SLevelViewport;
 class SLevelViewportToolBar;
+class FLevelEditorViewportClient;
 
 UCLASS()
 class LEVELEDITOR_API ULevelEditorMenuContext : public UObject
