Convert level editor viewport toolbar view menu
#rnx

diff --git a/Engine/Source/Editor/LevelEditor/Private/SLevelViewportToolBar.cpp b/Engine/Source/Editor/LevelEditor/Private/SLevelViewportToolBar.cpp
--- a/Engine/Source/Editor/LevelEditor/Private/SLevelViewportToolBar.cpp
+++ b/Engine/Source/Editor/LevelEditor/Private/SLevelViewportToolBar.cpp
@@ -36,6 +36,7 @@
 #include "Scalability.h"
 #include "SScalabilitySettings.h"
 #include "Editor/EditorPerformanceSettings.h"
+#include "SEditorViewportViewMenuContext.h"
 
 #define LOCTEXT_NAMESPACE "LevelViewportToolBar"
 
@@ -43,6 +44,29 @@
 class SLevelEditorViewportViewMenu : public SEditorViewportViewMenu
 {
 public:
+
+	void Construct(const FArguments& InArgs, TSharedRef<SEditorViewport> InViewport, TSharedRef<class SViewportToolBar> InParentToolBar)
+	{
+		SEditorViewportViewMenu::Construct(InArgs, InViewport, InParentToolBar);
+		MenuName = FName("LevelEditor.LevelViewportToolBar.View");
+	}
+
+	virtual void RegisterMenus() const override
+	{
+		SEditorViewportViewMenu::RegisterMenus();
+
+		if (!UToolMenus::Get()->IsMenuRegistered("LevelEditor.LevelViewportToolBar.View"))
+		{
+			UToolMenu* Menu = UToolMenus::Get()->RegisterMenu("LevelEditor.LevelViewportToolBar.View", "UnrealEd.ViewportToolbar.View");
+			Menu->AddDynamicSection("LevelSection", FNewToolMenuDelegate::CreateLambda([](UToolMenu* InMenu)
+			{
+				UEditorViewportViewMenuContext* Context = InMenu->FindContext<UEditorViewportViewMenuContext>();
+				TSharedPtr<SLevelViewportToolBar> LevelViewportToolBar = StaticCastSharedPtr<SLevelViewportToolBar>(Context->EditorViewportViewMenu.Pin()->GetParentToolBar().Pin());
+				LevelViewportToolBar->FillViewMenu(InMenu);
+			}));
+		}
+	}
+
 	virtual TSharedRef<SWidget> GenerateViewMenuContent() const override
 	{
 		SLevelViewport* LevelViewport = static_cast<SLevelViewport*>(Viewport.Pin().Get());
@@ -1442,69 +1466,95 @@ TWeakObjectPtr<UWorld> SLevelViewportToolBar::GetWorld() const
 TSharedPtr<FExtender> SLevelViewportToolBar::GetViewMenuExtender()
 {
 	FLevelEditorModule& LevelEditorModule = FModuleManager::GetModuleChecked<FLevelEditorModule>(TEXT("LevelEditor"));
-	TSharedPtr<FExtender> LevelEditorExtenders = LevelEditorModule.GetMenuExtensibilityManager()->GetAllExtenders();
-
-	TSharedRef<FExtender> Extender(new FExtender());
-	Extender->AddMenuExtension(
-		TEXT("ViewMode"),
-		EExtensionHook::After,
-		Viewport.Pin()->GetCommandList(),
-		FMenuExtensionDelegate::CreateSP(this, &SLevelViewportToolBar::CreateViewMenuExtensions));
-
-	TArray<TSharedPtr<FExtender>> Extenders;
-	Extenders.Reserve(2);
-	Extenders.Add(LevelEditorExtenders);
-	Extenders.Add(Extender);
-
-	return FExtender::Combine(Extenders);
+	return LevelEditorModule.GetMenuExtensibilityManager()->GetAllExtenders();
 }
 
-void SLevelViewportToolBar::CreateViewMenuExtensions(FMenuBuilder& MenuBuilder)
+void SLevelViewportToolBar::FillViewMenu(UToolMenu* Menu)
 {
-	MenuBuilder.BeginSection("LevelViewportDeferredRendering", LOCTEXT("DeferredRenderingHeader", "Deferred Rendering") );
-	MenuBuilder.EndSection();
-
-	MenuBuilder.AddSubMenu(LOCTEXT("VisualizeBufferViewModeDisplayName", "Buffer Visualization"),
-		LOCTEXT("BufferVisualizationMenu_ToolTip", "Select a mode for buffer visualization"),
-		FNewMenuDelegate::CreateStatic(&FBufferVisualizationMenuCommands::BuildVisualisationSubMenu),
-		FUIAction(FExecuteAction(), FCanExecuteAction(),
-			FIsActionChecked::CreateLambda([this]()
-			{
-				const TSharedRef<SEditorViewport> ViewportRef = Viewport.Pin().ToSharedRef();
-				const TSharedPtr<FEditorViewportClient> ViewportClient = ViewportRef->GetViewportClient();
-				check(ViewportClient.IsValid());
-				return ViewportClient->IsViewModeEnabled(VMI_VisualizeBuffer);
-			})),
-		/* InExtensionHook = */ NAME_None, EUserInterfaceActionType::RadioButton,
-		/* bInOpenSubMenuOnClick = */ false, FSlateIcon(FEditorStyle::GetStyleSetName(), "EditorViewport.VisualizeBufferMode"));
+	FToolMenuInsert InsertPosition("ViewMode", EToolMenuInsertType::After);
 
-	MenuBuilder.BeginSection("LevelViewportLandscape", LOCTEXT("LandscapeHeader", "Landscape") );
 	{
+		FToolMenuSection& Section = Menu->AddSection("LevelViewportDeferredRendering", LOCTEXT("DeferredRenderingHeader", "Deferred Rendering"), InsertPosition);
+	}
+
+	{
+		FToolMenuSection& Section = Menu->FindOrAddSection("ViewMode");
+		Section.AddSubMenu(
+			"VisualizeBufferViewMode",
+			LOCTEXT("VisualizeBufferViewModeDisplayName", "Buffer Visualization"),
+			LOCTEXT("BufferVisualizationMenu_ToolTip", "Select a mode for buffer visualization"),
+			FNewMenuDelegate::CreateStatic(&FBufferVisualizationMenuCommands::BuildVisualisationSubMenu),
+			FUIAction(
+				FExecuteAction(),
+				FCanExecuteAction(),
+				FIsActionChecked::CreateLambda([this]()
+				{
+					const TSharedRef<SEditorViewport> ViewportRef = Viewport.Pin().ToSharedRef();
+					const TSharedPtr<FEditorViewportClient> ViewportClient = ViewportRef->GetViewportClient();
+					check(ViewportClient.IsValid());
+					return ViewportClient->IsViewModeEnabled(VMI_VisualizeBuffer);
+				})
+			),
+			EUserInterfaceActionType::RadioButton,
+			/* bInOpenSubMenuOnClick = */ false,
+			FSlateIcon(FEditorStyle::GetStyleSetName(), "EditorViewport.VisualizeBufferMode")
+		);
+	}
+
+	{
+		FToolMenuSection& Section = Menu->AddSection("LevelViewportLandscape", LOCTEXT("LandscapeHeader", "Landscape"), InsertPosition);
+
 		struct Local
 		{
-			static void BuildLandscapeLODMenu(FMenuBuilder& Menu, SLevelViewportToolBar* Toolbar)
+			static void BuildLandscapeLODMenu(UToolMenu* Menu, SLevelViewportToolBar* Toolbar)
 			{
-				Menu.BeginSection("LevelViewportLandScapeLOD", LOCTEXT("LandscapeLODHeader", "Landscape LOD"));
 				{
+					FToolMenuSection& SubMenuSection = Menu->AddSection("LevelViewportLandScapeLOD", LOCTEXT("LandscapeLODHeader", "Landscape LOD"));
+
+					SubMenuSection.AddMenuEntry(
+						"LandscapeLODAuto",
+						LOCTEXT("LandscapeLODAuto", "Auto"),
+						FText(),
+						FSlateIcon(),
+						FUIAction(
+							FExecuteAction::CreateSP(Toolbar, &SLevelViewportToolBar::OnLandscapeLODChanged, -1),
+							FCanExecuteAction(),
+							FIsActionChecked::CreateSP(Toolbar, &SLevelViewportToolBar::IsLandscapeLODSettingChecked, -1)
+						),
+						EUserInterfaceActionType::RadioButton
+					);
+
+					SubMenuSection.AddMenuSeparator("LandscapeLODSeparator");
+
 					static const FText FormatString = LOCTEXT("LandscapeLODFixed", "Fixed at {0}");
-					Menu.AddMenuEntry(LOCTEXT("LandscapeLODAuto", "Auto"), FText(), FSlateIcon(), FUIAction(FExecuteAction::CreateSP(Toolbar, &SLevelViewportToolBar::OnLandscapeLODChanged, -1), FCanExecuteAction(), FIsActionChecked::CreateSP(Toolbar, &SLevelViewportToolBar::IsLandscapeLODSettingChecked, -1)), NAME_None, EUserInterfaceActionType::RadioButton);
-					Menu.AddMenuSeparator();
-					Menu.AddMenuEntry(FText::Format(FormatString, FText::AsNumber(0)), FText(), FSlateIcon(), FUIAction(FExecuteAction::CreateSP(Toolbar, &SLevelViewportToolBar::OnLandscapeLODChanged, 0), FCanExecuteAction(), FIsActionChecked::CreateSP(Toolbar, &SLevelViewportToolBar::IsLandscapeLODSettingChecked, 0)), NAME_None, EUserInterfaceActionType::RadioButton);
-					Menu.AddMenuEntry(FText::Format(FormatString, FText::AsNumber(1)), FText(), FSlateIcon(), FUIAction(FExecuteAction::CreateSP(Toolbar, &SLevelViewportToolBar::OnLandscapeLODChanged, 1), FCanExecuteAction(), FIsActionChecked::CreateSP(Toolbar, &SLevelViewportToolBar::IsLandscapeLODSettingChecked, 1)), NAME_None, EUserInterfaceActionType::RadioButton);
-					Menu.AddMenuEntry(FText::Format(FormatString, FText::AsNumber(2)), FText(), FSlateIcon(), FUIAction(FExecuteAction::CreateSP(Toolbar, &SLevelViewportToolBar::OnLandscapeLODChanged, 2), FCanExecuteAction(), FIsActionChecked::CreateSP(Toolbar, &SLevelViewportToolBar::IsLandscapeLODSettingChecked, 2)), NAME_None, EUserInterfaceActionType::RadioButton);
-					Menu.AddMenuEntry(FText::Format(FormatString, FText::AsNumber(3)), FText(), FSlateIcon(), FUIAction(FExecuteAction::CreateSP(Toolbar, &SLevelViewportToolBar::OnLandscapeLODChanged, 3), FCanExecuteAction(), FIsActionChecked::CreateSP(Toolbar, &SLevelViewportToolBar::IsLandscapeLODSettingChecked, 3)), NAME_None, EUserInterfaceActionType::RadioButton);
-					Menu.AddMenuEntry(FText::Format(FormatString, FText::AsNumber(4)), FText(), FSlateIcon(), FUIAction(FExecuteAction::CreateSP(Toolbar, &SLevelViewportToolBar::OnLandscapeLODChanged, 4), FCanExecuteAction(), FIsActionChecked::CreateSP(Toolbar, &SLevelViewportToolBar::IsLandscapeLODSettingChecked, 4)), NAME_None, EUserInterfaceActionType::RadioButton);
-					Menu.AddMenuEntry(FText::Format(FormatString, FText::AsNumber(5)), FText(), FSlateIcon(), FUIAction(FExecuteAction::CreateSP(Toolbar, &SLevelViewportToolBar::OnLandscapeLODChanged, 5), FCanExecuteAction(), FIsActionChecked::CreateSP(Toolbar, &SLevelViewportToolBar::IsLandscapeLODSettingChecked, 5)), NAME_None, EUserInterfaceActionType::RadioButton);
-					Menu.AddMenuEntry(FText::Format(FormatString, FText::AsNumber(6)), FText(), FSlateIcon(), FUIAction(FExecuteAction::CreateSP(Toolbar, &SLevelViewportToolBar::OnLandscapeLODChanged, 6), FCanExecuteAction(), FIsActionChecked::CreateSP(Toolbar, &SLevelViewportToolBar::IsLandscapeLODSettingChecked, 6)), NAME_None, EUserInterfaceActionType::RadioButton);
-					Menu.AddMenuEntry(FText::Format(FormatString, FText::AsNumber(7)), FText(), FSlateIcon(), FUIAction(FExecuteAction::CreateSP(Toolbar, &SLevelViewportToolBar::OnLandscapeLODChanged, 7), FCanExecuteAction(), FIsActionChecked::CreateSP(Toolbar, &SLevelViewportToolBar::IsLandscapeLODSettingChecked, 7)), NAME_None, EUserInterfaceActionType::RadioButton);
+					for (int32 i = 0; i < 8; ++i)
+					{
+						SubMenuSection.AddMenuEntry(
+							NAME_None,
+							FText::Format(FormatString, FText::AsNumber(i)),
+							FText(),
+							FSlateIcon(),
+							FUIAction(
+								FExecuteAction::CreateSP(Toolbar, &SLevelViewportToolBar::OnLandscapeLODChanged, i),
+								FCanExecuteAction(),
+								FIsActionChecked::CreateSP(Toolbar, &SLevelViewportToolBar::IsLandscapeLODSettingChecked, i)
+							),
+							EUserInterfaceActionType::RadioButton
+						);
+					}
 				}
-				Menu.EndSection();
 			}
 		};
 
-		MenuBuilder.AddSubMenu(LOCTEXT("LandscapeLODDisplayName", "LOD"), LOCTEXT("LandscapeLODMenu_ToolTip", "Override Landscape LOD in this viewport"), FNewMenuDelegate::CreateStatic(&Local::BuildLandscapeLODMenu, this), /*Default*/false, FSlateIcon());
+		Section.AddSubMenu(
+			"LandscapeLOD",
+			LOCTEXT("LandscapeLODDisplayName", "LOD"),
+			LOCTEXT("LandscapeLODMenu_ToolTip", "Override Landscape LOD in this viewport"),
+			FNewToolMenuDelegate::CreateStatic(&Local::BuildLandscapeLODMenu, this),
+			/*bInOpenSubMenuOnClick=*/ false,
+			FSlateIcon()
+		);
 	}
-	MenuBuilder.EndSection();
 }
 
 
diff --git a/Engine/Source/Editor/LevelEditor/Private/SLevelViewportToolBar.h b/Engine/Source/Editor/LevelEditor/Private/SLevelViewportToolBar.h
--- a/Engine/Source/Editor/LevelEditor/Private/SLevelViewportToolBar.h
+++ b/Engine/Source/Editor/LevelEditor/Private/SLevelViewportToolBar.h
@@ -33,6 +33,9 @@ public:
 	/** @return Level editor viewport client. */ 
 	FLevelEditorViewportClient* GetLevelViewportClient();
 
+	/** Fills view menu */
+	void FillViewMenu(UToolMenu* InMenu);
+
 private:
 	/**
 	 * Returns the label for the "Camera" tool bar menu, which changes depending on the viewport type
@@ -269,8 +272,6 @@ private:
 	/** Gets the extender for the view menu */
 	TSharedPtr<FExtender> GetViewMenuExtender();
 
-	void CreateViewMenuExtensions(FMenuBuilder& MenuBuilder);
-
 	/** Called when the user disables realtime override from the toolbar */
 	void OnDisableRealtimeOverride();
 	bool IsRealtimeOverrideToggleVisible() const;
diff --git a/Engine/Source/Editor/UnrealEd/Private/SEditorViewportToolBarMenu.cpp b/Engine/Source/Editor/UnrealEd/Private/SEditorViewportToolBarMenu.cpp
--- a/Engine/Source/Editor/UnrealEd/Private/SEditorViewportToolBarMenu.cpp
+++ b/Engine/Source/Editor/UnrealEd/Private/SEditorViewportToolBarMenu.cpp
@@ -148,3 +148,7 @@ EVisibility SEditorViewportToolbarMenu::GetLabelIconVisibility() const
 	return LabelIconBrush.Get() != NULL ? EVisibility::Visible : EVisibility::Collapsed;
 }
 
+TWeakPtr<class SViewportToolBar> SEditorViewportToolbarMenu::GetParentToolBar() const
+{
+	return ParentToolBar;
+}
\ No newline at end of file
diff --git a/Engine/Source/Editor/UnrealEd/Private/SEditorViewportViewMenu.cpp b/Engine/Source/Editor/UnrealEd/Private/SEditorViewportViewMenu.cpp
--- a/Engine/Source/Editor/UnrealEd/Private/SEditorViewportViewMenu.cpp
+++ b/Engine/Source/Editor/UnrealEd/Private/SEditorViewportViewMenu.cpp
@@ -12,9 +12,12 @@
 
 #define LOCTEXT_NAMESPACE "EditorViewportViewMenu"
 
+const FName SEditorViewportViewMenu::BaseMenuName("UnrealEd.ViewportToolbar.View");
+
 void SEditorViewportViewMenu::Construct( const FArguments& InArgs, TSharedRef<SEditorViewport> InViewport, TSharedRef<class SViewportToolBar> InParentToolBar )
 {
 	Viewport = InViewport;
+	MenuName = BaseMenuName;
 	MenuExtenders = InArgs._MenuExtenders;
 
 	SEditorViewportToolbarMenu::Construct
@@ -191,21 +194,27 @@ const FSlateBrush* SEditorViewportViewMenu::GetViewMenuLabelIcon() const
 	return FEditorStyle::GetBrush(Icon);
 }
 
-TSharedRef<SWidget> SEditorViewportViewMenu::GenerateViewMenuContent() const
+void SEditorViewportViewMenu::RegisterMenus() const
 {
-	static const FName MenuName("UnrealEd.ViewportToolbar.View");
-	if (!UToolMenus::Get()->IsMenuRegistered(MenuName))
+	if (!UToolMenus::Get()->IsMenuRegistered(BaseMenuName))
 	{
-		UToolMenu* Menu = UToolMenus::Get()->RegisterMenu(MenuName);
-		Menu->AddDynamicSection("DynamicSection", FNewToolMenuDelegate::CreateLambda([](UToolMenu* InMenu)
+		UToolMenu* Menu = UToolMenus::Get()->RegisterMenu(BaseMenuName);
+		Menu->AddDynamicSection("BaseSection", FNewToolMenuDelegate::CreateLambda([](UToolMenu* InMenu)
 		{
-			UEditorViewportViewMenuContext* Context = InMenu->FindContext<UEditorViewportViewMenuContext>();
-			Context->Widget.Pin()->FillViewMenu(InMenu);
+			if (UEditorViewportViewMenuContext* Context = InMenu->FindContext<UEditorViewportViewMenuContext>())
+			{
+				Context->EditorViewportViewMenu.Pin()->FillViewMenu(InMenu);
+			}
 		}));
 	}
+}
+
+TSharedRef<SWidget> SEditorViewportViewMenu::GenerateViewMenuContent() const
+{
+	RegisterMenus();
 
 	UEditorViewportViewMenuContext* ContextObject = NewObject<UEditorViewportViewMenuContext>();
-	ContextObject->Widget = SharedThis(this);
+	ContextObject->EditorViewportViewMenu = SharedThis(this);
 
 	FToolMenuContext MenuContext(Viewport.Pin()->GetCommandList(), MenuExtenders, ContextObject);
 	return UToolMenus::Get()->GenerateWidget(MenuName, MenuContext);
diff --git a/Engine/Source/Editor/UnrealEd/Public/SEditorViewportToolBarMenu.h b/Engine/Source/Editor/UnrealEd/Public/SEditorViewportToolBarMenu.h
--- a/Engine/Source/Editor/UnrealEd/Public/SEditorViewportToolBarMenu.h
+++ b/Engine/Source/Editor/UnrealEd/Public/SEditorViewportToolBarMenu.h
@@ -56,6 +56,12 @@ public:
 	 * Constructs the menu
 	 */
 	void Construct( const FArguments& Declaration );
+
+	/**
+	 * Returns parent tool bar
+	 */
+	TWeakPtr<class SViewportToolBar> GetParentToolBar() const;
+
 private:
 	/**
 	 * Called when the menu button is clicked.  Will toggle the visibility of the menu content                   
@@ -74,6 +80,9 @@ protected:
 	/** Parent tool bar for querying other open menus */
 	TWeakPtr<class SViewportToolBar> ParentToolBar;
 
+	/** Name of tool menu */
+	FName MenuName;
+
 private:
 	/** Our menus anchor */
 	TSharedPtr<SMenuAnchor> MenuAnchor;
diff --git a/Engine/Source/Editor/UnrealEd/Public/SEditorViewportViewMenu.h b/Engine/Source/Editor/UnrealEd/Public/SEditorViewportViewMenu.h
--- a/Engine/Source/Editor/UnrealEd/Public/SEditorViewportViewMenu.h
+++ b/Engine/Source/Editor/UnrealEd/Public/SEditorViewportViewMenu.h
@@ -30,8 +30,10 @@ private:
 
 protected:
 	virtual TSharedRef<SWidget> GenerateViewMenuContent() const;
-	TWeakPtr<SEditorViewport> Viewport;
+	virtual void RegisterMenus() const;
 
-private:
+	TWeakPtr<SEditorViewport> Viewport;
 	TSharedPtr<class FExtender> MenuExtenders;
+
+	static const FName BaseMenuName;
 };
diff --git a/Engine/Source/Editor/UnrealEd/Public/SEditorViewportViewMenuContext.h b/Engine/Source/Editor/UnrealEd/Public/SEditorViewportViewMenuContext.h
--- a/Engine/Source/Editor/UnrealEd/Public/SEditorViewportViewMenuContext.h
+++ b/Engine/Source/Editor/UnrealEd/Public/SEditorViewportViewMenuContext.h
@@ -14,5 +14,5 @@ class UNREALED_API UEditorViewportViewMenuContext : public UObject
 	GENERATED_BODY()
 public:
 
-	TWeakPtr<const SEditorViewportViewMenu> Widget;
+	TWeakPtr<const SEditorViewportViewMenu> EditorViewportViewMenu;
 };
