Set names on level editor toolbar camera menu items
#rnx

diff --git a/Engine/Source/Editor/LevelEditor/Private/SLevelViewportToolBar.cpp b/Engine/Source/Editor/LevelEditor/Private/SLevelViewportToolBar.cpp
--- a/Engine/Source/Editor/LevelEditor/Private/SLevelViewportToolBar.cpp
+++ b/Engine/Source/Editor/LevelEditor/Private/SLevelViewportToolBar.cpp
@@ -889,10 +889,11 @@ void SLevelViewportToolBar::GeneratePlacedCameraMenuEntries(UToolMenu* Menu, TAr
 void SLevelViewportToolBar::GenerateViewportTypeMenu(FToolMenuSection& Section) const
 {
 	FLevelEditorModule& LevelEditorModule = FModuleManager::GetModuleChecked<FLevelEditorModule>(TEXT("LevelEditor"));
-	LevelEditorModule.IterateViewportTypes([&](FName, const FViewportTypeDefinition& InDefinition) {
+	LevelEditorModule.IterateViewportTypes([&](FName ViewportTypeName, const FViewportTypeDefinition& InDefinition)
+	{
 		if (InDefinition.ActivationCommand.IsValid())
 		{
-			Section.AddMenuEntry(NAME_None, InDefinition.ActivationCommand);
+			Section.AddMenuEntry(*FString::Printf(TEXT("ViewportType_%s"), *ViewportTypeName.ToString()), InDefinition.ActivationCommand);
 		}
 	});
 }
