Hide details panel extension when the selection is empty.

diff --git a/Engine/Source/Editor/LevelEditor/Private/SActorDetails.cpp b/Engine/Source/Editor/LevelEditor/Private/SActorDetails.cpp
--- a/Engine/Source/Editor/LevelEditor/Private/SActorDetails.cpp
+++ b/Engine/Source/Editor/LevelEditor/Private/SActorDetails.cpp
@@ -159,7 +159,7 @@ void SActorDetails::Construct(const FArguments& InArgs, const FName TabIdentifie
 		.Padding(0.0f, 0.0f, 0.0f, 0.0f)
 		.AutoHeight()
 		[
-			SNew(SExtensionPanel)
+			SAssignNew(ExtensionPanel, SExtensionPanel)
 			.ExtensionPanelID("ActorDetailsPanel")
 		]
 		+SVerticalBox::Slot()
@@ -240,28 +240,30 @@ void SActorDetails::OnDetailsViewObjectArrayChanged(const FString& InTitle, cons
 	// The DetailsView will already check validity every tick and hide itself when invalid, so this piggy-backs on that code instead of needing a second tick function.
 	if (InObjects.Num() == 0 && !LockedActorSelection.IsValid())
 	{
-		bShowingComponents = false;
+		bHasComponentsToShow = false;
 	}
 }
 
 void SActorDetails::SetObjects(const TArray<UObject*>& InObjects, bool bForceRefresh)
 {
-	if(!DetailsView->IsLocked())
+	if (!DetailsView->IsLocked())
 	{
 		DetailsView->SetObjects(InObjects, bForceRefresh);
 
-		bShowingComponents = false;
+		ExtensionPanel->SetVisibility(InObjects.Num() > 0 ? EVisibility::Visible : EVisibility::Collapsed);
+
+		bHasComponentsToShow = false;
 
-		if(InObjects.Num() == 1 && FKismetEditorUtilities::CanCreateBlueprintOfClass(InObjects[0]->GetClass()))
+		if (InObjects.Num() == 1 && FKismetEditorUtilities::CanCreateBlueprintOfClass(InObjects[0]->GetClass()))
 		{
 			AActor* Actor = GetSelectedActorInEditor();
-			if(Actor)
+			if (Actor)
 			{
 				LockedActorSelection = Actor;
-				bShowingComponents = true;
+				bHasComponentsToShow = true;
 
 				// Update the tree if a new actor is selected
-				if(GEditor->GetSelectedComponentCount() == 0)
+				if (GEditor->GetSelectedComponentCount() == 0)
 				{
 					// Enable the selection guard to prevent OnTreeSelectionChanged() from altering the editor's component selection
 					TGuardValue<bool> SelectionGuard(bSelectionGuard, true);
@@ -270,7 +272,7 @@ void SActorDetails::SetObjects(const TArray<UObject*>& InObjects, bool bForceRef
 			}
 		}
 
-		if(DetailsView->GetHostTabManager().IsValid())
+		if (DetailsView->GetHostTabManager().IsValid())
 		{
 			TSharedPtr<SDockTab> Tab = DetailsView->GetHostTabManager()->FindExistingLiveTab(DetailsView->GetIdentifier());
 			if (Tab.IsValid() && !Tab->IsForeground() )
@@ -707,7 +709,7 @@ void SActorDetails::OnNativeComponentWarningHyperlinkClicked(const FSlateHyperli
 
 EVisibility SActorDetails::GetComponentsBoxVisibility() const
 {
-	return (bShowingComponents & ShowComponents->GetBool()) ? EVisibility::Visible : EVisibility::Collapsed;
+	return (bHasComponentsToShow & ShowComponents->GetBool()) ? EVisibility::Visible : EVisibility::Collapsed;
 }
 
 EVisibility SActorDetails::GetUCSComponentWarningVisibility() const
diff --git a/Engine/Source/Editor/LevelEditor/Private/SActorDetails.h b/Engine/Source/Editor/LevelEditor/Private/SActorDetails.h
--- a/Engine/Source/Editor/LevelEditor/Private/SActorDetails.h
+++ b/Engine/Source/Editor/LevelEditor/Private/SActorDetails.h
@@ -83,6 +83,7 @@ private:
 	TSharedPtr<class IDetailsView> DetailsView;
 	TSharedPtr<SBox> ComponentsBox;
 	TSharedPtr<class SSCSEditor> SCSEditor;
+	TSharedPtr<class SExtensionPanel> ExtensionPanel;
 
 	// The actor selected when the details panel was locked
 	TWeakObjectPtr<AActor> LockedActorSelection;
@@ -94,8 +95,8 @@ private:
 	// Used to prevent reentrant changes
 	bool bSelectionGuard;
 
-	// True if the actor details should be showing the components.
-	bool bShowingComponents = false;
+	// True if the actor details has any component to show.
+	bool bHasComponentsToShow = false;
 
 	// True if the actor "root" node in the SCS editor is currently shown as selected
 	bool bShowingRootActorNodeSelected;
