Fixed static analysis warning

#rnx

diff --git a/Engine/Source/Editor/LevelEditor/Private/SActorDetails.cpp b/Engine/Source/Editor/LevelEditor/Private/SActorDetails.cpp
--- a/Engine/Source/Editor/LevelEditor/Private/SActorDetails.cpp
+++ b/Engine/Source/Editor/LevelEditor/Private/SActorDetails.cpp
@@ -709,7 +709,7 @@ void SActorDetails::OnNativeComponentWarningHyperlinkClicked(const FSlateHyperli
 
 EVisibility SActorDetails::GetComponentsBoxVisibility() const
 {
-	return (bHasComponentsToShow & ShowComponents->GetBool()) ? EVisibility::Visible : EVisibility::Collapsed;
+	return (bHasComponentsToShow && ShowComponents->GetBool()) ? EVisibility::Visible : EVisibility::Collapsed;
 }
 
 EVisibility SActorDetails::GetUCSComponentWarningVisibility() const
