Fix an editor UI crash introduced by my recent change to support CAC template expansion in the IWCE tree view.

#rnx

diff --git a/Engine/Source/Editor/LevelEditor/Private/SActorDetails.cpp b/Engine/Source/Editor/LevelEditor/Private/SActorDetails.cpp
--- a/Engine/Source/Editor/LevelEditor/Private/SActorDetails.cpp
+++ b/Engine/Source/Editor/LevelEditor/Private/SActorDetails.cpp
@@ -719,7 +719,7 @@ EVisibility SActorDetails::GetUCSComponentWarningVisibility() const
 	// Check to see if any selected components are inherited from blueprint
 	for (const FSCSEditorTreeNodePtrType& Node : SCSEditor->GetSelectedNodes())
 	{
-		if (Node->IsBlueprintComponent())
+		if (!Node->IsNativeComponent())
 		{
 			UActorComponent* Component = Node->GetComponentTemplate();
 			bIsUneditableBlueprintComponent = Component ? Component->CreationMethod == EComponentCreationMethod::UserConstructionScript : false;
@@ -764,7 +764,7 @@ EVisibility SActorDetails::GetInheritedBlueprintComponentWarningVisibility() con
 	// Check to see if any selected components are inherited from blueprint
 	for (const FSCSEditorTreeNodePtrType& Node : SCSEditor->GetSelectedNodes())
 	{
-		if (Node->IsBlueprintComponent())
+		if (!Node->IsNativeComponent())
 		{
 			if (UActorComponent* Component = Node->GetComponentTemplate())
 			{
