Initial implementation of a status bar. All it has right now is source control

Renamed old status bar API to TitleBar because thats where it was putting things. Fixed up a few use cases.  Other use cases will be moved to the status bar

diff --git a/Engine/Plugins/Experimental/VirtualProductionUtilities/Source/VPUtilitiesEditor/Private/VPUtilitiesEditorModule.cpp b/Engine/Plugins/Experimental/VirtualProductionUtilities/Source/VPUtilitiesEditor/Private/VPUtilitiesEditorModule.cpp
--- a/Engine/Plugins/Experimental/VirtualProductionUtilities/Source/VPUtilitiesEditor/Private/VPUtilitiesEditorModule.cpp
+++ b/Engine/Plugins/Experimental/VirtualProductionUtilities/Source/VPUtilitiesEditor/Private/VPUtilitiesEditorModule.cpp
@@ -88,11 +88,11 @@ public:
 		FLevelEditorModule* LevelEditorModule = FModuleManager::GetModulePtr<FLevelEditorModule>("LevelEditor");
 		if (LevelEditorModule != nullptr)
 		{
-			FLevelEditorModule::FStatusBarItem Item;
+			FLevelEditorModule::FTitleBarItem Item;
 			Item.Label = LOCTEXT("VPRolesLabel", "VP Roles: ");
 			Item.Value = MakeAttributeLambda([]() { return FText::FromString(GetMutableDefault<UVPSettings>()->GetRoles().ToStringSimple()); });
 			Item.Visibility = MakeAttributeLambda([]() { return GetMutableDefault<UVPSettings>()->bShowRoleInEditor ? EVisibility::SelfHitTestInvisible : EVisibility::Collapsed; });
-			LevelEditorModule->AddStatusBarItem(VPRoleNotificationBarIdentifier, Item);
+			LevelEditorModule->AddTitleBarItem(VPRoleNotificationBarIdentifier, Item);
 		}
 	}
 
@@ -109,7 +109,7 @@ public:
 		FLevelEditorModule* LevelEditorModule = FModuleManager::GetModulePtr<FLevelEditorModule>("LevelEditor");
 		if (LevelEditorModule != nullptr)
 		{
-			LevelEditorModule->RemoveStatusBarItem(VPRoleNotificationBarIdentifier);
+			LevelEditorModule->RemoveTitleBarItem(VPRoleNotificationBarIdentifier);
 		}
 	}
 };
diff --git a/Engine/Plugins/Media/MediaFrameworkUtilities/Source/MediaFrameworkUtilitiesEditor/Private/MediaFrameworkUtilitiesEditorModule.cpp b/Engine/Plugins/Media/MediaFrameworkUtilities/Source/MediaFrameworkUtilitiesEditor/Private/MediaFrameworkUtilitiesEditorModule.cpp
--- a/Engine/Plugins/Media/MediaFrameworkUtilities/Source/MediaFrameworkUtilitiesEditor/Private/MediaFrameworkUtilitiesEditorModule.cpp
+++ b/Engine/Plugins/Media/MediaFrameworkUtilities/Source/MediaFrameworkUtilitiesEditor/Private/MediaFrameworkUtilitiesEditorModule.cpp
@@ -92,11 +92,11 @@ public:
 				FLevelEditorModule* LevelEditorModule = FModuleManager::GetModulePtr<FLevelEditorModule>("LevelEditor");
 				if (LevelEditorModule != nullptr)
 				{
-					FLevelEditorModule::FStatusBarItem Item;
+					FLevelEditorModule::FTitleBarItem Item;
 					Item.Label = LOCTEXT("MediaProfileLabel", "MediaProfile: ");
 					Item.Value = MakeAttributeLambda([]() { UObject* MediaProfile = UMediaProfileBlueprintLibrary::GetMediaProfile(); return MediaProfile ? FText::FromName(MediaProfile->GetFName()) : FText::GetEmpty(); });
 					Item.Visibility = MakeAttributeLambda([]() { return GetDefault<UMediaProfileEditorSettings>()->bDisplayInMainEditor ? EVisibility::SelfHitTestInvisible : EVisibility::Collapsed; });
-					LevelEditorModule->AddStatusBarItem(NotificationBarIdentifier, Item);
+					LevelEditorModule->AddTitleBarItem(NotificationBarIdentifier, Item);
 				}
 			}
 		}
@@ -109,7 +109,7 @@ public:
 			FLevelEditorModule* LevelEditorModule = FModuleManager::GetModulePtr<FLevelEditorModule>("LevelEditor");
 			if (LevelEditorModule != nullptr)
 			{
-				LevelEditorModule->RemoveStatusBarItem(NotificationBarIdentifier);
+				LevelEditorModule->RemoveTitleBarItem(NotificationBarIdentifier);
 			}
 
 			FMediaProfileMenuEntry::Unregister();
diff --git a/Engine/Source/Editor/EditorStyle/Private/StarshipStyle.cpp b/Engine/Source/Editor/EditorStyle/Private/StarshipStyle.cpp
--- a/Engine/Source/Editor/EditorStyle/Private/StarshipStyle.cpp
+++ b/Engine/Source/Editor/EditorStyle/Private/StarshipStyle.cpp
@@ -245,6 +245,7 @@ void FStarshipEditorStyle::FStyle::Initialize()
 	SetupAutomationStyles();
 	SetupUMGEditorStyles();
 	SetupMyBlueprintStyles();
+	SetupStatusBarStyle();
 
 //	LogUnusedBrushResources();
 
@@ -4340,17 +4341,6 @@ void FStarshipEditorStyle::FStyle::SetupLevelEditorStyle()
 		Set( "LevelEditor.Recompile", new IMAGE_BRUSH( "Icons/icon_compile_40x", Icon40x40 ) );
 		Set( "LevelEditor.Recompile.Small", new IMAGE_BRUSH( "Icons/icon_compile_40x", Icon20x20 ) );
 
-		Set("LevelEditor.SourceControl", new IMAGE_BRUSH("Icons/icon_source_control_40x", Icon40x40));
-		Set("LevelEditor.SourceControl.Small", new IMAGE_BRUSH("Icons/icon_source_control_40x", Icon20x20));
-		Set("LevelEditor.SourceControl.On", new IMAGE_BRUSH("Icons/icon_source_control_40x_on", Icon40x40));
-		Set("LevelEditor.SourceControl.On.Small", new IMAGE_BRUSH("Icons/icon_source_control_40x_on", Icon20x20));
-		Set("LevelEditor.SourceControl.Off", new IMAGE_BRUSH("Icons/icon_source_control_40x_off", Icon40x40));
-		Set("LevelEditor.SourceControl.Off.Small", new IMAGE_BRUSH("Icons/icon_source_control_40x_off", Icon20x20));
-		Set("LevelEditor.SourceControl.Unknown", new IMAGE_BRUSH("Icons/icon_source_control_40x_unknown", Icon40x40));
-		Set("LevelEditor.SourceControl.Unknown.Small", new IMAGE_BRUSH("Icons/icon_source_control_40x_unknown", Icon20x20));
-		Set("LevelEditor.SourceControl.Problem", new IMAGE_BRUSH("Icons/icon_source_control_40x_problem", Icon40x40));
-		Set("LevelEditor.SourceControl.Problem.Small", new IMAGE_BRUSH("Icons/icon_source_control_40x_problem", Icon20x20));
-
 		Set("LevelEditor.PreviewMode.Enabled", new IMAGE_BRUSH("Icons/icon_PreviewMode_SM5_Enabled_40x", Icon40x40));
 		Set("LevelEditor.PreviewMode.Enabled.Small", new IMAGE_BRUSH("Icons/icon_PreviewMode_SM5_Enabled_40x", Icon20x20));
 		Set("LevelEditor.PreviewMode.Disabled", new IMAGE_BRUSH("Icons/icon_PreviewMode_SM5_Disabled_40x", Icon40x40));
@@ -7331,6 +7321,11 @@ void FStarshipEditorStyle::FStyle::SetupSourceControlStyles()
 	//Source Control
 #if WITH_EDITOR || (IS_PROGRAM && WITH_UNREAL_DEVELOPER_TOOLS)
 	{
+		Set("SourceControl.StatusIcon.On", new CORE_IMAGE_BRUSH_SVG("Starship/Common/check-circle", Icon16x16, FStyleColors::AccentGreen));
+		Set("SourceControl.StatusIcon.Error", new CORE_IMAGE_BRUSH_SVG("Starship/Common/alert-circle", Icon16x16, FStyleColors::AccentYellow));
+		Set("SourceControl.StatusIcon.Off", new CORE_IMAGE_BRUSH_SVG("Starship/Common/reject", Icon16x16, FStyleColors::Foreground));
+		Set("SourceControl.StatusIcon.Unknown", new CORE_IMAGE_BRUSH_SVG("Starship/Common/help", Icon16x16, FStyleColors::AccentYellow));
+
 		Set( "SourceControl.Add", new IMAGE_BRUSH( "Old/SourceControl/SCC_Action_Add",Icon10x10));
 		Set( "SourceControl.Edit", new IMAGE_BRUSH( "Old/SourceControl/SCC_Action_Edit",Icon10x10));
 		Set( "SourceControl.Delete", new IMAGE_BRUSH( "Old/SourceControl/SCC_Action_Delete",Icon10x10));
@@ -7338,10 +7333,7 @@ void FStarshipEditorStyle::FStyle::SetupSourceControlStyles()
 		Set( "SourceControl.Integrate", new IMAGE_BRUSH( "Old/SourceControl/SCC_Action_Integrate",Icon10x10));
 		Set( "SourceControl.Settings.StatusBorder", new BOX_BRUSH( "Common/GroupBorder", FMargin(4.0f/16.0f), FLinearColor(0.5f,0.5f,0.5f,1.0f)  ) );
 		Set( "SourceControl.Settings.StatusFont", FTextBlockStyle(NormalText).SetFont(DEFAULT_FONT( "Bold", 12 ) ));
-		Set( "SourceControl.StatusIcon.On", new IMAGE_BRUSH( "Icons/SourceControlOn_16x", Icon16x16 ) );
-		Set( "SourceControl.StatusIcon.Error", new IMAGE_BRUSH( "Icons/SourceControlProblem_16x", Icon16x16 ) );
-		Set( "SourceControl.StatusIcon.Off", new IMAGE_BRUSH( "Icons/SourceControlOff_16x", Icon16x16 ) );
-		Set( "SourceControl.StatusIcon.Unknown", new IMAGE_BRUSH( "Icons/SourceControlUnknown_16x", Icon16x16 ) );
+
 		Set( "SourceControl.LoginWindow.Font", DEFAULT_FONT( "Regular", 8 ) );
 		Set( "SourceControl.ProgressWindow.Warning", new IMAGE_BRUSH( "Icons/alert", Icon32x32) );
 
@@ -7891,6 +7883,68 @@ void FStarshipEditorStyle::FStyle::SetupMyBlueprintStyles()
 	Set("MyBlueprint.DeleteEntry", new IMAGE_BRUSH("Icons/Edit/icon_Edit_Delete_40x", Icon16x16));
 }
 
+void FStarshipEditorStyle::FStyle::SetupStatusBarStyle()
+{
+	const FButtonStyle StatusBarButton = FButtonStyle(FStarshipCoreStyle::GetCoreStyle().GetWidgetStyle<FButtonStyle>("NoBorder"))
+		.SetNormalForeground(FStyleColors::Foreground)
+		.SetNormalPadding(FMargin(2, 2, 2, 2))
+		.SetPressedPadding(FMargin(2, 3, 2, 1));
+
+	const FComboButtonStyle StatusBarComboButton = FComboButtonStyle(FStarshipCoreStyle::GetCoreStyle().GetWidgetStyle<FComboButtonStyle>("ComboButton"))
+		.SetDownArrowImage(FSlateNoResource())
+		.SetButtonStyle(StatusBarButton);
+		
+	Set("StatusBar.StatusBarButton", StatusBarButton);
+	Set("StatusBar.StatusBarComboButton", StatusBarComboButton);
+
+	Set("StatusBar.ContentBrowserIcon", new CORE_IMAGE_BRUSH_SVG("Starship/Common/tile", Icon16x16));
+	Set("StatusBar.ContentBrowserUp", new CORE_IMAGE_BRUSH_SVG("Starship/Common/chevron-up", Icon16x16));
+	Set("StatusBar.ContentBrowserDown", new CORE_IMAGE_BRUSH_SVG("Starship/Common/chevron-down", Icon16x16));
+	Set("StatusBar.Height", 32.0f);
+	Set("StatusBar.Background", new FSlateColorBrush(FStyleColors::Background));
+
+	const FButtonStyle SlimToolBarButton = FButtonStyle(FStarshipCoreStyle::GetCoreStyle().GetWidgetStyle<FButtonStyle>("SlimToolBar.Button"));
+
+	Set("StatusBarToolBar.Background", new FSlateNoResource());
+	Set("StatusBarToolBar.Expand", new IMAGE_BRUSH("Icons/toolbar_expand_16x", Icon16x16));
+	Set("StatusBarToolBar.SubMenuIndicator", new IMAGE_BRUSH("Common/SubmenuArrow", Icon8x8));
+	Set("StatusBarToolBar.SToolBarComboButtonBlock.Padding", FMargin(6.0f, 0.0f));
+	Set("StatusBarToolBar.SToolBarButtonBlock.Padding", FMargin(4.0f, 0.0f));
+	Set("StatusBarToolBar.SToolBarButtonBlock.CheckBox.Padding", FMargin(10.0f, 0.0f));
+
+	Set("StatusBarToolBar.Separator", new FSlateColorBrush(FStyleColors::Input));
+	Set("StatusBarToolBar.Separator.Padding", FMargin(8.f, 0));
+
+	Set("StatusBarToolBar.Label", FTextBlockStyle(NormalText).SetFont(DEFAULT_FONT("Regular", 10)));
+	Set("StatusBarToolBar.Label.Padding", FMargin(5, 5, 0, 5));
+	Set("StatusBarToolBar.Label.DropShadowSize", FVector2D::ZeroVector);
+
+	Set("StatusBarToolBar.EditableText", FEditableTextBoxStyle(NormalEditableTextBoxStyle).SetFont(DEFAULT_FONT("Regular", 9)));
+
+	//	Style->Set("SlimToolBar.CheckBox", ToolBarCheckBoxStyle);
+
+	const FCheckBoxStyle SlimToolBarToggleButtonCheckBoxStyle = FCheckBoxStyle()
+		.SetCheckBoxType(ESlateCheckBoxType::ToggleButton)
+		.SetUncheckedImage(FSlateNoResource())
+		.SetUncheckedPressedImage(FSlateNoResource())
+		.SetUncheckedHoveredImage(FSlateNoResource())
+		.SetCheckedImage(FSlateNoResource())
+		.SetCheckedHoveredImage(FSlateNoResource())
+		.SetCheckedPressedImage(FSlateNoResource())
+		.SetForegroundColor(FStyleColors::Foreground)
+		.SetPressedForegroundColor(FStyleColors::ForegroundHover)
+		.SetHoveredForegroundColor(FStyleColors::ForegroundHover)
+		.SetCheckedForegroundColor(FStyleColors::Primary)
+		.SetCheckedPressedForegroundColor(FStyleColors::PrimaryPress)
+		.SetCheckedHoveredForegroundColor(FStyleColors::PrimaryHover);
+
+	Set("StatusBarToolBar.ToggleButton", SlimToolBarToggleButtonCheckBoxStyle);
+	Set("StatusBarToolBar.Button", SlimToolBarButton);
+
+	Set("StatusBarToolBar.SimpleComboButton", FStarshipCoreStyle::GetCoreStyle().GetWidgetStyle<FComboButtonStyle>("StatusBarToolBar.SimpleComboButton"));
+	Set("StatusBarToolBar.IconSize", Icon16x16);
+}
+
 #undef IMAGE_BRUSH
 #undef BOX_BRUSH
 #undef BORDER_BRUSH
diff --git a/Engine/Source/Editor/EditorStyle/Private/StarshipStyle.h b/Engine/Source/Editor/EditorStyle/Private/StarshipStyle.h
--- a/Engine/Source/Editor/EditorStyle/Private/StarshipStyle.h
+++ b/Engine/Source/Editor/EditorStyle/Private/StarshipStyle.h
@@ -53,8 +53,7 @@ public:
 			SettingsModule->UnregisterSettings("Editor", "General", "Appearance");
 		}
 
-		ResetToDefault();
-		ensure(StyleInstance.IsUnique());
+		//FEditorStyle::SetStyle()
 		StyleInstance.Reset();
 	}
 
@@ -98,6 +97,7 @@ public:
 		void SetupAutomationStyles();
 		void SetupUMGEditorStyles();
 		void SetupMyBlueprintStyles();
+		void SetupStatusBarStyle();
 
 		void SettingsChanged(UObject* ChangedObject, FPropertyChangedEvent& PropertyChangedEvent);
 		void SyncSettings();
diff --git a/Engine/Source/Editor/LevelEditor/Private/LevelEditor.cpp b/Engine/Source/Editor/LevelEditor/Private/LevelEditor.cpp
--- a/Engine/Source/Editor/LevelEditor/Private/LevelEditor.cpp
+++ b/Engine/Source/Editor/LevelEditor/Private/LevelEditor.cpp
@@ -28,7 +28,6 @@
 #include "Toolkits/GlobalEditorCommonCommands.h"
 #include "ISlateReflectorModule.h"
 #include "Widgets/Docking/SDockTab.h"
-#include "IIntroTutorials.h"
 #include "Interfaces/IProjectManager.h"
 #include "LevelViewportLayoutEntity.h"
 #include "PixelInspectorModule.h"
@@ -55,7 +54,6 @@ const FName LevelEditorTabIds::LevelEditorViewport_Clone1(TEXT("LevelEditorViewp
 const FName LevelEditorTabIds::LevelEditorViewport_Clone2(TEXT("LevelEditorViewport_Clone2"));
 const FName LevelEditorTabIds::LevelEditorViewport_Clone3(TEXT("LevelEditorViewport_Clone3"));
 const FName LevelEditorTabIds::LevelEditorViewport_Clone4(TEXT("LevelEditorViewport_Clone4"));
-const FName LevelEditorTabIds::LevelEditorToolBar(TEXT("LevelEditorToolBar"));
 const FName LevelEditorTabIds::LevelEditorToolBox(TEXT("LevelEditorToolBox"));
 const FName LevelEditorTabIds::LevelEditorSelectionDetails(TEXT("LevelEditorSelectionDetails"));
 const FName LevelEditorTabIds::LevelEditorSelectionDetails2(TEXT("LevelEditorSelectionDetails2"));
@@ -95,8 +93,13 @@ public:
 		FString ProjectNameWatermarkPrefix;
 		GConfig->GetString(TEXT("LevelEditor"), TEXT("ProjectNameWatermarkPrefix"), /*out*/ ProjectNameWatermarkPrefix, GEditorPerProjectIni);
 
-		FColor BadgeBackgroundColor = FColor::Black;
-		GConfig->GetColor(TEXT("LevelEditor"), TEXT("ProjectBadgeBackgroundColor"), /*out*/ BadgeBackgroundColor, GEditorPerProjectIni);
+		FLinearColor BadgeBackgroundColor = FAppStyle::Get().GetColor("Colors.Foldout");
+
+		FColor ConfigColor;
+		if (GConfig->GetColor(TEXT("LevelEditor"), TEXT("ProjectBadgeBackgroundColor"), /*out*/ ConfigColor, GEditorPerProjectIni))
+		{
+			BadgeBackgroundColor = ConfigColor;
+		}
 
 		FColor BadgeTextColor = FColor(128,128,128,255);
 		GConfig->GetColor(TEXT("LevelEditor"), TEXT("ProjectBadgeTextColor"), /*out*/ BadgeTextColor, GEditorPerProjectIni);
@@ -150,7 +153,7 @@ public:
 			[
 				SNew(SBorder)
 				.BorderImage(FEditorStyle::GetBrush("SProjectBadge.BadgeShape"))
-				.Padding(FMargin(10.0f, 2.5f))
+				.Padding(FAppStyle::Get().GetMargin("SProjectBadge.BadgePadding"))
 				.BorderBackgroundColor(BadgeBackgroundColor)
 				.VAlign(VAlign_Top)
 				[
@@ -177,10 +180,7 @@ private:
 	FGeometry CachedGeometry;
 };
 
-static FMargin GetRoomForBadge(TWeakPtr<SProjectBadge> ProjBadge)
-{
-	return FMargin(8.0f, 0.0f, ProjBadge.Pin()->GetSizeLastFrame().X + 8.0f, 0.0f);
-}
+
 
 TSharedRef<SDockTab> FLevelEditorModule::SpawnLevelEditor( const FSpawnTabArgs& InArgs )
 {
@@ -197,12 +197,13 @@ TSharedRef<SDockTab> FLevelEditorModule::SpawnLevelEditor( const FSpawnTabArgs&
 		OwnerWindow = MainFrameModule.GetParentWindow();
 	}
 
-	if ( OwnerWindow.IsValid() )
+	TSharedPtr<SLevelEditor> LevelEditorTmp;
+	if (OwnerWindow.IsValid())
 	{
-		TSharedPtr<SLevelEditor> LevelEditorTmp;
-		LevelEditorTab->SetContent( SAssignNew(LevelEditorTmp, SLevelEditor ) );
+
+		LevelEditorTab->SetContent(SAssignNew(LevelEditorTmp, SLevelEditor));
 		SetLevelEditorInstance(LevelEditorTmp);
-		LevelEditorTmp->Initialize( LevelEditorTab, OwnerWindow.ToSharedRef() );
+		LevelEditorTmp->Initialize(LevelEditorTab, OwnerWindow.ToSharedRef());
 
 		if (GetDefault<UEditorStyleSettings>()->bEnableLegacyEditorModeUI)
 		{
@@ -225,46 +226,33 @@ TSharedRef<SDockTab> FLevelEditorModule::SpawnLevelEditor( const FSpawnTabArgs&
 			if (PlacementBrowserTab.IsValid())
 			{
 				PlacementBrowserTab->RequestCloseTab();
+
 			}
 		}
 
 		LevelEditorCreatedEvent.Broadcast(LevelEditorTmp);
-	}
 
-	IIntroTutorials& IntroTutorials = FModuleManager::LoadModuleChecked<IIntroTutorials>(TEXT("IntroTutorials"));
-	TSharedRef<SWidget> TutorialWidget = IntroTutorials.CreateTutorialsWidget(TEXT("LevelEditor"), OwnerWindow);
+		TSharedRef<SProjectBadge> ProjectBadge = SNew(SProjectBadge);
 
-	TSharedRef<SProjectBadge> ProjectBadge = SNew(SProjectBadge);
-	TAttribute<FMargin> BadgeSizeGetter = TAttribute<FMargin>::Create(TAttribute<FMargin>::FGetter::CreateStatic(&GetRoomForBadge, TWeakPtr<SProjectBadge>(ProjectBadge)));
-		
-	TSharedPtr< SWidget > RightContent=
-		SNew( SHorizontalBox )
-
-// Put the level editor stats/notification widgets on the main window title bar since we don't have a menu bar on OS X
-#if PLATFORM_MAC
-		+SHorizontalBox::Slot()
-		.AutoWidth()
-		.Padding(8.0f, 0.0f, 0.0f, 0.0f)
-		.VAlign(VAlign_Center)
-		[
-			LevelEditorTab->GetRightContent()
-		]
-#endif
-
-		+ SHorizontalBox::Slot()
-		.AutoWidth()
-		.Padding(BadgeSizeGetter)
-		.VAlign(VAlign_Center)
-		[
-			TutorialWidget
-		];
+		TSharedPtr< SWidget > RightContent =
+			SNew(SHorizontalBox)
+			+ SHorizontalBox::Slot()
+			.AutoWidth()
+			.Padding(8.0f, 0.0f)
+			.VAlign(VAlign_Center)
+			[
+				LevelEditorTmp->GetTitleBarMessageWidget()
+			]
+			+ SHorizontalBox::Slot()
+			.AutoWidth()
+			.Padding(0.0f, 0.0f, 8.0f, 0.0f)
+			[
+				ProjectBadge
+			];
 
-	LevelEditorTab->SetRightContent( RightContent.ToSharedRef() );
 
-	LevelEditorTab->SetBackgroundContent(
-		ProjectBadge
-	);
-	
+		LevelEditorTab->SetTitleBarRightContent(RightContent.ToSharedRef());
+	}
 	
 	return LevelEditorTab;
 }
@@ -280,13 +268,10 @@ void FLevelEditorModule::StartupModule()
 
 	FModuleManager::LoadModuleChecked<FCommonMenuExtensionsModule>(CommonMenuExtensionsName);
 
-	MenuExtensibilityManager = MakeShareable(new FExtensibilityManager);
-	
-	ToolBarExtensibilityManager = MakeShareable(new FExtensibilityManager);
-
-	ModeBarExtensibilityManager = MakeShareable(new FExtensibilityManager);
-
-	NotificationBarExtensibilityManager = MakeShareable(new FExtensibilityManager);
+	MenuExtensibilityManager = MakeShared<FExtensibilityManager>();
+	ToolBarExtensibilityManager = MakeShared<FExtensibilityManager>();
+	ModeBarExtensibilityManager = MakeShared<FExtensibilityManager>();
+	NotificationBarExtensibilityManager = MakeShared<FExtensibilityManager>();
 
 	// Note this must come before any tab spawning because that can create the SLevelEditor and attempt to map commands
 	FEditorViewportCommands::Register();
@@ -519,6 +504,7 @@ void FLevelEditorModule::SetLevelEditorTabManager( const TSharedPtr<SDockTab>& O
 	{
 		LevelEditorTabManager = FGlobalTabmanager::Get()->NewTabManager(OwnerTab.ToSharedRef());
 		LevelEditorTabManager->SetOnPersistLayout(FTabManager::FOnPersistLayout::CreateRaw(this, &FLevelEditorModule::HandleTabManagerPersistLayout));
+		LevelEditorTabManager->SetAllowWindowMenuBar(true);
 
 		TabManagerChangedEvent.Broadcast();
 	}
@@ -606,16 +592,16 @@ TSharedPtr<SDockTab> FLevelEditorModule::GetLevelEditorTab() const
 	return LevelEditorInstanceTabPtr.Pin();
 }
 
-void FLevelEditorModule::AddStatusBarItem(FName InStatusBarIdentifier, const FStatusBarItem& InStatusBarItem)
+void FLevelEditorModule::AddTitleBarItem(FName InTitleBarIdentifier, const FTitleBarItem& InTitleBarItem)
 {
-	StatusBarItems.FindOrAdd(InStatusBarIdentifier) = InStatusBarItem;
-	BroadcastNotificationBarChanged();
+	TitleBarItems.FindOrAdd(InTitleBarIdentifier) = InTitleBarItem;
+	BroadcastTitleBarMessagesChanged();
 }
 
-void FLevelEditorModule::RemoveStatusBarItem(FName InStatusBarIdentifier)
+void FLevelEditorModule::RemoveTitleBarItem(FName InTitleBarIdentifier)
 {
-	StatusBarItems.Remove(InStatusBarIdentifier);
-	BroadcastNotificationBarChanged();
+	TitleBarItems.Remove(InTitleBarIdentifier);
+	BroadcastTitleBarMessagesChanged();
 }
 
 TSharedRef<ILevelViewportLayoutEntity> FLevelEditorModule::FactoryViewport(FName InTypeName, const FViewportConstructionArgs& ConstructionArgs) const
@@ -695,28 +681,7 @@ void FLevelEditorModule::BindGlobalLevelEditorCommands()
 		FExecuteAction::CreateStatic( &FLevelEditorActionCallbacks::Build_Execute ),
 		FCanExecuteAction::CreateStatic( &FLevelEditorActionCallbacks::Build_CanExecute ) );
 
-	ActionList.MapAction(
-		Commands.ConnectToSourceControl,
-		FExecuteAction::CreateStatic(&FLevelEditorActionCallbacks::ConnectToSourceControl_Clicked)
-		);
-
-	ActionList.MapAction(
-		Commands.ChangeSourceControlSettings,
-		FExecuteAction::CreateStatic(&FLevelEditorActionCallbacks::ConnectToSourceControl_Clicked)
-		);
-
-	ActionList.MapAction(
-		Commands.CheckOutModifiedFiles,
-		FExecuteAction::CreateStatic(&FLevelEditorActionCallbacks::CheckOutModifiedFiles_Clicked),
-		FCanExecuteAction::CreateStatic(&FLevelEditorActionCallbacks::CheckOutModifiedFiles_CanExecute)
-		);
-
-	ActionList.MapAction(
-		Commands.SubmitToSourceControl,
-		FExecuteAction::CreateStatic(&FLevelEditorActionCallbacks::SubmitToSourceControl_Clicked),
-		FCanExecuteAction::CreateStatic(&FLevelEditorActionCallbacks::SubmitToSourceControl_CanExecute)
-		);
-
+	
 	ActionList.MapAction(Commands.RecompileGameCode,
 		FExecuteAction::CreateStatic( &FLevelEditorActionCallbacks::RecompileGameCode_Clicked ),
 		FCanExecuteAction::CreateStatic( &FLevelEditorActionCallbacks::Recompile_CanExecute )
diff --git a/Engine/Source/Editor/LevelEditor/Private/LevelEditorActions.cpp b/Engine/Source/Editor/LevelEditor/Private/LevelEditorActions.cpp
--- a/Engine/Source/Editor/LevelEditor/Private/LevelEditorActions.cpp
+++ b/Engine/Source/Editor/LevelEditor/Private/LevelEditorActions.cpp
@@ -1142,53 +1142,6 @@ void FLevelEditorActionCallbacks::LiveCoding_Settings_Clicked()
 }
 #endif
 
-void FLevelEditorActionCallbacks::ConnectToSourceControl_Clicked()
-{
-	// Show login window regardless of current status - its useful as a shortcut to change settings.
-	ISourceControlModule& SourceControlModule = ISourceControlModule::Get();
-	SourceControlModule.ShowLoginDialog(FSourceControlLoginClosed(), ELoginWindowMode::Modeless, EOnLoginWindowStartup::PreserveProvider);
-}
-
-bool FLevelEditorActionCallbacks::CheckOutModifiedFiles_CanExecute()
-{
-	ISourceControlModule& SourceControlModule = ISourceControlModule::Get();
-	if (ISourceControlModule::Get().IsEnabled() &&
-		ISourceControlModule::Get().GetProvider().IsAvailable())
-	{
-		TArray<UPackage*> PackagesToSave;
-		FEditorFileUtils::GetDirtyWorldPackages(PackagesToSave);
-		FEditorFileUtils::GetDirtyContentPackages(PackagesToSave);
-
-		return PackagesToSave.Num() > 0;
-	}
-
-	return false;
-}
-
-void FLevelEditorActionCallbacks::CheckOutModifiedFiles_Clicked()
-{
-	TArray<UPackage*> PackagesToSave;
-	FEditorFileUtils::GetDirtyWorldPackages(PackagesToSave);
-	FEditorFileUtils::GetDirtyContentPackages(PackagesToSave);
-
-	const bool bCheckDirty = true;
-	const bool bPromptUserToSave = false;
-	FEditorFileUtils::PromptForCheckoutAndSave(PackagesToSave, bCheckDirty, bPromptUserToSave);
-}
-
-bool FLevelEditorActionCallbacks::SubmitToSourceControl_CanExecute()
-{
-	ISourceControlModule& SourceControlModule = ISourceControlModule::Get();
-	return ISourceControlModule::Get().IsEnabled() &&
-		   ISourceControlModule::Get().GetProvider().IsAvailable() &&
-		   FSourceControlWindows::CanChoosePackagesToCheckIn();
-}
-
-void FLevelEditorActionCallbacks::SubmitToSourceControl_Clicked()
-{
-	FSourceControlWindows::ChoosePackagesToCheckIn();
-}
-
 
 void FLevelEditorActionCallbacks::GoToCodeForActor_Clicked()
 {
@@ -3347,11 +3300,6 @@ void FLevelEditorCommands::RegisterCommands()
 	UI_COMMAND(PreviewPlatformOverride_AndroidVulkanSM5, "Android Vulkan SM5", "Mobile preview using Android Vulkan SM5 quality settings.", EUserInterfaceActionType::Check, FInputChord());
 	UI_COMMAND(PreviewPlatformOverride_IOSMetalES31, "iOS", "Mobile preview using iOS material quality settings.", EUserInterfaceActionType::Check, FInputChord());
 
-	UI_COMMAND( ConnectToSourceControl, "Connect to Source Control...", "Opens a dialog to connect to source control.", EUserInterfaceActionType::Button, FInputChord());
-	UI_COMMAND( ChangeSourceControlSettings, "Change Source Control Settings...", "Opens a dialog to change source control settings.", EUserInterfaceActionType::Button, FInputChord());
-	UI_COMMAND( CheckOutModifiedFiles, "Check Out Modified Files...", "Opens a dialog to check out any assets which have been modified.", EUserInterfaceActionType::Button, FInputChord());
-	UI_COMMAND( SubmitToSourceControl, "Submit to Source Control...", "Opens a dialog with check in options for content and levels.", EUserInterfaceActionType::Button, FInputChord());
-
 	UI_COMMAND(GeometryCollectionSelectAllGeometry, "Select All Geometry In Hierarchy", "Select all geometry in hierarchy", EUserInterfaceActionType::Button, FInputChord());
 	UI_COMMAND(GeometryCollectionSelectNone, "Deselect All Geometry In Hierarchy", "Deselect all geometry in hierarchy", EUserInterfaceActionType::Button, FInputChord());
 	UI_COMMAND(GeometryCollectionSelectInverseGeometry, "Select Inverse Geometry In Hierarchy", "Select inverse geometry in hierarchy", EUserInterfaceActionType::Button, FInputChord());
diff --git a/Engine/Source/Editor/LevelEditor/Private/LevelEditorMenu.cpp b/Engine/Source/Editor/LevelEditor/Private/LevelEditorMenu.cpp
--- a/Engine/Source/Editor/LevelEditor/Private/LevelEditorMenu.cpp
+++ b/Engine/Source/Editor/LevelEditor/Private/LevelEditorMenu.cpp
@@ -16,6 +16,10 @@
 #include "Framework/Commands/GenericCommands.h"
 #include "IDocumentation.h"
 #include "ToolMenus.h"
+#include "LevelEditorMenuContext.h"
+#include "Widgets/Input/SSpinBox.h"
+
+#include "Widgets/Input/SCheckBox.h"
 
 #define LOCTEXT_NAMESPACE "LevelEditorMenu"
 
@@ -184,18 +188,20 @@ void FLevelEditorMenu::RegisterLevelEditorMenus()
 		static void ExtendEditMenu()
 		{
 			UToolMenu* Menu = UToolMenus::Get()->RegisterMenu("LevelEditor.MainMenu.Edit", "MainFrame.MainMenu.Edit");
-			// Edit Actor
 			{
-				FToolMenuSection& Section = Menu->AddSection("EditMain", LOCTEXT("MainHeading", "Edit"), FToolMenuInsert("EditHistory", EToolMenuInsertType::After));
-				{		
-					Section.AddMenuEntry( FGenericCommands::Get().Cut );
-					Section.AddMenuEntry( FGenericCommands::Get().Copy );
-					Section.AddMenuEntry( FGenericCommands::Get().Paste );
-
-					Section.AddMenuEntry( FGenericCommands::Get().Duplicate );
-					Section.AddMenuEntry( FGenericCommands::Get().Delete );
+				// Edit Actor
+				{
+					FToolMenuSection& Section = Menu->AddSection("EditMain", LOCTEXT("MainHeading", "Edit"), FToolMenuInsert("EditHistory", EToolMenuInsertType::After));
+
+					Section.AddMenuEntry(FGenericCommands::Get().Cut);
+					Section.AddMenuEntry(FGenericCommands::Get().Copy);
+					Section.AddMenuEntry(FGenericCommands::Get().Paste);
+
+					Section.AddMenuEntry(FGenericCommands::Get().Duplicate);
+					Section.AddMenuEntry(FGenericCommands::Get().Delete);
 				}
 			}
+
 		}
 
 		static void ExtendHelpMenu()
@@ -214,6 +220,23 @@ void FLevelEditorMenu::RegisterLevelEditorMenus()
 				Section.AddMenuEntry( FLevelEditorCommands::Get().BrowseViewportControls );
 			}
 		}
+
+		static void ExtendMenuBar()
+		{
+			UToolMenu* Menu = UToolMenus::Get()->ExtendMenu("LevelEditor.MainMenu");
+
+			FToolMenuSection& Section = Menu->FindOrAddSection(NAME_None);
+
+			FToolMenuEntry& Entry =
+				Section.AddSubMenu(
+					"Build",
+					LOCTEXT("BuildMenu", "Build"),
+					LOCTEXT("BuildMenu_ToolTip", "Level Build Options"),
+					FNewToolMenuChoice()
+				);
+
+			Entry.InsertPosition = FToolMenuInsert("Help", EToolMenuInsertType::Before);
+		}
 	};
 
 	UToolMenus* ToolMenus = UToolMenus::Get();
@@ -221,6 +244,8 @@ void FLevelEditorMenu::RegisterLevelEditorMenus()
 	ToolMenus->RegisterMenu("LevelEditor.MainMenu.File", "MainFrame.MainTabMenu.File");
 	ToolMenus->RegisterMenu("LevelEditor.MainMenu.Window", "MainFrame.MainMenu.Window");
 
+	Local::ExtendMenuBar();
+
 	// Add level loading and saving menu items
 	Local::RegisterFileLoadAndSaveItems();
 
@@ -232,6 +257,9 @@ void FLevelEditorMenu::RegisterLevelEditorMenus()
 
 	// Extend the Help menu
 	Local::ExtendHelpMenu();
+
+	RegisterBuildMenu();
+
 }
 
 TSharedRef< SWidget > FLevelEditorMenu::MakeLevelEditorMenu( const TSharedPtr<FUICommandList>& CommandList, TSharedPtr<class SLevelEditor> LevelEditor )
@@ -241,7 +269,7 @@ TSharedRef< SWidget > FLevelEditorMenu::MakeLevelEditorMenu( const TSharedPtr<FU
 	FToolMenuContext ToolMenuContext(CommandList, Extenders.ToSharedRef());
 
 	IMainFrameModule& MainFrameModule = FModuleManager::LoadModuleChecked<IMainFrameModule>( "MainFrame" );
-	TSharedRef< SWidget > MenuBarWidget = MainFrameModule.MakeMainTabMenu( LevelEditor->GetTabManager(), "LevelEditor.MainMenu", ToolMenuContext );
+	TSharedRef< SWidget > MenuBarWidget = MainFrameModule.MakeMainMenu( LevelEditor->GetTabManager(), "LevelEditor.MainMenu", ToolMenuContext );
 
 	return MenuBarWidget;
 }
@@ -251,7 +279,7 @@ TSharedRef< SWidget > FLevelEditorMenu::MakeNotificationBar( const TSharedPtr<FU
 	FLevelEditorModule& LevelEditorModule = FModuleManager::GetModuleChecked<FLevelEditorModule>( "LevelEditor");
 	const TSharedPtr<FExtender> NotificationBarExtenders = LevelEditorModule.GetNotificationBarExtensibilityManager()->GetAllExtenders();
 
-	FToolBarBuilder NotificationBarBuilder( CommandList, FMultiBoxCustomization::None, NotificationBarExtenders, Orient_Horizontal );
+	FToolBarBuilder NotificationBarBuilder( CommandList, FMultiBoxCustomization::None, NotificationBarExtenders );
 	NotificationBarBuilder.SetStyle(&FEditorStyle::Get(), "NotificationBar");
 	{
 		NotificationBarBuilder.BeginSection("Start");
@@ -261,4 +289,321 @@ TSharedRef< SWidget > FLevelEditorMenu::MakeNotificationBar( const TSharedPtr<FU
 	return NotificationBarBuilder.MakeWidget();
 }
 
+void FLevelEditorMenu::RegisterBuildMenu()
+{
+	static const FName BaseMenuName = "LevelEditor.MainMenu.Build";
+	UToolMenu* Menu = UToolMenus::Get()->RegisterMenu(BaseMenuName);
+
+	struct FLightingMenus
+	{
+	public:
+
+		static void RegisterMenus(const FName InBaseMenuName)
+		{
+			FLightingMenus::RegisterLightingQualityMenu(InBaseMenuName);
+			FLightingMenus::RegisterLightingInfoMenu(InBaseMenuName);
+		}
+
+	private:
+
+		/** Generates a lighting quality sub-menu */
+		static void RegisterLightingQualityMenu(const FName InBaseMenuName)
+		{
+			UToolMenu* SubMenu = UToolMenus::Get()->RegisterMenu(UToolMenus::JoinMenuPaths(InBaseMenuName, "LightingQuality"));
+
+			{
+				FToolMenuSection& Section = SubMenu->AddSection("LevelEditorBuildLightingQuality", LOCTEXT("LightingQualityHeading", "Quality Level"));
+				Section.AddMenuEntry(FLevelEditorCommands::Get().LightingQuality_Production);
+				Section.AddMenuEntry(FLevelEditorCommands::Get().LightingQuality_High);
+				Section.AddMenuEntry(FLevelEditorCommands::Get().LightingQuality_Medium);
+				Section.AddMenuEntry(FLevelEditorCommands::Get().LightingQuality_Preview);
+			}
+		}
+
+		/** Generates a lighting density sub-menu */
+		static void RegisterLightingDensityMenu(const FName InBaseMenuName)
+		{
+			UToolMenu* SubMenu = UToolMenus::Get()->RegisterMenu(UToolMenus::JoinMenuPaths(InBaseMenuName, "LightingDensity"));
+
+			{
+				FToolMenuSection& Section = SubMenu->AddSection("LevelEditorBuildLightingDensity", LOCTEXT("LightingDensityHeading", "Density Rendering"));
+				TSharedRef<SWidget> Ideal = SNew(SHorizontalBox)
+					+ SHorizontalBox::Slot()
+					.Padding(FMargin(27.0f, 0.0f, 0.0f, 0.0f))
+					.FillWidth(1.0f)
+					[
+						SNew(SSpinBox<float>)
+						.MinValue(0.f)
+						.MaxValue(100.f)
+						.Value(FLevelEditorActionCallbacks::GetLightingDensityIdeal())
+						.OnValueChanged_Static(&FLevelEditorActionCallbacks::SetLightingDensityIdeal)
+					];
+
+				Section.AddEntry(FToolMenuEntry::InitWidget("Ideal", Ideal, LOCTEXT("LightingDensity_Ideal", "Ideal Density")));
+
+				TSharedRef<SWidget> Maximum = SNew(SHorizontalBox)
+					+ SHorizontalBox::Slot()
+					.FillWidth(1.0f)
+					[
+						SNew(SSpinBox<float>)
+						.MinValue(0.01f)
+						.MaxValue(100.01f)
+						.Value(FLevelEditorActionCallbacks::GetLightingDensityMaximum())
+						.OnValueChanged_Static(&FLevelEditorActionCallbacks::SetLightingDensityMaximum)
+					];
+
+				Section.AddEntry(FToolMenuEntry::InitWidget("Maximum", Maximum, LOCTEXT("LightingDensity_Maximum", "Maximum Density")));
+
+				TSharedRef<SWidget> ClrScale = SNew(SHorizontalBox)
+					+ SHorizontalBox::Slot()
+					.Padding(FMargin(35.0f, 0.0f, 0.0f, 0.0f))
+					.FillWidth(1.0f)
+					[
+						SNew(SSpinBox<float>)
+						.MinValue(0.f)
+						.MaxValue(10.f)
+						.Value(FLevelEditorActionCallbacks::GetLightingDensityColorScale())
+						.OnValueChanged_Static(&FLevelEditorActionCallbacks::SetLightingDensityColorScale)
+					];
+
+				Section.AddEntry(FToolMenuEntry::InitWidget("ColorScale", ClrScale, LOCTEXT("LightingDensity_ColorScale", "Color Scale")));
+
+				TSharedRef<SWidget> GrayScale = SNew(SHorizontalBox)
+					+ SHorizontalBox::Slot()
+					.Padding(FMargin(11.0f, 0.0f, 0.0f, 0.0f))
+					.FillWidth(1.0f)
+					[
+						SNew(SSpinBox<float>)
+						.MinValue(0.f)
+						.MaxValue(10.f)
+						.Value(FLevelEditorActionCallbacks::GetLightingDensityGrayscaleScale())
+						.OnValueChanged_Static(&FLevelEditorActionCallbacks::SetLightingDensityGrayscaleScale)
+					];
+
+				Section.AddEntry(FToolMenuEntry::InitWidget("GrayscaleScale", GrayScale, LOCTEXT("LightingDensity_GrayscaleScale", "Grayscale Scale")));
+
+				Section.AddMenuEntry(FLevelEditorCommands::Get().LightingDensity_RenderGrayscale);
+			}
+		}
+
+		/** Generates a lighting resolution sub-menu */
+		static void RegisterLightingResolutionMenu(const FName InBaseMenuName)
+		{
+			UToolMenu* SubMenu = UToolMenus::Get()->RegisterMenu(UToolMenus::JoinMenuPaths(InBaseMenuName, "LightingResolution"));
+
+			{
+				FToolMenuSection& Section = SubMenu->AddSection("LevelEditorBuildLightingResolution1", LOCTEXT("LightingResolutionHeading1", "Primitive Types"));
+				TSharedRef<SWidget> Meshes = SNew(SHorizontalBox)
+					+ SHorizontalBox::Slot()
+					.AutoWidth()
+					[
+						SNew(SCheckBox)
+						.Style(FEditorStyle::Get(), "Menu.CheckBox")
+						.ToolTipText(LOCTEXT("StaticMeshesToolTip", "Static Meshes will be adjusted if checked."))
+						.IsChecked_Static(&FLevelEditorActionCallbacks::IsLightingResolutionStaticMeshesChecked)
+						.OnCheckStateChanged_Static(&FLevelEditorActionCallbacks::SetLightingResolutionStaticMeshes)
+						.Content()
+						[
+							SNew(STextBlock)
+							.Text(LOCTEXT("StaticMeshes", "Static Meshes"))
+						]
+					]
+					+ SHorizontalBox::Slot()
+					.AutoWidth()
+					.Padding(FMargin(4.0f, 0.0f, 11.0f, 0.0f))
+					[
+						SNew(SSpinBox<float>)
+						.MinValue(4.f)
+						.MaxValue(4096.f)
+						.ToolTipText(LOCTEXT("LightingResolutionStaticMeshesMinToolTip", "The minimum lightmap resolution for static mesh adjustments. Anything outside of Min/Max range will not be touched when adjusting."))
+						.Value(FLevelEditorActionCallbacks::GetLightingResolutionMinSMs())
+						.OnValueChanged_Static(&FLevelEditorActionCallbacks::SetLightingResolutionMinSMs)
+					]
+					+ SHorizontalBox::Slot()
+					.AutoWidth()
+					[
+						SNew(SSpinBox<float>)
+						.MinValue(4.f)
+					.MaxValue(4096.f)
+					.ToolTipText(LOCTEXT("LightingResolutionStaticMeshesMaxToolTip", "The maximum lightmap resolution for static mesh adjustments. Anything outside of Min/Max range will not be touched when adjusting."))
+					.Value(FLevelEditorActionCallbacks::GetLightingResolutionMaxSMs())
+					.OnValueChanged_Static(&FLevelEditorActionCallbacks::SetLightingResolutionMaxSMs)
+					];
+				Section.AddEntry(FToolMenuEntry::InitWidget("Meshes", Meshes, FText::GetEmpty(), true));
+
+				TSharedRef<SWidget> BSPs = SNew(SHorizontalBox)
+				+ SHorizontalBox::Slot()
+				.AutoWidth()
+				[
+					SNew(SCheckBox)
+					.Style(FEditorStyle::Get(), "Menu.CheckBox")
+					.ToolTipText(LOCTEXT("BSPSurfacesToolTip", "BSP Surfaces will be adjusted if checked."))
+					.IsChecked_Static(&FLevelEditorActionCallbacks::IsLightingResolutionBSPSurfacesChecked)
+					.OnCheckStateChanged_Static(&FLevelEditorActionCallbacks::SetLightingResolutionBSPSurfaces)
+					.Content()
+					[
+						SNew(STextBlock)
+						.Text(LOCTEXT("BSPSurfaces", "BSP Surfaces"))
+					]
+				]
+				+ SHorizontalBox::Slot()
+				.AutoWidth()
+				.Padding(FMargin(6.0f, 0.0f, 4.0f, 0.0f))
+				[
+					SNew(SSpinBox<float>)
+					.MinValue(1.f)
+					.MaxValue(63556.f)
+					.ToolTipText(LOCTEXT("LightingResolutionBSPsMinToolTip", "The minimum lightmap resolution of a BSP surface to adjust. When outside of the Min/Max range, the BSP surface will no be altered."))
+					.Value(FLevelEditorActionCallbacks::GetLightingResolutionMinBSPs())
+					.OnValueChanged_Static(&FLevelEditorActionCallbacks::SetLightingResolutionMinBSPs)
+				]
+				+ SHorizontalBox::Slot()
+				.AutoWidth()
+				[
+					SNew(SSpinBox<float>)
+					.MinValue(1.f)
+					.MaxValue(63556.f)
+					.ToolTipText(LOCTEXT("LightingResolutionBSPsMaxToolTip", "The maximum lightmap resolution of a BSP surface to adjust. When outside of the Min/Max range, the BSP surface will no be altered."))
+					.Value(FLevelEditorActionCallbacks::GetLightingResolutionMaxBSPs())
+					.OnValueChanged_Static(&FLevelEditorActionCallbacks::SetLightingResolutionMaxBSPs)
+				];
+				Section.AddEntry(FToolMenuEntry::InitWidget("BSPs", BSPs, FText::GetEmpty(), true));
+			}
+
+			{
+				FToolMenuSection& Section = SubMenu->AddSection("LevelEditorBuildLightingResolution2", LOCTEXT("LightingResolutionHeading2", "Select Options"));
+				Section.AddMenuEntry(FLevelEditorCommands::Get().LightingResolution_CurrentLevel);
+				Section.AddMenuEntry(FLevelEditorCommands::Get().LightingResolution_SelectedLevels);
+				Section.AddMenuEntry(FLevelEditorCommands::Get().LightingResolution_AllLoadedLevels);
+				Section.AddMenuEntry(FLevelEditorCommands::Get().LightingResolution_SelectedObjectsOnly);
+			}
+
+			{
+				FToolMenuSection& Section = SubMenu->AddSection("LevelEditorBuildLightingResolution3", LOCTEXT("LightingResolutionHeading3", "Ratio"));
+				TSharedRef<SWidget> Ratio = SNew(SSpinBox<int32>)
+					.MinValue(0)
+					.MaxValue(400)
+					.ToolTipText(LOCTEXT("LightingResolutionRatioToolTip", "Ratio to apply (New Resolution = Ratio / 100.0f * CurrentResolution)."))
+					.Value(FLevelEditorActionCallbacks::GetLightingResolutionRatio())
+					.OnEndSliderMovement_Static(&FLevelEditorActionCallbacks::SetLightingResolutionRatio)
+					.OnValueCommitted_Static(&FLevelEditorActionCallbacks::SetLightingResolutionRatioCommit);
+				Section.AddEntry(FToolMenuEntry::InitWidget("Ratio", Ratio, LOCTEXT("LightingResolutionRatio", "Ratio")));
+			}
+		}
+
+		/** Generates a lighting info dialogs sub-menu */
+		static void RegisterLightingInfoMenu(const FName InBaseMenuName)
+		{
+			RegisterLightingDensityMenu(UToolMenus::JoinMenuPaths(InBaseMenuName, "LightingInfo"));
+			RegisterLightingResolutionMenu(UToolMenus::JoinMenuPaths(InBaseMenuName, "LightingInfo"));
+
+			UToolMenu* SubMenu = UToolMenus::Get()->RegisterMenu(UToolMenus::JoinMenuPaths(InBaseMenuName, "LightingInfo"));
+
+			{
+				FToolMenuSection& Section = SubMenu->AddSection("LevelEditorBuildLightingInfo", LOCTEXT("LightingInfoHeading", "Lighting Info Dialogs"));
+				Section.AddSubMenu(
+					"LightingDensity",
+					LOCTEXT("LightingDensityRenderingSubMenu", "LightMap Density Rendering Options"),
+					LOCTEXT("LightingDensityRenderingSubMenu_ToolTip", "Shows the LightMap Density Rendering viewmode options."),
+					FNewToolMenuChoice());
+
+				Section.AddSubMenu(
+					"LightingResolution",
+					LOCTEXT("LightingResolutionAdjustmentSubMenu", "LightMap Resolution Adjustment"),
+					LOCTEXT("LightingResolutionAdjustmentSubMenu_ToolTip", "Shows the LightMap Resolution Adjustment options."),
+					FNewToolMenuChoice());
+
+				Section.AddMenuEntry(FLevelEditorCommands::Get().LightingStaticMeshInfo, LOCTEXT("BuildLightingInfo_LightingStaticMeshInfo", "Lighting StaticMesh Info..."));
+			}
+		}
+	};
+
+	{
+		FToolMenuSection& Section = Menu->AddSection("Level", LOCTEXT("LevelHeading", "Level"));
+
+		Section.AddMenuEntry(FLevelEditorCommands::Get().Build, LOCTEXT("Build", "Build All Levels"));
+	}
+
+	FLightingMenus::RegisterMenus(BaseMenuName);
+
+	{
+		FToolMenuSection& Section = Menu->AddSection("LevelEditorLighting", LOCTEXT("LightingHeading", "Lighting"));
+		Section.AddMenuEntry(FLevelEditorCommands::Get().BuildLightingOnly, LOCTEXT("BuildLightingOnlyHeading", "Build Lighting Only"));
+
+		Section.AddSubMenu(
+			"LightingQuality",
+			LOCTEXT("LightingQualitySubMenu", "Lighting Quality"),
+			LOCTEXT("LightingQualitySubMenu_ToolTip", "Allows you to select the quality level for precomputed lighting"),
+			FNewToolMenuChoice());
+
+		Section.AddSubMenu(
+			"LightingInfo",
+			LOCTEXT("BuildLightingInfoSubMenu", "Lighting Info"),
+			LOCTEXT("BuildLightingInfoSubMenu_ToolTip", "Access the lighting info dialogs"),
+			FNewToolMenuChoice());
+
+		Section.AddMenuEntry(FLevelEditorCommands::Get().LightingBuildOptions_UseErrorColoring);
+		Section.AddMenuEntry(FLevelEditorCommands::Get().LightingBuildOptions_ShowLightingStats);
+	}
+
+	{
+		FToolMenuSection& Section = Menu->AddSection("LevelEditorReflections", LOCTEXT("ReflectionHeading", "Reflections"));
+		Section.AddMenuEntry(FLevelEditorCommands::Get().BuildReflectionCapturesOnly);
+	}
+
+	{
+		FToolMenuSection& Section = Menu->AddSection("LevelEditorVisibility", LOCTEXT("VisibilityHeading", "Visibility"));
+		Section.AddMenuEntry(FLevelEditorCommands::Get().BuildLightingOnly_VisibilityOnly);
+	}
+
+	{
+		FToolMenuSection& Section = Menu->AddSection("LevelEditorGeometry", LOCTEXT("GeometryHeading", "Geometry"));
+		Section.AddMenuEntry(FLevelEditorCommands::Get().BuildGeometryOnly);
+		Section.AddMenuEntry(FLevelEditorCommands::Get().BuildGeometryOnly_OnlyCurrentLevel);
+	}
+
+	{
+		FToolMenuSection& Section = Menu->AddSection("LevelEditorNavigation", LOCTEXT("NavigationHeading", "Navigation"));
+		Section.AddMenuEntry(FLevelEditorCommands::Get().BuildPathsOnly);
+	}
+
+	{
+		FToolMenuSection& Section = Menu->AddSection("LevelEditorLOD", LOCTEXT("LODHeading", "Hierarchical LOD"));
+		Section.AddMenuEntry(FLevelEditorCommands::Get().BuildLODsOnly);
+	}
+
+	{
+		FToolMenuSection& Section = Menu->AddSection("LevelEditorTextureStreaming", LOCTEXT("TextureStreamingHeading", "Texture Streaming"));
+		Section.AddDynamicEntry("BuildTextureStreamingOnly", FNewToolMenuSectionDelegate::CreateLambda([](FToolMenuSection& InSection)
+			{
+				if (CVarStreamingUseNewMetrics.GetValueOnAnyThread() != 0) // There is no point of in building texture streaming data with the old system.
+				{
+					InSection.AddMenuEntry(FLevelEditorCommands::Get().BuildTextureStreamingOnly);
+				}
+			}));
+		Section.AddMenuEntry(FLevelEditorCommands::Get().BuildVirtualTextureOnly);
+	}
+
+	{
+		FToolMenuSection& Section = Menu->AddSection("LevelEditorGrass", LOCTEXT("GrassHeading", "Grass Maps"));
+		Section.AddMenuEntry(FLevelEditorCommands::Get().BuildGrassMapsOnly);
+	}
+
+	{
+		FToolMenuSection& Section = Menu->AddSection("LevelEditorAutomation", LOCTEXT("AutomationHeading", "Automation"));
+		Section.AddMenuEntry(
+			FLevelEditorCommands::Get().BuildAndSubmitToSourceControl,
+			TAttribute<FText>(),
+			TAttribute<FText>(),
+			FSlateIcon(FEditorStyle::GetStyleSetName(), "LevelEditor.Tabs.BuildAndSubmit")
+		);
+	}
+
+	// Map Check
+	{
+		FToolMenuSection& Section = Menu->AddSection("LevelEditorVerification", LOCTEXT("VerificationHeading", "Verification"));
+		Section.AddMenuEntry(FLevelEditorCommands::Get().MapCheck, LOCTEXT("OpenMapCheck", "Map Check"));
+	}
+
+}
 #undef LOCTEXT_NAMESPACE
\ No newline at end of file
diff --git a/Engine/Source/Editor/LevelEditor/Private/SLevelEditor.cpp b/Engine/Source/Editor/LevelEditor/Private/SLevelEditor.cpp
--- a/Engine/Source/Editor/LevelEditor/Private/SLevelEditor.cpp
+++ b/Engine/Source/Editor/LevelEditor/Private/SLevelEditor.cpp
@@ -52,6 +52,7 @@
 #include "EditorViewportCommands.h"
 #include "IPlacementModeModule.h"
 #include "Classes/EditorStyleSettings.h"
+#include "Subsystems/StatusBarSubsystem.h"
 
 static const FName MainFrameModuleName("MainFrame");
 static const FName LevelEditorModuleName("LevelEditor");
@@ -64,7 +65,8 @@ namespace LevelEditorConstants
 }
 
 SLevelEditor::SLevelEditor()
-	: World(NULL), bNeedsRefresh(false)
+	: World(nullptr)
+	, bNeedsRefresh(false)
 {
 	const bool bAreRealTimeThumbnailsAllowed = false;
 	ThumbnailPool = MakeShareable(new FAssetThumbnailPool(LevelEditorConstants::ThumbnailPoolSize, bAreRealTimeThumbnailsAllowed));
@@ -152,12 +154,16 @@ void SLevelEditor::Construct( const SLevelEditor::FArguments& InArgs)
 	// It's imperative that our delegate is removed in the destructor for the level editor module to play nicely with reloading.
 
 	FLevelEditorModule& LevelEditorModule = FModuleManager::GetModuleChecked< FLevelEditorModule >( LevelEditorModuleName );
-	LevelEditorModule.OnNotificationBarChanged().AddRaw( this, &SLevelEditor::ConstructNotificationBar );
+	LevelEditorModule.OnTitleBarMessagesChanged().AddRaw( this, &SLevelEditor::ConstructTitleBarMessages );
 
 	GetMutableDefault<UEditorExperimentalSettings>()->OnSettingChanged().AddRaw(this, &SLevelEditor::HandleExperimentalSettingChanged);
 
 	BindCommands();
 
+	// Rebuild the editor mode commands and their tab spawners before we restore the layout,
+	// or there wont be any tab spawners for the modes.
+	RefreshEditorModeCommands();
+
 	RegisterMenus();
 
 	// We need to register when modes list changes so that we can refresh the auto generated commands.
@@ -193,91 +199,63 @@ void SLevelEditor::Initialize( const TSharedRef<SDockTab>& OwnerTab, const TShar
 
 	LevelEditorModule.OnActorSelectionChanged().AddSP(this, &SLevelEditor::OnActorSelectionChanged);
 
-	TSharedRef<SWidget> Widget2 = RestoreContentArea( OwnerTab, OwnerWindow );
-	TSharedRef<SWidget> Widget1 = FLevelEditorMenu::MakeLevelEditorMenu( LevelEditorCommands, SharedThis(this) );
-
+	TSharedRef<SWidget> ContentArea = RestoreContentArea( OwnerTab, OwnerWindow );
+	TSharedRef<SWidget> MenuBar =
+		SNew(SBox)
+		.AddMetaData<FTagMetaData>(FTagMetaData(TEXT("MainMenu")))
+		[
+			FLevelEditorMenu::MakeLevelEditorMenu(LevelEditorCommands, SharedThis(this))
+		];
+		
 	ChildSlot
 	[
-		SNew( SVerticalBox )
+		SNew(SVerticalBox)
 		+SVerticalBox::Slot()
+		.Padding(FMargin(0.0f,0.0f,0.0f,2.0f))
 		.AutoHeight()
 		[
-			SNew( SOverlay )
-
-			+SOverlay::Slot()
-			[
-				SNew( SBox )
-				.AddMetaData<FTagMetaData>(FTagMetaData(TEXT("MainMenu")))
-				[
-					Widget1
-				]
-			]
-		 
-// For platforms without a global menu bar we can put the perf. tools in the editor window's menu bar
-#if !PLATFORM_MAC
-			+ SOverlay::Slot()
-			.HAlign( HAlign_Right )
-			.VAlign( VAlign_Center )
-			[
-				SAssignNew( NotificationBarBox, SHorizontalBox )
-				.AddMetaData<FTagMetaData>(FTagMetaData(TEXT("PerformanceTools")))
-			]
-#endif
+			FLevelEditorToolBar::MakeLevelEditorToolBar(LevelEditorCommands.ToSharedRef(), SharedThis(this))
 		]
-
-#if PLATFORM_MAC
-		// Without the in-window menu bar, we need some space between the tab bar and tab contents
 		+SVerticalBox::Slot()
-		.AutoHeight()
+		.FillHeight( 1.0f )
 		[
-			SNew( SBox )
-			.HeightOverride( 1.0f )
+			ContentArea
 		]
-#endif
-
 		+SVerticalBox::Slot()
-		.FillHeight( 1.0f )
+		.AutoHeight()
 		[
-			Widget2
+			GEditor->GetEditorSubsystem<UStatusBarSubsystem>()->MakeStatusBarWidget(TEXT("LevelEditor.StatusBar"))
 		]
 	];
 	
-// For OS X we need to put it into the window's title bar since there's no per-window menu bar
-#if PLATFORM_MAC	
-	OwnerTab->SetRightContent(
-		SAssignNew( NotificationBarBox, SHorizontalBox )
-		.AddMetaData<FTagMetaData>(FTagMetaData(TEXT("PerformanceTools")))
-		);
-#endif
+	TtileBarMessageBox = SNew(SHorizontalBox);
 
-	ConstructNotificationBar();
+	ConstructTitleBarMessages();
 
 	OnLayoutHasChanged();
 }
 
-void SLevelEditor::ConstructNotificationBar()
+void SLevelEditor::ConstructTitleBarMessages()
 {
-	NotificationBarBox->ClearChildren();
+	TtileBarMessageBox->ClearChildren();
 
-	// level editor commands
-	NotificationBarBox->AddSlot()
+	TtileBarMessageBox->AddSlot()
 		.AutoWidth()
 		.Padding(5.0f, 0.0f, 0.0f, 0.0f)
 		[
 			FLevelEditorMenu::MakeNotificationBar( LevelEditorCommands, SharedThis(this ) )
 		];
 
-	// developer tools
 	const IMainFrameModule& MainFrameModule = FModuleManager::GetModuleChecked<IMainFrameModule>(MainFrameModuleName);
 	const FLevelEditorModule& LevelEditorModule = FModuleManager::GetModuleChecked< FLevelEditorModule >(LevelEditorModuleName);
 	
 	TArray<FMainFrameDeveloperTool> Tools;
-	for (const TTuple<FName, FLevelEditorModule::FStatusBarItem>& Item : LevelEditorModule.GetStatusBarItems())
+	for (const TPair<FName, FLevelEditorModule::FTitleBarItem>& Item : LevelEditorModule.GetTitleBarItems())
 	{
 		Tools.Add({Item.Value.Visibility, Item.Value.Label, Item.Value.Value});
 	}
 
-	NotificationBarBox->AddSlot()
+	TtileBarMessageBox->AddSlot()
 		.AutoWidth()
 		.Padding(5.0f, 0.0f, 0.0f, 0.0f)
 		[
@@ -293,7 +271,7 @@ SLevelEditor::~SLevelEditor()
 	HostedToolkits.Reset();
 	
 	FLevelEditorModule& LevelEditorModule = FModuleManager::GetModuleChecked< FLevelEditorModule >( LevelEditorModuleName );
-	LevelEditorModule.OnNotificationBarChanged().RemoveAll( this );
+	LevelEditorModule.OnTitleBarMessagesChanged().RemoveAll( this );
 	
 	if(UObjectInitialized())
 	{
@@ -632,25 +610,6 @@ TSharedRef<SDockTab> SLevelEditor::SpawnLevelEditorTab( const FSpawnTabArgs& Arg
 	{
 		return this->BuildViewportTab( NSLOCTEXT("LevelViewportTypes", "LevelEditorViewport_Clone3", "Viewport 4"), TEXT("Viewport 4"), InitializationPayload );
 	}
-	else if( TabIdentifier == LevelEditorTabIds::LevelEditorToolBar)
-	{
-		return SNew( SDockTab )
-			.Label( NSLOCTEXT("LevelEditor", "ToolBarTabTitle", "Toolbar") )
-			.ShouldAutosize(true)
-			.Icon( FEditorStyle::GetBrush("ToolBar.Icon") )
-			[
-				SNew(SHorizontalBox)
-				.AddMetaData<FTagMetaData>(FTagMetaData(TEXT("LevelEditorToolbar")))
-				+SHorizontalBox::Slot()
-				.FillWidth(1)
-				.VAlign(VAlign_Bottom)
-				.HAlign(HAlign_Left)
-				[
-					FLevelEditorToolBar::MakeLevelEditorToolBar( LevelEditorCommands.ToSharedRef(), SharedThis(this) )
-				]
-			];
-
-	}
 	else if (TabIdentifier == FEditorModeTools::EditorModeToolbarTabName)
 	{
 		return GLevelEditorModeTools().MakeModeToolbarTab();
@@ -1113,16 +1072,6 @@ TSharedRef<SWidget> SLevelEditor::RestoreContentArea( const TSharedRef<SDockTab>
 				.SetGroup( MenuStructure.GetLevelEditorViewportsCategory() )
 				.SetIcon(ViewportIcon);
 		}
-
-		{
-			const FSlateIcon ToolbarIcon(FEditorStyle::GetStyleSetName(), "LevelEditor.Tabs.Toolbar");
-			LevelEditorTabManager->RegisterTabSpawner(LevelEditorTabIds::LevelEditorToolBar, FOnSpawnTab::CreateSP<SLevelEditor, FName, FString>(this, &SLevelEditor::SpawnLevelEditorTab, LevelEditorTabIds::LevelEditorToolBar, FString()))
-				.SetDisplayName(NSLOCTEXT("LevelEditorTabs", "LevelEditorToolBar", "Toolbar"))
-				.SetTooltipText(NSLOCTEXT("LevelEditorTabs", "LevelEditorToolBarTooltipText", "Open the Toolbar tab, which provides access to the most common / important actions."))
-				.SetGroup( MenuStructure.GetLevelEditorCategory() )
-				.SetIcon( ToolbarIcon );
-		}
-
 		{
 			const FText DetailsTooltip = NSLOCTEXT("LevelEditorTabs", "LevelEditorSelectionDetailsTooltip", "Open a Details tab. Use this to view and edit properties of the selected object(s).");
 			const FSlateIcon DetailsIcon(FEditorStyle::GetStyleSetName(), "LevelEditor.Tabs.Details");
@@ -1269,9 +1218,6 @@ TSharedRef<SWidget> SLevelEditor::RestoreContentArea( const TSharedRef<SDockTab>
 		LevelEditorModule.OnRegisterTabs().Broadcast(LevelEditorTabManager);
 	}
 
-	// Rebuild the editor mode commands and their tab spawners before we restore the layout,
-	// or there wont be any tab spawners for the modes.
-	RefreshEditorModeCommands();
 
 	// IMPORTANT: If you want to change the default value of "LevelEditor_Layout_v1.1" or "UnrealEd_Layout_v1.4" (even if you only change their version numbers), these are the steps to follow:
 	// 1. Check out Engine\Config\Layouts\DefaultLayout.ini in Perforce.
@@ -1288,7 +1234,7 @@ TSharedRef<SWidget> SLevelEditor::RestoreContentArea( const TSharedRef<SDockTab>
 	// 10. Push the new "DefaultLayout.ini" together with your new code.
 	// 11. Also update these instructions if you change the version number (e.g., from "UnrealEd_Layout_v1.4" to "UnrealEd_Layout_v1.5").
 	const TSharedRef<FTabManager::FLayout> Layout = FLayoutSaveRestore::LoadFromConfig(GEditorLayoutIni,
-		FTabManager::NewLayout( "LevelEditor_Layout_v1.2" )
+		FTabManager::NewLayout( "LevelEditor_Layout_v1.3" )
 		->AddArea
 		(
 			FTabManager::NewPrimaryArea()
@@ -1318,12 +1264,6 @@ TSharedRef<SWidget> SLevelEditor::RestoreContentArea( const TSharedRef<SDockTab>
 						->SetOrientation(Orient_Vertical)
 						->SetSizeCoefficient( 1.15f )
 						->Split
-						(
-							FTabManager::NewStack()
-							->SetHideTabWell(true)
-							->AddTab(LevelEditorTabIds::LevelEditorToolBar, ETabState::OpenedTab)
-						)
-						->Split
 						(
 							FTabManager::NewStack()
 							->SetHideTabWell(true)
@@ -1384,7 +1324,7 @@ void SLevelEditor::HandleExperimentalSettingChanged(FName PropertyName)
 {
 	FLevelEditorModule& LevelEditorModule = FModuleManager::GetModuleChecked<FLevelEditorModule>("LevelEditor");
 	TSharedPtr<FTabManager> LevelEditorTabManager = LevelEditorModule.GetLevelEditorTabManager();
-	LevelEditorTabManager->UpdateMainMenu(true);
+	LevelEditorTabManager->UpdateMainMenu(nullptr, true);
 }
 
 FName SLevelEditor::GetEditorModeTabId( FEditorModeID ModeID )
diff --git a/Engine/Source/Editor/LevelEditor/Private/SLevelEditor.h b/Engine/Source/Editor/LevelEditor/Private/SLevelEditor.h
--- a/Engine/Source/Editor/LevelEditor/Private/SLevelEditor.h
+++ b/Engine/Source/Editor/LevelEditor/Private/SLevelEditor.h
@@ -136,6 +136,7 @@ public:
 	/** Returns current scene outliner associated with level editor's scene outliner tab, if it exists */
 	virtual TSharedPtr<ISceneOutliner> GetSceneOutliner() const override { return SceneOutlinerPtr.Pin();  }
 
+	TSharedRef<SWidget> GetTitleBarMessageWidget() const { return TtileBarMessageBox.ToSharedRef(); }
 private:
 	
 	TSharedRef<SDockTab> SpawnLevelEditorTab(const FSpawnTabArgs& Args, FName TabIdentifier, FString InitializationPayload);
@@ -194,8 +195,8 @@ private:
 	/** Callback for when the level editor layout has changed */
 	void OnLayoutHasChanged();
 	
-	/** Constructs the NotificationBar widgets */
-	void ConstructNotificationBar();
+	/** Constructs the title bar message widget */
+	void ConstructTitleBarMessages();
 
 	/** Builds a viewport tab. */
 	TSharedRef<SDockTab> BuildViewportTab( const FText& Label, const FString LayoutId, const FString& InitializationPayload );
@@ -248,8 +249,8 @@ private:
 	// The UWorld that this level editor is viewing and allowing the user to interact with through.
 	UWorld* World;
 
-	// The box that holds the notification bar.
-	TSharedPtr< SHorizontalBox > NotificationBarBox;
+	// The box that holds the title bar messages.
+	TSharedPtr<SHorizontalBox> TtileBarMessageBox;
 
 	// Holds the world settings details view.
 	TSharedPtr<IDetailsView> WorldSettingsView;
diff --git a/Engine/Source/Editor/LevelEditor/Public/LevelEditor.h b/Engine/Source/Editor/LevelEditor/Public/LevelEditor.h
--- a/Engine/Source/Editor/LevelEditor/Public/LevelEditor.h
+++ b/Engine/Source/Editor/LevelEditor/Public/LevelEditor.h
@@ -33,7 +33,6 @@ public:
 	static const FName LevelEditorViewport_Clone2;
 	static const FName LevelEditorViewport_Clone3;
 	static const FName LevelEditorViewport_Clone4;
-	static const FName LevelEditorToolBar;
 	static const FName LevelEditorToolBox;
 	static const FName LevelEditorSelectionDetails;
 	static const FName LevelEditorSelectionDetails2;
@@ -280,25 +279,25 @@ public:
 	DECLARE_EVENT_OneParam(ILevelEditor, FOnRegisterLayoutExtensions, FLayoutExtender&);
 	FOnRegisterLayoutExtensions& OnRegisterLayoutExtensions() { return RegisterLayoutExtensions; }
 
-	/** Add / Remove status bar item */
-	struct FStatusBarItem
+	/** Add / Remove title bar item */
+	struct FTitleBarItem
 	{
-		/* Visiblility of the status bar item. */
+		/* Visiblility of the title bar item. */
 		TAttribute<EVisibility> Visibility;
-		/* Label of the status bar item. */
+		/* Label of the title bar item. */
 		TAttribute<FText> Label;
-		/* Value of the status bar item. */
+		/* Value of the title bar item. */
 		TAttribute<FText> Value;
 	};
-	const TMap<FName, FStatusBarItem>& GetStatusBarItems() const { return StatusBarItems; }
-	virtual void AddStatusBarItem(FName InStatusBarIdentifier, const FStatusBarItem& InStatusBarItem);
-	virtual void RemoveStatusBarItem(FName InStatusBarIdentifier);
+	const TMap<FName, FTitleBarItem>& GetTitleBarItems() const { return TitleBarItems; }
+	virtual void AddTitleBarItem(FName InTitleBarIdentifier, const FTitleBarItem& InTitleBarITem);
+	virtual void RemoveTitleBarItem(FName InTitleBarIdentifier);
 
 	/** Called when a new map is loaded */
-	DECLARE_EVENT( FLevelEditorModule, FNotificationBarChanged );
-	virtual FNotificationBarChanged& OnNotificationBarChanged() { return NotificationBarChangedEvent; }
+	DECLARE_EVENT( FLevelEditorModule, FTitleBarMessagesChanged );
+	virtual FTitleBarMessagesChanged& OnTitleBarMessagesChanged() { return TitleBarMessagesChangedEvent; }
 
-	virtual void BroadcastNotificationBarChanged() { NotificationBarChangedEvent.Broadcast(); }
+	virtual void BroadcastTitleBarMessagesChanged() { TitleBarMessagesChangedEvent.Broadcast(); }
 
 	/** Called when a high res screenshot is requested. */
 	DECLARE_EVENT( FLevelEditorModule, FTakeHighResScreenShotsEvent );
@@ -386,8 +385,8 @@ private:
 	TSharedPtr<FExtensibilityManager> ModeBarExtensibilityManager;
 	TSharedPtr<FExtensibilityManager> NotificationBarExtensibilityManager;
 
-	TMap<FName, FStatusBarItem> StatusBarItems;
-	FNotificationBarChanged NotificationBarChangedEvent;
+	TMap<FName, FTitleBarItem> TitleBarItems;
+	FTitleBarMessagesChanged TitleBarMessagesChangedEvent;
 
 	/** 
 	 * A command list that can be passed around and isn't bound to an instance of a level editor. 
diff --git a/Engine/Source/Editor/LevelEditor/Public/LevelEditorActions.h b/Engine/Source/Editor/LevelEditor/Public/LevelEditorActions.h
--- a/Engine/Source/Editor/LevelEditor/Public/LevelEditorActions.h
+++ b/Engine/Source/Editor/LevelEditor/Public/LevelEditorActions.h
@@ -605,14 +605,6 @@ public:
 	//TSharedPtr< FUICommandInfo > LandscapeMode;
 	//TSharedPtr< FUICommandInfo > FoliageMode;
 
-	/**
-	 * Source Control Commands
-	 */
-	TSharedPtr< FUICommandInfo > ConnectToSourceControl;
-	TSharedPtr< FUICommandInfo > ChangeSourceControlSettings;
-	TSharedPtr< FUICommandInfo > CheckOutModifiedFiles;
-	TSharedPtr< FUICommandInfo > SubmitToSourceControl;
-
 	/**
 	 * Misc Commands
 	 */
diff --git a/Engine/Source/Editor/MainFrame/Private/Menus/MainMenu.cpp b/Engine/Source/Editor/MainFrame/Private/Menus/MainMenu.cpp
--- a/Engine/Source/Editor/MainFrame/Private/Menus/MainMenu.cpp
+++ b/Engine/Source/Editor/MainFrame/Private/Menus/MainMenu.cpp
@@ -118,7 +118,7 @@ void FMainMenu::RegisterEditMenu()
 	}
 
 	{
-		FToolMenuSection& Section = EditMenu->AddSection("EditLocalTabSpawners", LOCTEXT("ConfigurationHeading", "Configuration"));
+		FToolMenuSection& Section = EditMenu->AddSection("Configuration", LOCTEXT("ConfigurationHeading", "Configuration"));
 		if (GetDefault<UEditorStyleSettings>()->bExpandConfigurationMenus)
 		{
 			Section.AddSubMenu(
@@ -315,6 +315,7 @@ void FMainMenu::RegisterHelpMenu()
 		HelpOnlineSection.AddMenuEntry(FMainFrameCommands::Get().VisitWiki);
 
 
+
 		const FText SupportWebSiteLabel = NSLOCTEXT("MainHelpMenu", "VisitUnrealEngineSupportWebSite", "Unreal Engine Support Web Site...");
 
 		HelpOnlineSection.AddSeparator("EpicGamesHelp");
@@ -355,7 +356,7 @@ TSharedRef<SWidget> FMainMenu::MakeMainMenu(const TSharedPtr<FTabManager>& TabMa
 	{
 		// Tell tab-manager about the multi-box for platforms with a global menu bar
 		TSharedRef<SMultiBoxWidget> MultiBoxWidget = StaticCastSharedRef<SMultiBoxWidget>(MenuBarWidget);
-		TabManager->SetMenuMultiBox(ConstCastSharedRef<FMultiBox>(MultiBoxWidget->GetMultiBox()));
+		TabManager->SetMenuMultiBox(ConstCastSharedRef<FMultiBox>(MultiBoxWidget->GetMultiBox()), MultiBoxWidget);
 	}
 
 	return MenuBarWidget;
@@ -379,6 +380,10 @@ void FMainMenu::RegisterMainMenu()
 
 	UToolMenu* MenuBar = ToolMenus->RegisterMenu(MainMenuName, NAME_None, EMultiBoxType::MenuBar);
 
+	static const FName MainMenuStyleName("WindowMenuBar");
+
+	MenuBar->StyleName = MainMenuStyleName;
+
 	MenuBar->AddSubMenu(
 		"MainMenu",
 		NAME_None,
diff --git a/Engine/Source/Editor/UnrealEd/Private/Subsystems/StatusBarSubsystem.cpp b/Engine/Source/Editor/UnrealEd/Private/Subsystems/StatusBarSubsystem.cpp
new file mode 100644
--- /dev/null
+++ b/Engine/Source/Editor/UnrealEd/Private/Subsystems/StatusBarSubsystem.cpp
@@ -0,0 +1,491 @@
+// Copyright Epic Games, Inc. All Rights Reserved.
+
+#pragma once
+
+#include "Subsystems/StatusBarSubsystem.h"
+#include "ToolMenuContext.h"
+#include "ToolMenus.h"
+#include "ISourceControlOperation.h"
+#include "SourceControlOperations.h"
+#include "ISourceControlProvider.h"
+#include "ISourceControlModule.h"
+#include "SourceControlWindows.h"
+
+#define LOCTEXT_NAMESPACE "StatusBar"
+
+class FSourceControlCommands : public TCommands<FSourceControlCommands>
+{
+public:
+	FSourceControlCommands() : TCommands<FSourceControlCommands>
+	(
+		"SourceControl",
+		NSLOCTEXT("Contexts", "SourceControl", "Source Control"),
+		"LevelEditor",
+		FEditorStyle::GetStyleSetName()
+	)
+	{}
+
+	/**
+	 * Initialize commands
+	 */
+	virtual void RegisterCommands() override
+	{
+		UI_COMMAND(ConnectToSourceControl, "Connect to Source Control...", "Opens a dialog to connect to source control.", EUserInterfaceActionType::Button, FInputChord());
+		UI_COMMAND(ChangeSourceControlSettings, "Change Source Control Settings...", "Opens a dialog to change source control settings.", EUserInterfaceActionType::Button, FInputChord());
+		UI_COMMAND(CheckOutModifiedFiles, "Check Out Modified Files...", "Opens a dialog to check out any assets which have been modified.", EUserInterfaceActionType::Button, FInputChord());
+		UI_COMMAND(SubmitToSourceControl, "Submit to Source Control...", "Opens a dialog with check in options for content and levels.", EUserInterfaceActionType::Button, FInputChord());
+
+		ActionList->MapAction(
+			ConnectToSourceControl,
+			FExecuteAction::CreateStatic(&FSourceControlCommands::ConnectToSourceControl_Clicked)
+		);
+
+		ActionList->MapAction(
+			ChangeSourceControlSettings,
+			FExecuteAction::CreateStatic(&FSourceControlCommands::ConnectToSourceControl_Clicked)
+		);
+
+		ActionList->MapAction(
+			CheckOutModifiedFiles,
+			FExecuteAction::CreateStatic(&FSourceControlCommands::CheckOutModifiedFiles_Clicked),
+			FCanExecuteAction::CreateStatic(&FSourceControlCommands::CheckOutModifiedFiles_CanExecute)
+		);
+
+		ActionList->MapAction(
+			SubmitToSourceControl,
+			FExecuteAction::CreateStatic(&FSourceControlCommands::SubmitToSourceControl_Clicked),
+			FCanExecuteAction::CreateStatic(&FSourceControlCommands::SubmitToSourceControl_CanExecute)
+		);
+
+	}
+
+private:
+
+	static void ConnectToSourceControl_Clicked()
+	{
+		// Show login window regardless of current status - its useful as a shortcut to change settings.
+		ISourceControlModule& SourceControlModule = ISourceControlModule::Get();
+		SourceControlModule.ShowLoginDialog(FSourceControlLoginClosed(), ELoginWindowMode::Modeless, EOnLoginWindowStartup::PreserveProvider);
+	}
+
+	static bool CheckOutModifiedFiles_CanExecute()
+	{
+		ISourceControlModule& SourceControlModule = ISourceControlModule::Get();
+		if (ISourceControlModule::Get().IsEnabled() &&
+			ISourceControlModule::Get().GetProvider().IsAvailable())
+		{
+			TArray<UPackage*> PackagesToSave;
+			FEditorFileUtils::GetDirtyWorldPackages(PackagesToSave);
+			FEditorFileUtils::GetDirtyContentPackages(PackagesToSave);
+
+			return PackagesToSave.Num() > 0;
+		}
+
+		return false;
+	}
+
+	static void CheckOutModifiedFiles_Clicked()
+	{
+		TArray<UPackage*> PackagesToSave;
+		FEditorFileUtils::GetDirtyWorldPackages(PackagesToSave);
+		FEditorFileUtils::GetDirtyContentPackages(PackagesToSave);
+
+		const bool bCheckDirty = true;
+		const bool bPromptUserToSave = false;
+		FEditorFileUtils::PromptForCheckoutAndSave(PackagesToSave, bCheckDirty, bPromptUserToSave);
+	}
+
+	static bool SubmitToSourceControl_CanExecute()
+	{
+		ISourceControlModule& SourceControlModule = ISourceControlModule::Get();
+		return ISourceControlModule::Get().IsEnabled() &&
+			ISourceControlModule::Get().GetProvider().IsAvailable() &&
+			FSourceControlWindows::CanChoosePackagesToCheckIn();
+	}
+
+	static void SubmitToSourceControl_Clicked()
+	{
+		FSourceControlWindows::ChoosePackagesToCheckIn();
+	}
+
+	
+public:
+	/**
+	 * Source Control Commands
+	 */
+	TSharedPtr< FUICommandInfo > ConnectToSourceControl;
+	TSharedPtr< FUICommandInfo > ChangeSourceControlSettings;
+	TSharedPtr< FUICommandInfo > CheckOutModifiedFiles;
+	TSharedPtr< FUICommandInfo > SubmitToSourceControl;
+
+	static TSharedRef<FUICommandList> ActionList;
+};
+
+TSharedRef<FUICommandList> FSourceControlCommands::ActionList(new FUICommandList());
+
+struct FSourceControlMenuHelpers
+{
+	enum EQueryState
+	{
+		NotQueried,
+		Querying,
+		Queried,
+	};
+
+
+	static EQueryState QueryState;
+
+	static void CheckSourceControlStatus()
+	{
+		ISourceControlModule& SourceControlModule = ISourceControlModule::Get();
+		if (SourceControlModule.IsEnabled())
+		{
+			SourceControlModule.GetProvider().Execute(ISourceControlOperation::Create<FConnect>(),
+				EConcurrency::Asynchronous,
+				FSourceControlOperationComplete::CreateStatic(&FSourceControlMenuHelpers::OnSourceControlOperationComplete));
+			QueryState = EQueryState::Querying;
+		}
+	}
+
+	static void OnSourceControlOperationComplete(const FSourceControlOperationRef& InOperation, ECommandResult::Type InResult)
+	{
+		QueryState = EQueryState::Queried;
+	}
+
+	static TSharedRef<SWidget> GenerateSourceControlMenuContent()
+	{
+		UToolMenu* SourceControlMenu = UToolMenus::Get()->RegisterMenu("StatusBar.ToolBar.SourceControl", NAME_None, EMultiBoxType::Menu, false);
+
+		FToolMenuSection& Section = SourceControlMenu->AddSection("SourceControlActions", LOCTEXT("SourceControlMenuHeadingActions", "Actions"));
+
+		Section.AddDynamicEntry("ConnectToSourceControl", FNewToolMenuSectionDelegate::CreateLambda([](FToolMenuSection& InSection)
+			{
+				ISourceControlModule& SourceControlModule = ISourceControlModule::Get();
+				if (ISourceControlModule::Get().IsEnabled() && ISourceControlModule::Get().GetProvider().IsAvailable())
+				{
+					InSection.AddMenuEntry(
+						FSourceControlCommands::Get().ChangeSourceControlSettings,
+						TAttribute<FText>(),
+						TAttribute<FText>(),
+						FSlateIcon(FEditorStyle::GetStyleSetName(), "SourceControl.Actions.ChangeSettings")
+					);
+				}
+				else
+				{
+					InSection.AddMenuEntry(
+						FSourceControlCommands::Get().ConnectToSourceControl,
+						TAttribute<FText>(),
+						TAttribute<FText>(),
+						FSlateIcon(FEditorStyle::GetStyleSetName(), "SourceControl.Actions.Connect")
+					);
+				}
+			}));
+
+		Section.AddSeparator("SourceControlConnectionSeparator");
+
+		Section.AddMenuEntry(
+			FSourceControlCommands::Get().CheckOutModifiedFiles,
+			TAttribute<FText>(),
+			TAttribute<FText>(),
+			FSlateIcon(FEditorStyle::GetStyleSetName(), "SourceControl.Actions.CheckOut")
+		);
+
+
+		Section.AddMenuEntry(
+			FSourceControlCommands::Get().SubmitToSourceControl,
+			TAttribute<FText>(),
+			TAttribute<FText>(),
+			FSlateIcon(FEditorStyle::GetStyleSetName(), "SourceControl.Actions.Submit")
+		);
+
+		return UToolMenus::Get()->GenerateWidget("StatusBar.ToolBar.SourceControl", FToolMenuContext(FSourceControlCommands::ActionList));
+	}
+
+	static FText GetSourceControlStatusText()
+	{
+		if (QueryState == EQueryState::Querying)
+		{
+			return LOCTEXT("SourceControlStatus_Querying", "Contacting Server....");
+		}
+		else
+		{
+			ISourceControlModule& SourceControlModule = ISourceControlModule::Get();
+			if (SourceControlModule.IsEnabled())
+			{
+				if (!SourceControlModule.GetProvider().IsAvailable())
+				{
+					return LOCTEXT("SourceControlStatus_Error", "Server Unavailable");
+				}
+				else
+				{
+					return LOCTEXT("SourceControlStatus_Available", "Source Control");
+				}
+			}
+			else
+			{
+				return LOCTEXT("SourceControlStatus_Error", "Source Control Off");
+			}
+		}
+	}
+	static FText GetSourceControlTooltip()
+	{
+		if (QueryState == EQueryState::Querying)
+		{
+			return LOCTEXT("SourceControlUnknown", "Source control status is unknown");
+		}
+		else
+		{
+			return ISourceControlModule::Get().GetProvider().GetStatusText();
+		}
+	}
+
+	static const FSlateBrush* GetSourceControlIcon()
+	{
+
+		if (QueryState == EQueryState::Querying)
+		{
+			static const FSlateBrush* QueryBrush = FAppStyle::Get().GetBrush("SourceControl.StatusIcon.Unknown");
+			return QueryBrush;
+		}
+		else
+		{
+			ISourceControlModule& SourceControlModule = ISourceControlModule::Get();
+			if (SourceControlModule.IsEnabled())
+			{
+				if (!SourceControlModule.GetProvider().IsAvailable())
+				{
+					static const FSlateBrush* ErrorBrush = FAppStyle::Get().GetBrush("SourceControl.StatusIcon.Error");
+					return ErrorBrush;
+				}
+				else
+				{
+					static const FSlateBrush* OnBrush = FAppStyle::Get().GetBrush("SourceControl.StatusIcon.On");
+					return OnBrush;
+				}
+			}
+			else
+			{
+				static const FSlateBrush* OffBrush = FAppStyle::Get().GetBrush("SourceControl.StatusIcon.Off");
+				return OffBrush;
+			}
+		}
+	}
+
+	static FReply OnSourceControlStatusButtonClicked(TWeakPtr<SMenuAnchor> MenuAnchor)
+	{
+		if (TSharedPtr<SMenuAnchor> MenuAnchorPinned = MenuAnchor.Pin())
+		{
+			MenuAnchorPinned->SetIsOpen(!MenuAnchorPinned->IsOpen());
+		}
+
+		return FReply::Handled();
+	}
+
+	static TSharedRef<SWidget> MakeSourceControlStatusWidget()
+	{
+		TSharedPtr<SMenuAnchor> MenuAnchor;
+
+		return
+			SNew(SComboButton)
+			.ContentPadding(FMargin(6.0f,0.0f))
+			.ToolTipText_Static(&FSourceControlMenuHelpers::GetSourceControlTooltip)
+			.MenuPlacement(MenuPlacement_AboveAnchor)
+			.ComboButtonStyle(&FAppStyle::Get().GetWidgetStyle<FComboButtonStyle>("StatusBar.StatusBarComboButton"))
+			.ButtonContent()
+			[
+				SNew(SHorizontalBox)
+				+ SHorizontalBox::Slot()
+				.AutoWidth()
+				.VAlign(VAlign_Center)
+				.HAlign(HAlign_Center)
+				[
+					SNew(SImage)
+					.Image_Static(&FSourceControlMenuHelpers::GetSourceControlIcon)
+				]
+				+ SHorizontalBox::Slot()
+				.AutoWidth()
+				.VAlign(VAlign_Center)
+				.Padding(FMargin(5, 5, 0, 5))
+				[
+					SNew(STextBlock)
+					.TextStyle(&FAppStyle::Get().GetWidgetStyle<FTextBlockStyle>("NormalText"))
+					.Text_Static(&FSourceControlMenuHelpers::GetSourceControlStatusText)
+				]
+			]
+			.MenuContent()
+			[
+				GenerateSourceControlMenuContent()
+			];
+	}
+};
+
+FSourceControlMenuHelpers::EQueryState FSourceControlMenuHelpers::QueryState = FSourceControlMenuHelpers::EQueryState::NotQueried;
+
+
+class SStatusBar : public SCompoundWidget
+{
+	SLATE_BEGIN_ARGS(SStatusBar)
+	{}
+	SLATE_END_ARGS()
+
+public:
+	void Construct(const FArguments& InArgs)
+	{
+		UpArrow = FAppStyle::Get().GetBrush("StatusBar.ContentBrowserUp");
+		DownArrow = FAppStyle::Get().GetBrush("StatusBar.ContentBrowserDown");
+
+		const FSlateBrush* StatusBarBackground = FAppStyle::Get().GetBrush("StatusBar.Background");
+		ChildSlot
+		[
+			SNew(SBox)
+			.HeightOverride(FAppStyle::Get().GetFloat("StatusBar.Height"))
+			[
+				SNew(SHorizontalBox)
+				+ SHorizontalBox::Slot()
+				.Padding(1.0f, 0.0f)
+				.AutoWidth()
+				[
+					SNew(SBorder)
+					.Padding(FMargin(2.0f, 0.0f))
+					.BorderImage(StatusBarBackground)
+					.VAlign(VAlign_Center)
+					[			
+						SNew(SButton)
+						.ButtonStyle(&FAppStyle::Get().GetWidgetStyle<FButtonStyle>("StatusBar.StatusBarButton"))
+						[
+							SNew(SHorizontalBox)
+							+ SHorizontalBox::Slot()
+							.HAlign(HAlign_Left)
+							.VAlign(VAlign_Center)
+							.Padding(2.0f)
+							.AutoWidth()
+							[
+								SNew(SImage)
+								.ColorAndOpacity(FSlateColor::UseForeground())
+								.Image(this, &SStatusBar::GetContentBrowserExpandArrowImage)
+							]
+							+ SHorizontalBox::Slot()
+							.Padding(2.0f)
+							.HAlign(HAlign_Left)
+							.VAlign(VAlign_Center)
+							.AutoWidth()
+							[
+								SNew(SImage)
+								.ColorAndOpacity(FSlateColor::UseForeground())
+								.Image(FAppStyle::Get().GetBrush("StatusBar.ContentBrowserIcon"))
+							]
+							+ SHorizontalBox::Slot()
+							.VAlign(VAlign_Center)
+							.Padding(2.0f)
+							[
+								SNew(STextBlock)
+								.TextStyle(&FAppStyle::Get().GetWidgetStyle<FTextBlockStyle>("NormalText"))
+								.Text(LOCTEXT("StatusBar_ContentBrowserButton", "Content Browser"))
+							]
+						]
+					]
+				] 
+				+ SHorizontalBox::Slot()
+				.FillWidth(1.0f)
+				.Padding(1.0f, 0.0f)
+				[
+					SNew(SBorder)
+					.Padding(0.0f)
+					.BorderImage(StatusBarBackground)
+					.VAlign(VAlign_Center)
+					[
+						SNullWidget::NullWidget
+					]
+				]
+				+ SHorizontalBox::Slot()
+				.HAlign(HAlign_Right)
+				.AutoWidth()
+				.Padding(1.0f, 0.0f)
+				[
+					SNew(SBorder)
+					.Padding(0.0f)
+					.BorderImage(StatusBarBackground)
+					.VAlign(VAlign_Center)
+					[
+						MakeStatusBarToolBarWidget()
+					]
+				]
+			]
+		];
+	}
+
+private:
+	const FSlateBrush* GetContentBrowserExpandArrowImage() const
+	{
+		return DownArrow;
+	}
+
+	TSharedRef<SWidget> MakeStatusBarToolBarWidget()
+	{
+		RegisterStatusBarMenu();
+
+		FToolMenuContext MenuContext;
+		RegisterSourceControlStatus();
+
+
+		return UToolMenus::Get()->GenerateWidget("StatusBar.ToolBar", MenuContext);
+	}
+
+	void RegisterStatusBarMenu()
+	{
+		static const FName StatusBarName("StatusBar.ToolBar");
+		UToolMenus* ToolMenus = UToolMenus::Get();
+		if (ToolMenus->IsMenuRegistered(StatusBarName))
+		{
+			return;
+		}
+
+		UToolMenu* ToolBar = ToolMenus->RegisterMenu(StatusBarName, NAME_None, EMultiBoxType::SlimHorizontalToolBar);
+		ToolBar->StyleName = "StatusBarToolBar";
+
+	}
+
+	void RegisterSourceControlStatus()
+	{
+
+		// Source Control preferences
+		FSourceControlMenuHelpers::CheckSourceControlStatus();
+
+		{
+			UToolMenu* SourceControlMenu = UToolMenus::Get()->ExtendMenu("StatusBar.ToolBar");
+			FToolMenuSection& Section = SourceControlMenu->FindOrAddSection("SourceControl");
+
+			Section.AddEntry(
+				FToolMenuEntry::InitWidget(
+					"SourceControl",
+					FSourceControlMenuHelpers::MakeSourceControlStatusWidget(),
+					FText::GetEmpty(),
+					true,
+					false
+				));
+		}
+	}
+private:
+	const FSlateBrush* UpArrow;
+	const FSlateBrush* DownArrow;
+};
+
+
+void UStatusBarSubsystem::Initialize(FSubsystemCollectionBase& Collection)
+{
+	FSourceControlCommands::Register();
+}
+
+void UStatusBarSubsystem::Deinitialize()
+{
+	FSourceControlCommands::Unregister();
+}
+
+TSharedRef<SWidget> UStatusBarSubsystem::MakeStatusBarWidget(FName StatusBarName)
+{
+	TSharedRef<SStatusBar> StatusBar = SNew(SStatusBar);
+
+	StatusBars.Add(StatusBarName, StatusBar);
+
+	return StatusBar;
+}
diff --git a/Engine/Source/Editor/UnrealEd/Public/Subsystems/StatusBarSubsystem.h b/Engine/Source/Editor/UnrealEd/Public/Subsystems/StatusBarSubsystem.h
new file mode 100644
--- /dev/null
+++ b/Engine/Source/Editor/UnrealEd/Public/Subsystems/StatusBarSubsystem.h
@@ -0,0 +1,37 @@
+// Copyright Epic Games, Inc. All Rights Reserved.
+
+#pragma once
+
+#include "CoreMinimal.h"
+#include "EditorSubsystem.h"
+
+#include "StatusBarSubsystem.generated.h"
+
+class SStatusBar;
+
+UCLASS()
+class UNREALED_API UStatusBarSubsystem : public UEditorSubsystem
+{
+	GENERATED_BODY()
+
+public:
+
+	/**
+	 *	Prepares for use
+	 */
+	virtual void Initialize(FSubsystemCollectionBase& Collection) override final;
+
+	/**
+	 *	Internal cleanup
+	 */
+	virtual void Deinitialize() override final;
+
+	/** 
+	 * Creates a new instance of a status bar widget
+	 *
+	 * @param StatusBarName	The name of the status bar for updating it later.
+	 */
+	TSharedRef<SWidget> MakeStatusBarWidget(FName StatusBarName);
+private:
+	TMap<FName, TWeakPtr<SStatusBar>> StatusBars;
+};
diff --git a/Engine/Source/Runtime/Slate/Private/Framework/Notifications/NotificationManager.cpp b/Engine/Source/Runtime/Slate/Private/Framework/Notifications/NotificationManager.cpp
--- a/Engine/Source/Runtime/Slate/Private/Framework/Notifications/NotificationManager.cpp
+++ b/Engine/Source/Runtime/Slate/Private/Framework/Notifications/NotificationManager.cpp
@@ -37,9 +37,12 @@ void FSlateNotificationManager::FRegionalNotificationList::RemoveDeadNotificatio
 
 void FSlateNotificationManager::FRegionalNotificationList::Arrange()
 {
+	// If the app has a status bar push up notifications so that they don't overlap the status bar
+	static const float StausBarHeight = FAppStyle::Get().GetFloat("StatusBar.Height", nullptr, 0.0f);
+
 	FVector2D AnchorPoint(
 		Region.Right - NotificationManagerConstants::NotificationOffset.X,
-		Region.Bottom - NotificationManagerConstants::NotificationOffset.Y );
+		Region.Bottom - (NotificationManagerConstants::NotificationOffset.Y + StausBarHeight));
 
 	for (int32 ListIndex = Notifications.Num() - 1; ListIndex >= 0; --ListIndex)
 	{
@@ -141,7 +144,7 @@ TSharedPtr<SNotificationItem> FSlateNotificationManager::AddNotification(const F
 	if( FSlateApplication::IsInitialized() )
 	{
 		FSlateRect PreferredWorkArea;
-		// Retrieve the main editor window to display the notication in, otherwise use the preferred work area
+		// Retrieve the main editor window to display the notification in, otherwise use the preferred work area
 		if (RootWindowPtr.IsValid())
 		{
 			PreferredWorkArea = FSlateApplication::Get().GetWorkArea(RootWindowPtr.Pin()->GetRectInScreen());
