Non-unity build fixes.

#rnx

diff --git a/Engine/Source/Editor/LevelEditor/Private/LevelEditorToolBar.cpp b/Engine/Source/Editor/LevelEditor/Private/LevelEditorToolBar.cpp
--- a/Engine/Source/Editor/LevelEditor/Private/LevelEditorToolBar.cpp
+++ b/Engine/Source/Editor/LevelEditor/Private/LevelEditorToolBar.cpp
@@ -14,11 +14,9 @@
 #include "Widgets/Input/SCheckBox.h"
 #include "Widgets/Input/SSpinBox.h"
 #include "Widgets/Input/SSlider.h"
+#include "Widgets/Layout/SSpacer.h"
 #include "EditorStyleSet.h"
-#include "ISourceControlOperation.h"
-#include "SourceControlOperations.h"
-#include "ISourceControlProvider.h"
-#include "ISourceControlModule.h"
+
 #include "Settings/EditorExperimentalSettings.h"
 #include "GameMapsSettings.h"
 #include "GameFramework/PlayerController.h"
@@ -62,6 +60,7 @@
 #include "ToolMenus.h"
 #include "Subsystems/AssetEditorSubsystem.h"
 #include "LevelEditorModesActions.h"
+#include "ISourceControlModule.h"
 
 static TAutoConsoleVariable<int32> CVarAllowMatineeActors(
 	TEXT("Matinee.AllowMatineeActors"),
@@ -1174,8 +1173,7 @@ void FLevelEditorToolBar::RegisterLevelEditorToolBar( const TSharedRef<FUIComman
 {
 	RegisterSourceControlMenu();
 	RegisterCinematicsMenu();
-	RegisterBuildMenu();
-	RegisterEditorModesMenu();
+
 #if WITH_LIVE_CODING
 	RegisterCompileMenu();
 #endif
@@ -1185,183 +1183,168 @@ void FLevelEditorToolBar::RegisterLevelEditorToolBar( const TSharedRef<FUIComman
 
 #define LOCTEXT_NAMESPACE "LevelEditorToolBar"
 
-	UToolMenu* Toolbar = UToolMenus::Get()->RegisterMenu("LevelEditor.LevelEditorToolBar", NAME_None, EMultiBoxType::ToolBar);
-
+	UToolMenu* AssetsToolBar = UToolMenus::Get()->RegisterMenu("LevelEditor.LevelEditorToolBar.AssetsToolBar", NAME_None, EMultiBoxType::SlimHorizontalToolBar);
 	{
-		FToolMenuSection& Section = Toolbar->AddSection("File");
+		{
+			FToolMenuSection& Section = AssetsToolBar->AddSection("File");
 
-		// Save All Levels
-		Section.AddEntry(FToolMenuEntry::InitToolBarButton(
-			FLevelEditorCommands::Get().Save,
-			TAttribute<FText>(),
-			TAttribute<FText>(),
-			FSlateIcon(FEditorStyle::GetStyleSetName(), "AssetEditor.SaveAsset"),
-			NAME_None,
-			FName("SaveAllLevels")
-		));
+			// Save All Levels
+			Section.AddEntry(FToolMenuEntry::InitToolBarButton(
+				FLevelEditorCommands::Get().Save,
+				TAttribute<FText>(),
+				TAttribute<FText>(),
+				FSlateIcon(FEditorStyle::GetStyleSetName(), "AssetEditor.SaveAsset"),
+				NAME_None,
+				FName("SaveAllLevels")
+				));
+		}
 
-		// Source control buttons
 		{
-			enum EQueryState
-			{
-				NotQueried,
-				Querying,
-				Queried,
-			};
+			FToolMenuSection& Section = AssetsToolBar->AddSection("Content");
+			Section.AddEntry(FToolMenuEntry::InitToolBarButton(FLevelEditorCommands::Get().OpenContentBrowser, LOCTEXT("ContentBrowser_Override", "Content"), TAttribute<FText>(), TAttribute<FSlateIcon>(), "LevelToolbarContent"));
 
-			static EQueryState QueryState = EQueryState::NotQueried;
+			Section.AddEntry(FToolMenuEntry::InitComboButton(
+				"OpenBlueprint",
+				FUIAction(),
+				FOnGetContent::CreateStatic(&FLevelEditorToolBar::GenerateOpenBlueprintMenuContent, InCommandList, TWeakPtr<SLevelEditor>(InLevelEditor)),
+				LOCTEXT("OpenBlueprint_Label", "Blueprints"),
+				LOCTEXT("OpenBlueprint_ToolTip", "List of world Blueprints available to the user for editing or creation."),
+				FSlateIcon(FEditorStyle::GetStyleSetName(), "LevelEditor.OpenLevelBlueprint")
+			));
 
-			struct FSourceControlStatus
-			{
-				static void CheckSourceControlStatus()
-				{
-					ISourceControlModule& SourceControlModule = ISourceControlModule::Get();
-					if (SourceControlModule.IsEnabled())
-					{
-						SourceControlModule.GetProvider().Execute(ISourceControlOperation::Create<FConnect>(),
-																  EConcurrency::Asynchronous,
-																  FSourceControlOperationComplete::CreateStatic(&FSourceControlStatus::OnSourceControlOperationComplete));
-						QueryState = EQueryState::Querying;
-					}
-				}
+			Section.AddEntry(FToolMenuEntry::InitComboButton(
+				"EditCinematics",
+				FUIAction(),
+				FOnGetContent::CreateStatic(&FLevelEditorToolBar::GenerateCinematicsMenuContent, InCommandList, TWeakPtr<SLevelEditor>(InLevelEditor)),
+				LOCTEXT("EditCinematics_Label", "Cinematics"),
+				LOCTEXT("EditCinematics_Tooltip", "Displays a list of Level Sequence objects to open in their respective editors"),
+				FSlateIcon(FEditorStyle::GetStyleSetName(), "LevelEditor.EditMatinee")
+			));
 
-				static void OnSourceControlOperationComplete(const FSourceControlOperationRef& InOperation, ECommandResult::Type InResult)
-				{
-					QueryState = EQueryState::Queried;
-				}
+			Section.AddEntry(FToolMenuEntry::InitToolBarButton(FLevelEditorCommands::Get().ToggleVR, LOCTEXT("ToggleVR", "VR Mode")));
+		}
 
-				static FText GetSourceControlTooltip()
-				{
-					if (QueryState == EQueryState::Querying)
-					{
-						return LOCTEXT("SourceControlUnknown", "Source control status is unknown");
-					}
-					else
-					{
-						return ISourceControlModule::Get().GetProvider().GetStatusText();
-					}
-				}
+		//if(FCoreStyle::IsStarshipStyle())
+		{
+			FToolMenuSection& Section = AssetsToolBar->AddSection("Compile");
 
-				static FSlateIcon GetSourceControlIcon()
+			Section.AddDynamicEntry("CompilerAvailable", FNewToolMenuSectionDelegate::CreateLambda([](FToolMenuSection& InSection)
 				{
-					if (QueryState == EQueryState::Querying)
+					// Only show the compile options on machines with the solution (assuming they can build it)
+					if (FSourceCodeNavigation::IsCompilerAvailable())
 					{
-						return FSlateIcon(FEditorStyle::GetStyleSetName(), "LevelEditor.SourceControl.Unknown");
-					}
-					else
-					{
-						ISourceControlModule& SourceControlModule = ISourceControlModule::Get();
-						if (SourceControlModule.IsEnabled())
-						{
-							if (!SourceControlModule.GetProvider().IsAvailable())
-							{
-								return FSlateIcon(FEditorStyle::GetStyleSetName(), "LevelEditor.SourceControl.Problem");
-							}
-							else
-							{
-								return FSlateIcon(FEditorStyle::GetStyleSetName(), "LevelEditor.SourceControl.On");
-							}
-						}
-						else
-						{
-							return FSlateIcon(FEditorStyle::GetStyleSetName(), "LevelEditor.SourceControl.Off");
-						}
-					}
-				}
-			};
-
-			FSourceControlStatus::CheckSourceControlStatus();
+						// Since we can always add new code to the project, only hide these buttons if we haven't done so yet
+						InSection.AddEntry(FToolMenuEntry::InitToolBarButton(
+							"CompileButton",
+							FUIAction(
+								FExecuteAction::CreateStatic(&FLevelEditorActionCallbacks::RecompileGameCode_Clicked),
+								FCanExecuteAction::CreateStatic(&FLevelEditorActionCallbacks::Recompile_CanExecute),
+								FIsActionChecked(),
+								FIsActionButtonVisible::CreateStatic(FLevelEditorActionCallbacks::CanShowSourceCodeActions)),
+							LOCTEXT("CompileMenuButton", "Compile"),
+							FLevelEditorCommands::Get().RecompileGameCode->GetDescription(),
+							FSlateIcon(FEditorStyle::GetStyleSetName(), "LevelEditor.Recompile")
+						));
 
-			Section.AddEntry(FToolMenuEntry::InitComboButton(
-				"SourceControl",
-				FUIAction(),
-				FOnGetContent::CreateStatic(&FLevelEditorToolBar::GenerateSourceControlMenu, InCommandList, TWeakPtr<SLevelEditor>(InLevelEditor)),
-				LOCTEXT("SourceControl_Label", "Source Control"),
-				TAttribute<FText>::Create(&FSourceControlStatus::GetSourceControlTooltip),
-				TAttribute<FSlateIcon>::Create(&FSourceControlStatus::GetSourceControlIcon),
-				false
-				));
+#if WITH_LIVE_CODING
+						InSection.AddEntry(FToolMenuEntry::InitComboButton(
+							"CompileComboButton",
+							FUIAction(
+								FExecuteAction(),
+								FCanExecuteAction(),
+								FIsActionChecked(),
+								FIsActionButtonVisible::CreateStatic(FLevelEditorActionCallbacks::CanShowSourceCodeActions)),
+							FNewToolMenuChoice(),
+							LOCTEXT("CompileCombo_Label", "Compile Options"),
+							LOCTEXT("CompileComboToolTip", "Compile options menu"),
+							FSlateIcon(FEditorStyle::GetStyleSetName(), "LevelEditor.Recompile"),
+							true
+						));
+#endif
+					}
+				}));
 		}
+		
+
 	}
 
+	UToolMenu* ModesToolbar = UToolMenus::Get()->RegisterMenu("LevelEditor.LevelEditorToolBar.ModesToolBar", NAME_None, EMultiBoxType::SlimHorizontalToolBar);
 	{
+		FToolMenuSection& Section = ModesToolbar->AddSection("EditorModes");
+
+		Section.AddSeparator(NAME_None);
 
-		struct FEditorModesStatus
+		const FLevelEditorModule& LevelEditorModule = FModuleManager::LoadModuleChecked<FLevelEditorModule>("LevelEditor");
+		const FLevelEditorModesCommands& Commands = LevelEditorModule.GetLevelEditorModesCommands();
+
+		TArray<FEditorModeInfo, TInlineAllocator<1>> DefaultModes;
+
+		TArray<FEditorModeInfo, TInlineAllocator<10>> NonDefaultModes;
+
+		for (const FEditorModeInfo& Mode : FEditorModeRegistry::Get().GetSortedModeInfo())
 		{
-			static FSlateIcon GetEditorModesIcon()
+			// If the mode isn't visible don't create a menu option for it.
+			if (!Mode.bVisible)
 			{
-				for (const FEditorModeInfo& Mode : FEditorModeRegistry::Get().GetSortedModeInfo())
-				{
-					if (!Mode.bVisible)
-					{
-						continue;
-					}
+				continue;
+			}
 
-					if (GLevelEditorModeTools().IsModeActive(Mode.ID))
-					{
-						// if its a default mode, use the default tool icon
-						if (GLevelEditorModeTools().IsDefaultMode(Mode.ID))
-						{
-							return FSlateIcon(FEditorStyle::GetStyleSetName(), "LevelEditor.EditorModes");
-						}
+			if (GLevelEditorModeTools().IsDefaultMode(Mode.ID))
+			{
+				DefaultModes.Add(Mode.ID);
+			}
+			else
+			{
+				NonDefaultModes.Add(Mode.ID);
+			}
 
-						FName EditorModeCommandName = FName(*(FString("EditorMode.") + Mode.ID.ToString()));
-						const FLevelEditorModule& LevelEditorModule = FModuleManager::LoadModuleChecked<FLevelEditorModule>("LevelEditor");
-						const FLevelEditorModesCommands& Commands = LevelEditorModule.GetLevelEditorModesCommands();
+		}
 
-						TSharedPtr<FUICommandInfo> EditorModeCommand =
-							FInputBindingManager::Get().FindCommandInContext(Commands.GetContextName(), EditorModeCommandName);
-						if (EditorModeCommand.IsValid())
-						{
-							return EditorModeCommand->GetIcon();
-						}
-					}
-					
+		auto BuildEditorModes =
+			[&Commands, &Section](TArrayView<FEditorModeInfo> Modes)
+		{
+			for (const FEditorModeInfo& Mode : Modes)
+			{
+				FName EditorModeCommandName = FName(*(FString("EditorMode.") + Mode.ID.ToString()));
+
+				TSharedPtr<FUICommandInfo> EditorModeCommand =
+					FInputBindingManager::Get().FindCommandInContext(Commands.GetContextName(), EditorModeCommandName);
+
+				// If a command isn't yet registered for this mode, we need to register one.
+				if (!EditorModeCommand.IsValid())
+				{
+					continue;
 				}
-				return FSlateIcon(FEditorStyle::GetStyleSetName(), "LevelEditor.EditorModes");
 
+				Section.AddEntry(FToolMenuEntry::InitToolBarButton(EditorModeCommand));
 			}
+
 		};
 
-		FToolMenuSection& Section = Toolbar->AddSection("Modes");
-		Section.AddEntry(FToolMenuEntry::InitComboButton(
-			"EditorModes",
-			FUIAction(),
-			FOnGetContent::CreateStatic(&FLevelEditorToolBar::GenerateEditorModesMenu, InCommandList, TWeakPtr<SLevelEditor>(InLevelEditor)),
-			LOCTEXT("EditorModes_Label", "Modes"),
-			LOCTEXT("EditorModes_Tooltip", "Displays a list of editing modes that can be toggled"),
-			TAttribute<FSlateIcon>::Create(&FEditorModesStatus::GetEditorModesIcon)
-		));
+		// Build Default Modes first
+		BuildEditorModes(MakeArrayView(DefaultModes));
+
+		// Build non-default modes second
+		BuildEditorModes(MakeArrayView(NonDefaultModes));
 	}
 
+	UToolMenu* SettingsToolBar = UToolMenus::Get()->RegisterMenu("LevelEditor.LevelEditorToolBar.MarketplaceToolBar", NAME_None, EMultiBoxType::SlimHorizontalToolBar);
 	{
-		FToolMenuSection& Section = Toolbar->AddSection("Content");
-		Section.AddEntry(FToolMenuEntry::InitToolBarButton(FLevelEditorCommands::Get().OpenContentBrowser, LOCTEXT( "ContentBrowser_Override", "Content" ), TAttribute<FText>(), TAttribute<FSlateIcon>(), "LevelToolbarContent"));
-		if (FLauncherPlatformModule::Get()->CanOpenLauncher(true)) 
+		FToolMenuSection& SettingsSection = SettingsToolBar->AddSection("Settings");
+		if (FLauncherPlatformModule::Get()->CanOpenLauncher(true))
 		{
-			Section.AddEntry(FToolMenuEntry::InitToolBarButton(FLevelEditorCommands::Get().OpenMarketplace, LOCTEXT("Marketplace_Override", "Marketplace"), TAttribute<FText>(), TAttribute<FSlateIcon>(), "LevelToolbarMarketplace"));
+			SettingsSection.AddEntry(FToolMenuEntry::InitToolBarButton(FLevelEditorCommands::Get().OpenMarketplace, LOCTEXT("Marketplace_Override", "Marketplace"), TAttribute<FText>(), TAttribute<FSlateIcon>(), "LevelToolbarMarketplace"));
 		}
 	}
 
-	FToolMenuSection& SettingsSection = Toolbar->AddSection("Settings");
+	UToolMenu* PlayToolBar = UToolMenus::Get()->RegisterMenu("LevelEditor.LevelEditorToolBar.PlayToolBar", NAME_None, EMultiBoxType::SlimHorizontalToolBar);
 	{
-		SettingsSection.AddEntry(FToolMenuEntry::InitComboButton(
-			"LevelToolbarQuickSettings",
-			FUIAction(),
-			FOnGetContent::CreateStatic(&FLevelEditorToolBar::GenerateQuickSettingsMenu, InCommandList, TWeakPtr<SLevelEditor>(InLevelEditor)),
-			LOCTEXT("QuickSettingsCombo", "Settings"),
-			LOCTEXT("QuickSettingsCombo_ToolTip", "Project and Editor settings"),
-			FSlateIcon(FEditorStyle::GetStyleSetName(), "LevelEditor.GameSettings"),
-			false,
-			"LevelToolbarQuickSettings"
-			));
+		FToolMenuSection& PlaySection = PlayToolBar->AddSection("Play");
 
-	}
+		PlaySection.AddSeparator(NAME_None);
 
-	{
 		struct FPreviewModeFunctionality
 		{
-
 			static FText GetPreviewModeText()
 			{
 				EShaderPlatform ShaderPlatform = ShaderFormatToLegacyShaderPlatform(GEditor->PreviewPlatform.PreviewShaderPlatformName);
@@ -1450,94 +1433,28 @@ void FLevelEditorToolBar::RegisterLevelEditorToolBar( const TSharedRef<FUIComman
 			}
 		};
 
-		SettingsSection.AddEntry(FToolMenuEntry::InitToolBarButton( 
+		PlaySection.AddEntry(FToolMenuEntry::InitToolBarButton(
 			FLevelEditorCommands::Get().ToggleFeatureLevelPreview,
 			TAttribute<FText>::Create(&FPreviewModeFunctionality::GetPreviewModeText),
         	TAttribute<FText>::Create(&FPreviewModeFunctionality::GetPreviewModeTooltip),
         	TAttribute<FSlateIcon>::Create(&FPreviewModeFunctionality::GetPreviewModeIcon)
 			));
-	}
-
-	{
-		FToolMenuSection& Section = Toolbar->AddSection("Misc");
-		Section.AddEntry(FToolMenuEntry::InitComboButton(
-			"OpenBlueprint",
-			FUIAction(),
-			FOnGetContent::CreateStatic( &FLevelEditorToolBar::GenerateOpenBlueprintMenuContent, InCommandList, TWeakPtr<SLevelEditor>( InLevelEditor ) ),
-			LOCTEXT( "OpenBlueprint_Label", "Blueprints" ),
-			LOCTEXT( "OpenBlueprint_ToolTip", "List of world Blueprints available to the user for editing or creation." ),
-			FSlateIcon(FEditorStyle::GetStyleSetName(), "LevelEditor.OpenLevelBlueprint")
-			));
-
-		Section.AddEntry(FToolMenuEntry::InitComboButton(
-			"EditCinematics",
-			FUIAction(),
-			FOnGetContent::CreateStatic( &FLevelEditorToolBar::GenerateCinematicsMenuContent, InCommandList, TWeakPtr<SLevelEditor>( InLevelEditor ) ),
-			LOCTEXT( "EditCinematics_Label", "Cinematics" ),
-			LOCTEXT( "EditCinematics_Tooltip", "Displays a list of Level Sequence objects to open in their respective editors"),
-			FSlateIcon(FEditorStyle::GetStyleSetName(), "LevelEditor.EditMatinee") 
-			));
 
-		Section.AddEntry(FToolMenuEntry::InitToolBarButton( FLevelEditorCommands::Get().ToggleVR, LOCTEXT("ToggleVR", "VR Mode")) );
-	}
-
-	{
-		FToolMenuSection& Section = Toolbar->AddSection("Compile");
-		// Build			
-		Section.AddEntry(FToolMenuEntry::InitToolBarButton( FLevelEditorCommands::Get().Build, LOCTEXT("BuildAll", "Build")) );
+	
+		// Add the shared play-world commands that will be shown on the Kismet toolbar as well
+		FPlayWorldCommands::BuildToolbar(PlaySection, true);
 
-		// Build menu drop down
-		Section.AddEntry(FToolMenuEntry::InitComboButton(
-			"BuildComboButton",
+		PlaySection.AddEntry(FToolMenuEntry::InitComboButton(
+			"LevelToolbarQuickSettings",
 			FUIAction(),
-			FOnGetContent::CreateStatic( &FLevelEditorToolBar::GenerateBuildMenuContent, InCommandList, TWeakPtr<SLevelEditor>(InLevelEditor) ),
-			LOCTEXT( "BuildCombo_Label", "Build Options" ),
-			LOCTEXT( "BuildComboToolTip", "Build options menu" ),
-			FSlateIcon(FEditorStyle::GetStyleSetName(), "LevelEditor.Build"),
-			true));
-
-		Section.AddDynamicEntry("CompilerAvailable", FNewToolMenuSectionDelegate::CreateLambda([](FToolMenuSection& InSection)
-		{
-			// Only show the compile options on machines with the solution (assuming they can build it)
-			if (FSourceCodeNavigation::IsCompilerAvailable())
-			{
-				// Since we can always add new code to the project, only hide these buttons if we haven't done so yet
-				InSection.AddEntry(FToolMenuEntry::InitToolBarButton(
-					"CompileButton",
-					FUIAction(
-						FExecuteAction::CreateStatic(&FLevelEditorActionCallbacks::RecompileGameCode_Clicked),
-						FCanExecuteAction::CreateStatic(&FLevelEditorActionCallbacks::Recompile_CanExecute),
-						FIsActionChecked(),
-						FIsActionButtonVisible::CreateStatic(FLevelEditorActionCallbacks::CanShowSourceCodeActions)),
-					LOCTEXT("CompileMenuButton", "Compile"),
-					FLevelEditorCommands::Get().RecompileGameCode->GetDescription(),
-					FSlateIcon(FEditorStyle::GetStyleSetName(), "LevelEditor.Recompile")
-				));
-
-#if WITH_LIVE_CODING
-				InSection.AddEntry(FToolMenuEntry::InitComboButton(
-					"CompileComboButton",
-					FUIAction(
-						FExecuteAction(),
-						FCanExecuteAction(),
-						FIsActionChecked(),
-						FIsActionButtonVisible::CreateStatic(FLevelEditorActionCallbacks::CanShowSourceCodeActions)),
-					FNewToolMenuChoice(),
-					LOCTEXT("CompileCombo_Label", "Compile Options"),
-					LOCTEXT("CompileComboToolTip", "Compile options menu"),
-					FSlateIcon(FEditorStyle::GetStyleSetName(), "LevelEditor.Recompile"),
-					true
-				));
-#endif
-			}
-		}));
-	}
-
-	{
-		FToolMenuSection& Section = Toolbar->AddSection("Game");
+			FOnGetContent::CreateStatic(&FLevelEditorToolBar::GenerateQuickSettingsMenu, InCommandList, TWeakPtr<SLevelEditor>(InLevelEditor)),
+			LOCTEXT("QuickSettingsCombo", "Settings"),
+			LOCTEXT("QuickSettingsCombo_ToolTip", "Project and Editor settings"),
+			FSlateIcon(FEditorStyle::GetStyleSetName(), "LevelEditor.GameSettings"),
+			false,
+			"LevelToolbarQuickSettings"
+		));
 
-		// Add the shared play-world commands that will be shown on the Kismet toolbar as well
-		FPlayWorldCommands::BuildToolbar(Section, true);
 	}
 
 #undef LOCTEXT_NAMESPACE
@@ -1550,8 +1467,6 @@ void FLevelEditorToolBar::RegisterLevelEditorToolBar( const TSharedRef<FUIComman
  */
 TSharedRef< SWidget > FLevelEditorToolBar::MakeLevelEditorToolBar( const TSharedRef<FUICommandList>& InCommandList, const TSharedRef<SLevelEditor> InLevelEditor )
 {
-#define LOCTEXT_NAMESPACE "LevelEditorToolBar"
-
 	FLevelEditorModule& LevelEditorModule = FModuleManager::LoadModuleChecked<FLevelEditorModule>("LevelEditor");
 	FToolMenuContext MenuContext(InCommandList, LevelEditorModule.GetToolBarExtensibilityManager()->GetAllExtenders());
 	ULevelEditorMenuContext* LevelEditorMenuContext = NewObject<ULevelEditorMenuContext>();
@@ -1560,343 +1475,51 @@ TSharedRef< SWidget > FLevelEditorToolBar::MakeLevelEditorToolBar( const TShared
 
 	// Create the tool bar!
 	return
-		SNew( SBorder )
-		.Padding(0)
-		.BorderImage( FEditorStyle::GetBrush("NoBorder") )
-		.IsEnabled( FSlateApplication::Get().GetNormalExecutionAttribute() )
+		SNew(SOverlay)
+		+SOverlay::Slot()
 		[
-			UToolMenus::Get()->GenerateWidget("LevelEditor.LevelEditorToolBar", MenuContext)
+			SNew(SImage)
+			.Image(FCoreStyle::IsStarshipStyle() ? FCoreStyle::Get().GetBrush("SlimToolBar.Background") : FCoreStyle::Get().GetBrush("ToolPanel.GroupBorder"))
+		]
+		+ SOverlay::Slot()
+		[
+			SNew(SHorizontalBox)
+			.IsEnabled(FSlateApplication::Get().GetNormalExecutionAttribute())
+			+ SHorizontalBox::Slot()
+			.AutoWidth()
+			[
+				UToolMenus::Get()->GenerateWidget("LevelEditor.LevelEditorToolBar.AssetsToolBar", MenuContext)
+			]
+			
+			+ SHorizontalBox::Slot()
+			.AutoWidth()
+			[
+				SNew(SSpacer)
+				.Size(FVector2D(120.f, 1.0f))
+			]
+			
+			+ SHorizontalBox::Slot()
+			.HAlign(HAlign_Left)
+			[
+				UToolMenus::Get()->GenerateWidget("LevelEditor.LevelEditorToolBar.ModesToolBar", MenuContext)
+			]
+			+ SHorizontalBox::Slot()
+			.HAlign(HAlign_Right)
+			.FillWidth(.65)
+			[
+				UToolMenus::Get()->GenerateWidget("LevelEditor.LevelEditorToolBar.PlayToolBar", MenuContext)
+			]
+			/*+ SHorizontalBox::Slot()
+			.HAlign(HAlign_Right)
+			[
+				UToolMenus::Get()->GenerateWidget("LevelEditor.LevelEditorToolBar.MarketplaceToolBar", MenuContext)
+			]*/
 		];
-#undef LOCTEXT_NAMESPACE
 }
 
 
 
-TSharedRef< SWidget > FLevelEditorToolBar::GenerateBuildMenuContent( TSharedRef<FUICommandList> InCommandList, TWeakPtr<SLevelEditor> InLevelEditor)
-{
-	// Get all menu extenders for this context menu from the level editor module
-	FLevelEditorModule& LevelEditorModule = FModuleManager::GetModuleChecked<FLevelEditorModule>( TEXT("LevelEditor") );
-	TSharedPtr<FExtender> MenuExtender = LevelEditorModule.AssembleExtenders(InCommandList, LevelEditorModule.GetAllLevelEditorToolbarBuildMenuExtenders());
-	FToolMenuContext MenuContext(InCommandList, MenuExtender);
-
-	ULevelEditorMenuContext* LevelEditorMenuContext = NewObject<ULevelEditorMenuContext>();
-	LevelEditorMenuContext->LevelEditor = InLevelEditor;
-	MenuContext.AddObject(LevelEditorMenuContext);
-
-	return UToolMenus::Get()->GenerateWidget("LevelEditor.LevelEditorToolBar.BuildComboButton", MenuContext);
-}
-
-void FLevelEditorToolBar::RegisterBuildMenu()
-{
-#define LOCTEXT_NAMESPACE "LevelToolBarBuildMenu"
 
-	static const FName BaseMenuName = "LevelEditor.LevelEditorToolBar.BuildComboButton";
-	UToolMenu* Menu = UToolMenus::Get()->RegisterMenu(BaseMenuName);
-
-	struct FLightingMenus
-	{
-	public:
-
-		static void RegisterMenus(const FName InBaseMenuName)
-		{
-			FLightingMenus::RegisterLightingQualityMenu(InBaseMenuName);
-			FLightingMenus::RegisterLightingInfoMenu(InBaseMenuName);
-		}
-
-	private:
-
-		/** Generates a lighting quality sub-menu */
-		static void RegisterLightingQualityMenu(const FName InBaseMenuName)
-		{
-			UToolMenu* SubMenu = UToolMenus::Get()->RegisterMenu(UToolMenus::JoinMenuPaths(InBaseMenuName, "LightingQuality"));
-
-			{
-				FToolMenuSection& Section = SubMenu->AddSection("LevelEditorBuildLightingQuality", LOCTEXT( "LightingQualityHeading", "Quality Level" ) );
-				Section.AddMenuEntry( FLevelEditorCommands::Get().LightingQuality_Production );
-				Section.AddMenuEntry( FLevelEditorCommands::Get().LightingQuality_High );
-				Section.AddMenuEntry( FLevelEditorCommands::Get().LightingQuality_Medium );
-				Section.AddMenuEntry( FLevelEditorCommands::Get().LightingQuality_Preview );
-			}
-		}
-
-		/** Generates a lighting density sub-menu */
-		static void RegisterLightingDensityMenu(const FName InBaseMenuName)
-		{
-			UToolMenu* SubMenu = UToolMenus::Get()->RegisterMenu(UToolMenus::JoinMenuPaths(InBaseMenuName, "LightingDensity"));
-
-			{
-				FToolMenuSection& Section = SubMenu->AddSection("LevelEditorBuildLightingDensity", LOCTEXT( "LightingDensityHeading", "Density Rendering" ) );
-				TSharedRef<SWidget> Ideal =		SNew(SHorizontalBox)
-												+SHorizontalBox::Slot()
-												.Padding( FMargin( 27.0f, 0.0f, 0.0f, 0.0f ) )
-												.FillWidth(1.0f)
-												[
-													SNew(SSpinBox<float>)
-													.MinValue(0.f)
-													.MaxValue(100.f)
-													.Value(FLevelEditorActionCallbacks::GetLightingDensityIdeal())
-													.OnValueChanged_Static(&FLevelEditorActionCallbacks::SetLightingDensityIdeal)
-												];
-				Section.AddEntry(FToolMenuEntry::InitWidget("Ideal", Ideal, LOCTEXT("LightingDensity_Ideal","Ideal Density")));
-				
-				TSharedRef<SWidget> Maximum =	SNew(SHorizontalBox)
-												+SHorizontalBox::Slot()
-												.FillWidth(1.0f)
-												[
-													SNew(SSpinBox<float>)
-													.MinValue(0.01f)
-													.MaxValue(100.01f)
-													.Value(FLevelEditorActionCallbacks::GetLightingDensityMaximum())
-													.OnValueChanged_Static(&FLevelEditorActionCallbacks::SetLightingDensityMaximum)
-												];
-				Section.AddEntry(FToolMenuEntry::InitWidget("Maximum", Maximum, LOCTEXT("LightingDensity_Maximum","Maximum Density")));
-
-				TSharedRef<SWidget> ClrScale =	SNew(SHorizontalBox)
-												+SHorizontalBox::Slot()
-												.Padding( FMargin( 35.0f, 0.0f, 0.0f, 0.0f ) )
-												.FillWidth(1.0f)
-												[
-													SNew(SSpinBox<float>)
-													.MinValue(0.f)
-													.MaxValue(10.f)
-													.Value(FLevelEditorActionCallbacks::GetLightingDensityColorScale())
-													.OnValueChanged_Static(&FLevelEditorActionCallbacks::SetLightingDensityColorScale)
-												];
-				Section.AddEntry(FToolMenuEntry::InitWidget("ColorScale", ClrScale, LOCTEXT("LightingDensity_ColorScale","Color Scale")));
-
-				TSharedRef<SWidget> GrayScale =	SNew(SHorizontalBox)
-												+SHorizontalBox::Slot()
-												.Padding( FMargin( 11.0f, 0.0f, 0.0f, 0.0f ) )
-												.FillWidth(1.0f)
-												[
-													SNew(SSpinBox<float>)
-													.MinValue(0.f)
-													.MaxValue(10.f)
-													.Value(FLevelEditorActionCallbacks::GetLightingDensityGrayscaleScale())
-													.OnValueChanged_Static(&FLevelEditorActionCallbacks::SetLightingDensityGrayscaleScale)
-												];
-				Section.AddEntry(FToolMenuEntry::InitWidget("GrayscaleScale", GrayScale, LOCTEXT("LightingDensity_GrayscaleScale","Grayscale Scale")));
-
-				Section.AddMenuEntry( FLevelEditorCommands::Get().LightingDensity_RenderGrayscale );
-			}
-		}
-
-		/** Generates a lighting resolution sub-menu */
-		static void RegisterLightingResolutionMenu(const FName InBaseMenuName)
-		{
-			UToolMenu* SubMenu = UToolMenus::Get()->RegisterMenu(UToolMenus::JoinMenuPaths(InBaseMenuName, "LightingResolution"));
-
-			{
-				FToolMenuSection& Section = SubMenu->AddSection("LevelEditorBuildLightingResolution1", LOCTEXT( "LightingResolutionHeading1", "Primitive Types" ) );
-				TSharedRef<SWidget> Meshes =	SNew(SHorizontalBox)
-												+SHorizontalBox::Slot()
-												.AutoWidth()
-												[
-													SNew( SCheckBox )
-													.Style( FEditorStyle::Get(), "Menu.CheckBox" )
-													.ToolTipText(LOCTEXT( "StaticMeshesToolTip", "Static Meshes will be adjusted if checked." ))
-													.IsChecked_Static(&FLevelEditorActionCallbacks::IsLightingResolutionStaticMeshesChecked)
-													.OnCheckStateChanged_Static(&FLevelEditorActionCallbacks::SetLightingResolutionStaticMeshes)
-													.Content()
-													[
-														SNew( STextBlock )
-														.Text( LOCTEXT("StaticMeshes", "Static Meshes") )
-													]
-												]
-												+SHorizontalBox::Slot()
-												.AutoWidth()
-												.Padding( FMargin( 4.0f, 0.0f, 11.0f, 0.0f ) )
-												[
-													SNew(SSpinBox<float>)
-													.MinValue(4.f)
-													.MaxValue(4096.f)
-													.ToolTipText(LOCTEXT( "LightingResolutionStaticMeshesMinToolTip", "The minimum lightmap resolution for static mesh adjustments. Anything outside of Min/Max range will not be touched when adjusting." ))
-													.Value(FLevelEditorActionCallbacks::GetLightingResolutionMinSMs())
-													.OnValueChanged_Static(&FLevelEditorActionCallbacks::SetLightingResolutionMinSMs)
-												]
-												+SHorizontalBox::Slot()
-												.AutoWidth()
-												[
-													SNew(SSpinBox<float>)
-													.MinValue(4.f)
-													.MaxValue(4096.f)
-													.ToolTipText(LOCTEXT( "LightingResolutionStaticMeshesMaxToolTip", "The maximum lightmap resolution for static mesh adjustments. Anything outside of Min/Max range will not be touched when adjusting." ))
-													.Value(FLevelEditorActionCallbacks::GetLightingResolutionMaxSMs())
-													.OnValueChanged_Static(&FLevelEditorActionCallbacks::SetLightingResolutionMaxSMs)
-												];
-				Section.AddEntry(FToolMenuEntry::InitWidget("Meshes", Meshes, FText::GetEmpty(), true));
-				
-				TSharedRef<SWidget> BSPs =		SNew(SHorizontalBox)
-												+SHorizontalBox::Slot()
-												.AutoWidth()
-												[
-													SNew( SCheckBox )
-													.Style(FEditorStyle::Get(), "Menu.CheckBox")
-													.ToolTipText(LOCTEXT( "BSPSurfacesToolTip", "BSP Surfaces will be adjusted if checked." ))
-													.IsChecked_Static(&FLevelEditorActionCallbacks::IsLightingResolutionBSPSurfacesChecked)
-													.OnCheckStateChanged_Static(&FLevelEditorActionCallbacks::SetLightingResolutionBSPSurfaces)
-													.Content()
-													[
-														SNew( STextBlock )
-														.Text( LOCTEXT("BSPSurfaces", "BSP Surfaces") )
-													]
-												]
-												+SHorizontalBox::Slot()
-												.AutoWidth()
-												.Padding( FMargin( 6.0f, 0.0f, 4.0f, 0.0f ) )
-												[
-													SNew(SSpinBox<float>)
-													.MinValue(1.f)
-													.MaxValue(63556.f)
-													.ToolTipText(LOCTEXT( "LightingResolutionBSPsMinToolTip", "The minimum lightmap resolution of a BSP surface to adjust. When outside of the Min/Max range, the BSP surface will no be altered." ))
-													.Value(FLevelEditorActionCallbacks::GetLightingResolutionMinBSPs())
-													.OnValueChanged_Static(&FLevelEditorActionCallbacks::SetLightingResolutionMinBSPs)
-												]
-												+SHorizontalBox::Slot()
-												.AutoWidth()
-												[
-													SNew(SSpinBox<float>)
-													.MinValue(1.f)
-													.MaxValue(63556.f)
-													.ToolTipText(LOCTEXT( "LightingResolutionBSPsMaxToolTip", "The maximum lightmap resolution of a BSP surface to adjust. When outside of the Min/Max range, the BSP surface will no be altered." ))
-													.Value(FLevelEditorActionCallbacks::GetLightingResolutionMaxBSPs())
-													.OnValueChanged_Static(&FLevelEditorActionCallbacks::SetLightingResolutionMaxBSPs)
-												];
-				Section.AddEntry(FToolMenuEntry::InitWidget("BSPs", BSPs, FText::GetEmpty(), true));
-			}
-
-			{
-				FToolMenuSection& Section = SubMenu->AddSection("LevelEditorBuildLightingResolution2", LOCTEXT( "LightingResolutionHeading2", "Select Options" ) );
-				Section.AddMenuEntry( FLevelEditorCommands::Get().LightingResolution_CurrentLevel );
-				Section.AddMenuEntry( FLevelEditorCommands::Get().LightingResolution_SelectedLevels );
-				Section.AddMenuEntry( FLevelEditorCommands::Get().LightingResolution_AllLoadedLevels );
-				Section.AddMenuEntry( FLevelEditorCommands::Get().LightingResolution_SelectedObjectsOnly );
-			}
-
-			{
-				FToolMenuSection& Section = SubMenu->AddSection("LevelEditorBuildLightingResolution3", LOCTEXT( "LightingResolutionHeading3", "Ratio" ) );
-				TSharedRef<SWidget> Ratio =		SNew(SSpinBox<int32>)
-												.MinValue(0)
-												.MaxValue(400)
-												.ToolTipText(LOCTEXT( "LightingResolutionRatioToolTip", "Ratio to apply (New Resolution = Ratio / 100.0f * CurrentResolution)." ))
-												.Value(FLevelEditorActionCallbacks::GetLightingResolutionRatio())
-												.OnEndSliderMovement_Static(&FLevelEditorActionCallbacks::SetLightingResolutionRatio)
-												.OnValueCommitted_Static(&FLevelEditorActionCallbacks::SetLightingResolutionRatioCommit);
-				Section.AddEntry(FToolMenuEntry::InitWidget("Ratio", Ratio, LOCTEXT( "LightingResolutionRatio", "Ratio" )));
-			}
-		}
-
-		/** Generates a lighting info dialogs sub-menu */
-		static void RegisterLightingInfoMenu(const FName InBaseMenuName)
-		{
-			RegisterLightingDensityMenu(UToolMenus::JoinMenuPaths(InBaseMenuName, "LightingInfo"));
-			RegisterLightingResolutionMenu(UToolMenus::JoinMenuPaths(InBaseMenuName, "LightingInfo"));
-
-			UToolMenu* SubMenu = UToolMenus::Get()->RegisterMenu(UToolMenus::JoinMenuPaths(InBaseMenuName, "LightingInfo"));
-
-			{
-				FToolMenuSection& Section = SubMenu->AddSection("LevelEditorBuildLightingInfo", LOCTEXT( "LightingInfoHeading", "Lighting Info Dialogs" ) );
-				Section.AddSubMenu(
-					"LightingDensity",
-					LOCTEXT( "LightingDensityRenderingSubMenu", "LightMap Density Rendering Options" ),
-					LOCTEXT( "LightingDensityRenderingSubMenu_ToolTip", "Shows the LightMap Density Rendering viewmode options." ),
-					FNewToolMenuChoice() );
-
-				Section.AddSubMenu(
-					"LightingResolution",
-					LOCTEXT( "LightingResolutionAdjustmentSubMenu", "LightMap Resolution Adjustment" ),
-					LOCTEXT( "LightingResolutionAdjustmentSubMenu_ToolTip", "Shows the LightMap Resolution Adjustment options." ),
-					FNewToolMenuChoice() );
-
-				Section.AddMenuEntry( FLevelEditorCommands::Get().LightingStaticMeshInfo, LOCTEXT( "BuildLightingInfo_LightingStaticMeshInfo", "Lighting StaticMesh Info..." ) );
-			}
-		}
-	};
-
-	FLightingMenus::RegisterMenus(BaseMenuName);
-
-	{
-		FToolMenuSection& Section = Menu->AddSection( "LevelEditorLighting", LOCTEXT( "LightingHeading", "Lighting" ) );
-		Section.AddMenuEntry( FLevelEditorCommands::Get().BuildLightingOnly, LOCTEXT( "BuildLightingOnlyHeading", "Build Lighting Only" ) );
-
-		Section.AddSubMenu(
-			"LightingQuality",
-			LOCTEXT( "LightingQualitySubMenu", "Lighting Quality" ),
-			LOCTEXT( "LightingQualitySubMenu_ToolTip", "Allows you to select the quality level for precomputed lighting" ),
-			FNewToolMenuChoice() );
-
-		Section.AddSubMenu(
-			"LightingInfo",
-			LOCTEXT( "BuildLightingInfoSubMenu", "Lighting Info" ),
-			LOCTEXT( "BuildLightingInfoSubMenu_ToolTip", "Access the lighting info dialogs" ),
-			FNewToolMenuChoice() );
-
-		Section.AddMenuEntry( FLevelEditorCommands::Get().LightingBuildOptions_UseErrorColoring );
-		Section.AddMenuEntry( FLevelEditorCommands::Get().LightingBuildOptions_ShowLightingStats );
-	}
-
-	{
-		FToolMenuSection& Section = Menu->AddSection("LevelEditorReflections", LOCTEXT( "ReflectionHeading", "Reflections" ) );
-		Section.AddMenuEntry( FLevelEditorCommands::Get().BuildReflectionCapturesOnly );
-	}
-
-	{
-		FToolMenuSection& Section = Menu->AddSection("LevelEditorVisibility", LOCTEXT( "VisibilityHeading", "Visibility" ) );
-		Section.AddMenuEntry( FLevelEditorCommands::Get().BuildLightingOnly_VisibilityOnly );
-	}
-
-	{
-		FToolMenuSection& Section = Menu->AddSection("LevelEditorGeometry", LOCTEXT( "GeometryHeading", "Geometry" ) );
-		Section.AddMenuEntry( FLevelEditorCommands::Get().BuildGeometryOnly );
-		Section.AddMenuEntry( FLevelEditorCommands::Get().BuildGeometryOnly_OnlyCurrentLevel );
-	}
-
-	{
-		FToolMenuSection& Section = Menu->AddSection("LevelEditorNavigation", LOCTEXT( "NavigationHeading", "Navigation" ) );
-		Section.AddMenuEntry( FLevelEditorCommands::Get().BuildPathsOnly );
-	}
-
-	{
-		FToolMenuSection& Section = Menu->AddSection("LevelEditorLOD", LOCTEXT("LODHeading", "Hierarchical LOD"));
-		Section.AddMenuEntry(FLevelEditorCommands::Get().BuildLODsOnly);
-	}
-	
-	{
-		FToolMenuSection& Section = Menu->AddSection("LevelEditorTextureStreaming", LOCTEXT("TextureStreamingHeading", "Texture Streaming"));
-		Section.AddDynamicEntry("BuildTextureStreamingOnly", FNewToolMenuSectionDelegate::CreateLambda([](FToolMenuSection& InSection)
-		{
-			if (CVarStreamingUseNewMetrics.GetValueOnAnyThread() != 0) // There is no point of in building texture streaming data with the old system.
-			{
-				InSection.AddMenuEntry(FLevelEditorCommands::Get().BuildTextureStreamingOnly);
-			}
-		}));
-		Section.AddMenuEntry(FLevelEditorCommands::Get().BuildVirtualTextureOnly);
-	}
-
-	{
-		FToolMenuSection& Section = Menu->AddSection("LevelEditorGrass", LOCTEXT("GrassHeading", "Grass Maps"));
-		Section.AddMenuEntry(FLevelEditorCommands::Get().BuildGrassMapsOnly);
-	}
-
-	
-	{
-		FToolMenuSection& Section = Menu->AddSection("LevelEditorAutomation", LOCTEXT( "AutomationHeading", "Automation" ) );
-		Section.AddMenuEntry( 
-			FLevelEditorCommands::Get().BuildAndSubmitToSourceControl,
-			TAttribute<FText>(),
-			TAttribute<FText>(),
-			FSlateIcon(FEditorStyle::GetStyleSetName(), "LevelEditor.Tabs.BuildAndSubmit")
-			);
-	}
-
-	// Map Check
-	{
-		FToolMenuSection& Section = Menu->AddSection("LevelEditorVerification", LOCTEXT( "VerificationHeading", "Verification" ) );
-		Section.AddMenuEntry( FLevelEditorCommands::Get().MapCheck, LOCTEXT("OpenMapCheck", "Map Check") );
-	}
-
-#undef LOCTEXT_NAMESPACE
-}
 
 static void MakeMaterialQualityLevelMenu( UToolMenu* InMenu )
 {
@@ -1980,7 +1603,7 @@ void FLevelEditorToolBar::RegisterCompileMenu()
 {
 #define LOCTEXT_NAMESPACE "LevelToolBarCompileMenu"
 
-	UToolMenu* Menu = UToolMenus::Get()->RegisterMenu("LevelEditor.LevelEditorToolBar.CompileComboButton");
+	UToolMenu* Menu = UToolMenus::Get()->RegisterMenu("LevelEditor.LevelEditorToolBar.AssetsToolBar.CompileComboButton");
 
 	{
 		FToolMenuSection& Section = Menu->AddSection("LiveCodingMode", LOCTEXT( "LiveCodingMode", "General" ) );
@@ -2144,67 +1767,13 @@ TSharedRef< SWidget > FLevelEditorToolBar::GenerateSourceControlMenu(TSharedRef<
 	return UToolMenus::Get()->GenerateWidget("LevelEditor.LevelEditorToolBar.SourceControl", MenuContext);
 }
 
-TSharedRef< SWidget > FLevelEditorToolBar::GenerateEditorModesMenu(TSharedRef<FUICommandList> InCommandList, TWeakPtr<SLevelEditor> InLevelEditor)
-{
-	// Get all menu extenders for this context menu from the level editor module
-	FLevelEditorModule& LevelEditorModule = FModuleManager::GetModuleChecked<FLevelEditorModule>(TEXT("LevelEditor"));
-
-	FToolMenuContext MenuContext(InCommandList);
-	ULevelEditorMenuContext* LevelEditorMenuContext = NewObject<ULevelEditorMenuContext>();
-	LevelEditorMenuContext->LevelEditor = InLevelEditor;
-	MenuContext.AddObject(LevelEditorMenuContext);
-
-	return UToolMenus::Get()->GenerateWidget("LevelEditor.LevelEditorToolBar.EditorModes", MenuContext);
-}
-
 void FLevelEditorToolBar::RegisterSourceControlMenu()
 {
 #define LOCTEXT_NAMESPACE "LevelToolBarSourceControlMenu"
 
-	UToolMenu* Menu = UToolMenus::Get()->RegisterMenu("LevelEditor.LevelEditorToolBar.SourceControl");
-	Menu->bShouldCloseWindowAfterMenuSelection = true;
-	FToolMenuSection& Section = Menu->AddSection("SourceControlActions", LOCTEXT("SourceControlMenuHeadingActions", "Actions"));
-
-	Section.AddDynamicEntry("ConnectToSourceControl", FNewToolMenuSectionDelegate::CreateLambda([](FToolMenuSection& InSection)
-	{
-		ISourceControlModule& SourceControlModule = ISourceControlModule::Get();
-		if (ISourceControlModule::Get().IsEnabled() && ISourceControlModule::Get().GetProvider().IsAvailable())
-		{
-			InSection.AddMenuEntry(
-				FLevelEditorCommands::Get().ChangeSourceControlSettings, 
-				TAttribute<FText>(),
-				TAttribute<FText>(),
-				FSlateIcon(FEditorStyle::GetStyleSetName(), "SourceControl.Actions.ChangeSettings")
-				);
-		}
-		else
-		{
-			InSection.AddMenuEntry(
-				FLevelEditorCommands::Get().ConnectToSourceControl,
-				TAttribute<FText>(),
-				TAttribute<FText>(),
-				FSlateIcon(FEditorStyle::GetStyleSetName(), "SourceControl.Actions.Connect")
-				);
-		}
-	}));
-
-	Section.AddSeparator("SourceControlConnectionSeparator");
-
-	Section.AddMenuEntry(
-		FLevelEditorCommands::Get().CheckOutModifiedFiles,
-		TAttribute<FText>(),
-		TAttribute<FText>(),
-		FSlateIcon(FEditorStyle::GetStyleSetName(), "SourceControl.Actions.CheckOut")
-		);
 
 
-	Section.AddMenuEntry(
-		FLevelEditorCommands::Get().SubmitToSourceControl,
-		TAttribute<FText>(),
-		TAttribute<FText>(),
-		FSlateIcon(FEditorStyle::GetStyleSetName(), "SourceControl.Actions.Submit")
-		);
-
+	
 #undef LOCTEXT_NAMESPACE
 }
 
@@ -2466,76 +2035,6 @@ void FLevelEditorToolBar::RegisterCinematicsMenu()
 #undef LOCTEXT_NAMESPACE
 }
 
-void FLevelEditorToolBar::RegisterEditorModesMenu()
-{
-#define LOCTEXT_NAMESPACE "LevelToolBarEditorModesMenu"
-	UToolMenu* Menu = UToolMenus::Get()->RegisterMenu("LevelEditor.LevelEditorToolBar.EditorModes");
-
-	FToolMenuSection& Section = Menu->AddSection("EditorModes", LOCTEXT("EditorModesMenu_NewHeading", "Editor Modes"));
-
-	Section.AddDynamicEntry("ModesList", FNewToolMenuSectionDelegate::CreateLambda([](FToolMenuSection& InSection)
-	{
-
-		const FLevelEditorModule& LevelEditorModule = FModuleManager::LoadModuleChecked<FLevelEditorModule>("LevelEditor");
-		const FLevelEditorModesCommands& Commands = LevelEditorModule.GetLevelEditorModesCommands();
-
-		TArray<FEditorModeInfo, TInlineAllocator<1>> DefaultModes;
-
-		TArray<FEditorModeInfo, TInlineAllocator<10>> NonDefaultModes;
-
-		for (const FEditorModeInfo& Mode : FEditorModeRegistry::Get().GetSortedModeInfo())
-		{
-			// If the mode isn't visible don't create a menu option for it.
-			if (!Mode.bVisible)
-			{
-				continue;
-			}
-
-			if (GLevelEditorModeTools().IsDefaultMode(Mode.ID))
-			{
-				DefaultModes.Add(Mode.ID);
-			}
-			else
-			{
-				NonDefaultModes.Add(Mode.ID);
-			}
-			
-		}
-
-		auto BuildEditorModes = 
-			[&Commands, &InSection](const TArrayView<FEditorModeInfo>& Modes)
-			{
-				for (const FEditorModeInfo& Mode : Modes)
-				{
-					FName EditorModeCommandName = FName(*(FString("EditorMode.") + Mode.ID.ToString()));
-
-					TSharedPtr<FUICommandInfo> EditorModeCommand =
-						FInputBindingManager::Get().FindCommandInContext(Commands.GetContextName(), EditorModeCommandName);
-
-					// If a command isn't yet registered for this mode, we need to register one.
-					if (!EditorModeCommand.IsValid())
-					{
-						continue;
-					}
-
-					InSection.AddMenuEntry(EditorModeCommand);
-				}
-
-			};
-
-		// Build Default Modes first
-		BuildEditorModes(MakeArrayView(DefaultModes));
-
-		InSection.AddSeparator(NAME_None);
-			
-		// Build non-default modes second
-		BuildEditorModes(MakeArrayView(NonDefaultModes));
-
-	}));
-
-#undef LOCTEXT_NAMESPACE
-}
-
 
 void FLevelEditorToolBar::OnCinematicsActorPicked( AActor* Actor )
 {
diff --git a/Engine/Source/Editor/UnrealEd/Private/Subsystems/StatusBarSubsystem.cpp b/Engine/Source/Editor/UnrealEd/Private/Subsystems/StatusBarSubsystem.cpp
--- a/Engine/Source/Editor/UnrealEd/Private/Subsystems/StatusBarSubsystem.cpp
+++ b/Engine/Source/Editor/UnrealEd/Private/Subsystems/StatusBarSubsystem.cpp
@@ -11,7 +11,11 @@
 #include "ISourceControlModule.h"
 #include "SourceControlWindows.h"
 #include "Framework/Commands/Commands.h"
+#include "Widgets/Images/SImage.h"
 #include "Widgets/Input/SComboButton.h"
+#include "Widgets/Input/SButton.h"
+#include "Widgets/Layout/SBorder.h"
+#include "Widgets/Layout/SBox.h"
 #include "EditorStyleSet.h"
 #include "FileHelpers.h"
 
