Updated the active viewport focus indicator

diff --git a/Engine/Source/Editor/EditorStyle/Private/StarshipStyle.cpp b/Engine/Source/Editor/EditorStyle/Private/StarshipStyle.cpp
--- a/Engine/Source/Editor/EditorStyle/Private/StarshipStyle.cpp
+++ b/Engine/Source/Editor/EditorStyle/Private/StarshipStyle.cpp
@@ -4170,6 +4170,9 @@ void FStarshipEditorStyle::FStyle::SetupLevelEditorStyle()
 		Set( "EditorViewport.Bottom", new IMAGE_BRUSH("Icons/icon_ViewMode_ViewBottom_16px", Icon16x16 ) );
 		Set( "EditorViewport.Right", new IMAGE_BRUSH("Icons/icon_ViewMode_ViewRight_16px", Icon16x16 ) );
 		Set( "EditorViewport.Back", new IMAGE_BRUSH("Icons/icon_ViewMode_ViewBack_16px", Icon16x16 ) );
+
+		Set("EditorViewport.ActiveBorderColor", FStyleColors::PrimaryHover);
+
 #endif
 
 #if WITH_EDITOR || (IS_PROGRAM && WITH_UNREAL_DEVELOPER_TOOLS)
@@ -4470,8 +4473,7 @@ void FStarshipEditorStyle::FStyle::SetupLevelEditorStyle()
 	// Level viewport 
 #if WITH_EDITOR || (IS_PROGRAM && WITH_UNREAL_DEVELOPER_TOOLS)
 	{
-		Set("LevelViewport.ActiveViewportBorder", new CORE_BORDER_BRUSH("Starship/Common/White", FMargin(1), FStyleColors::Primary));
-		Set( "LevelViewport.NoViewportBorder", new FSlateNoResource() );
+
 		Set( "LevelViewport.DebugBorder", new BOX_BRUSH( "Old/Window/ViewportDebugBorder", 0.8f, FLinearColor(.7,0,0,.5) ) );
 		Set( "LevelViewport.BlackBackground", new FSlateColorBrush( FLinearColor::Black ) ); 
 		Set( "LevelViewport.StartingPlayInEditorBorder", new BOX_BRUSH( "Old/Window/ViewportDebugBorder", 0.8f, FLinearColor(0.1f,1.0f,0.1f,1.0f) ) );
diff --git a/Engine/Source/Editor/LevelEditor/Private/SLevelViewport.cpp b/Engine/Source/Editor/LevelEditor/Private/SLevelViewport.cpp
--- a/Engine/Source/Editor/LevelEditor/Private/SLevelViewport.cpp
+++ b/Engine/Source/Editor/LevelEditor/Private/SLevelViewport.cpp
@@ -175,9 +175,6 @@ void SLevelViewport::Construct(const FArguments& InArgs)
 	ParentLevelEditor = StaticCastSharedRef<SLevelEditor>( InArgs._ParentLevelEditor.Pin().ToSharedRef() );
 	ConfigKey = InArgs._ConfigKey;
 
-	// Store border brushes for differentiating between active and inactive viewports
-	ActiveBorder = FEditorStyle::GetBrush( "LevelViewport.ActiveViewportBorder" );
-	NoBorder = FEditorStyle::GetBrush( "LevelViewport.NoViewportBorder" );
 	DebuggingBorder = FEditorStyle::GetBrush( "LevelViewport.DebugBorder" );
 	BlackBackground = FEditorStyle::GetBrush( "LevelViewport.BlackBackground" );
 	StartingPlayInEditorBorder = FEditorStyle::GetBrush( "LevelViewport.StartingPlayInEditorBorder" );
@@ -1646,16 +1643,6 @@ const FSlateBrush* SLevelViewport::OnGetViewportBorderBrush() const
 	const FSlateBrush* BorderBrush = nullptr;
 	if( FSlateApplication::Get().IsNormalExecution() )
 	{
-		// Only show the active border if we have a valid client, its the current client being edited and we arent in immersive (in immersive there is only one visible viewport)
-		if( LevelViewportClient.IsValid() && LevelViewportClient.Get() == GCurrentLevelEditingViewportClient && !IsImmersive() )
-		{
-			BorderBrush = ActiveBorder;
-		}
-		else
-		{
-			BorderBrush = NoBorder;
-		}
-
 		// If a PIE/SIE/Editor transition just completed, then we'll draw a border effect to draw attention to it
 		if( ViewTransitionAnim.IsPlaying() )
 		{
@@ -1687,6 +1674,25 @@ const FSlateBrush* SLevelViewport::OnGetViewportBorderBrush() const
 	return BorderBrush;
 }
 
+EVisibility SLevelViewport::OnGetFocusedViewportIndicatorVisibility() const
+{
+	EVisibility BaseVisibility = OnGetViewportContentVisibility();
+	if (BaseVisibility != EVisibility::Collapsed)
+	{
+		// Only show the active border if we have a valid client, its the current client being edited and we arent in immersive (in immersive there is only one visible viewport)
+		if (LevelViewportClient.IsValid() && LevelViewportClient.Get() == GCurrentLevelEditingViewportClient && !IsImmersive())
+		{
+			return EVisibility::HitTestInvisible;
+		}
+		else
+		{
+			return EVisibility::Collapsed;
+		}
+	}
+
+	return BaseVisibility;
+}
+
 FSlateColor SLevelViewport::OnGetViewportBorderColorAndOpacity() const
 {
 	FLinearColor ViewportBorderColorAndOpacity = FLinearColor::White;
diff --git a/Engine/Source/Editor/LevelEditor/Public/SLevelViewport.h b/Engine/Source/Editor/LevelEditor/Public/SLevelViewport.h
--- a/Engine/Source/Editor/LevelEditor/Public/SLevelViewport.h
+++ b/Engine/Source/Editor/LevelEditor/Public/SLevelViewport.h
@@ -425,6 +425,7 @@ protected:
 	virtual const FSlateBrush* OnGetViewportBorderBrush() const override;
 	virtual FSlateColor OnGetViewportBorderColorAndOpacity() const override;
 	virtual EVisibility OnGetViewportContentVisibility() const override;
+	virtual EVisibility OnGetFocusedViewportIndicatorVisibility() const override;
 	virtual void BindCommands() override;
 private:
 	/** Flag to know if we need to update the previews which is handled in the tick. */
@@ -757,10 +758,6 @@ private:
 	/** Level viewport client */
 	TSharedPtr<FLevelEditorViewportClient> LevelViewportClient;
 
-	/** The brush to use if this viewport is the active viewport */
-	const FSlateBrush* ActiveBorder;
-	/** The brush to use if this viewport is an inactive viewport or not showing a border */
-	const FSlateBrush* NoBorder;
 	/** The brush to use if this viewport is in debug mode */
 	const FSlateBrush* DebuggingBorder;
 	/** The brush to use for a black background */
diff --git a/Engine/Source/Editor/UnrealEd/Private/SEditorViewport.cpp b/Engine/Source/Editor/UnrealEd/Private/SEditorViewport.cpp
--- a/Engine/Source/Editor/UnrealEd/Private/SEditorViewport.cpp
+++ b/Engine/Source/Editor/UnrealEd/Private/SEditorViewport.cpp
@@ -22,6 +22,7 @@
 #include "MaterialShaderQualitySettings.h"
 #include "RHIShaderPlatformDefinitions.inl"
 #include "RayTracingDebugVisualizationMenuCommands.h"
+#include "Widgets/Colors/SComplexGradient.h"
 
 #define LOCTEXT_NAMESPACE "EditorViewport"
 
@@ -94,7 +95,7 @@ void SEditorViewport::Construct( const FArguments& InArgs )
 				ViewportToolbar.ToSharedRef()
 			];
 	}
-
+	
 	ViewportOverlay->AddSlot()
 	[
 		SNew(SBorder)
@@ -105,6 +106,27 @@ void SEditorViewport::Construct( const FArguments& InArgs )
 		.ShowEffectWhenDisabled(false)
 	];
 
+	// This makes a gradient that displays whether or not a viewport is active
+	FLinearColor ActiveBorderColor = FAppStyle::Get().GetColor("EditorViewport.ActiveBorderColor");
+	FLinearColor ActiveBorderColorTransparent = ActiveBorderColor;
+	ActiveBorderColorTransparent.A = 0.0f;
+
+	static TArray<FLinearColor> GradientStops{ ActiveBorderColorTransparent, ActiveBorderColor, ActiveBorderColorTransparent };
+
+	ViewportOverlay->AddSlot()
+	.VAlign(VAlign_Top)
+	[
+		SNew(SBox)
+		.Visibility(this, &SEditorViewport::OnGetFocusedViewportIndicatorVisibility)
+		.MaxDesiredHeight(1.0f)
+		.MinDesiredHeight(1.0f)
+		[
+			SNew(SComplexGradient)
+			.GradientColors(GradientStops)
+			.Orientation(EOrientation::Orient_Vertical)
+		]
+	];
+
 	PopulateViewportOverlays(ViewportOverlay.ToSharedRef());
 }
 
diff --git a/Engine/Source/Editor/UnrealEd/Public/SEditorViewport.h b/Engine/Source/Editor/UnrealEd/Public/SEditorViewport.h
--- a/Engine/Source/Editor/UnrealEd/Public/SEditorViewport.h
+++ b/Engine/Source/Editor/UnrealEd/Public/SEditorViewport.h
@@ -136,6 +136,11 @@ protected:
 	 */
 	virtual EVisibility OnGetViewportContentVisibility() const;
 
+	/**
+	 * @return The visibility of the viewport focus indicator.                                                             
+	 */
+	virtual EVisibility OnGetFocusedViewportIndicatorVisibility() const { return EVisibility::Collapsed; }
+
 	/** UI Command delegate bindings */
 	void OnToggleStats();
 
