Fix inconsistent padding above status bar

diff --git a/Engine/Source/Editor/LevelEditor/Private/SLevelEditor.cpp b/Engine/Source/Editor/LevelEditor/Private/SLevelEditor.cpp
--- a/Engine/Source/Editor/LevelEditor/Private/SLevelEditor.cpp
+++ b/Engine/Source/Editor/LevelEditor/Private/SLevelEditor.cpp
@@ -225,7 +225,7 @@ void SLevelEditor::Initialize( const TSharedRef<SDockTab>& OwnerTab, const TShar
 			ContentArea
 		]
 		+SVerticalBox::Slot()
-		.Padding(0.0f, 1.0f, 0.0f, 0.0f)
+		.Padding(0.0f, 2.0f, 0.0f, 0.0f)
 		.AutoHeight()
 		[
 			GEditor->GetEditorSubsystem<UStatusBarSubsystem>()->MakeStatusBarWidget(TEXT("LevelEditor.StatusBar"), OwnerTab)
