Added a GetCurrentLevel function to scripts

diff --git a/Engine/Source/Editor/LevelEditor/Private/LevelEditorSubsystem.cpp b/Engine/Source/Editor/LevelEditor/Private/LevelEditorSubsystem.cpp
--- a/Engine/Source/Editor/LevelEditor/Private/LevelEditorSubsystem.cpp
+++ b/Engine/Source/Editor/LevelEditor/Private/LevelEditorSubsystem.cpp
@@ -436,4 +436,24 @@ bool ULevelEditorSubsystem::SetCurrentLevelByName(FName LevelName)
 	}
 
 	return bLevelFound;
+}
+
+ULevel* ULevelEditorSubsystem::GetCurrentLevel()
+{
+	UUnrealEditorSubsystem* UnrealEditorSubsystem = GEditor->GetEditorSubsystem<UUnrealEditorSubsystem>();
+
+	if (!UnrealEditorSubsystem)
+	{
+		return nullptr;
+	}
+
+	UWorld* World = UnrealEditorSubsystem->GetEditorWorld();
+
+	if (!World)
+	{
+		UE_LOG(LevelEditorSubsystem, Error, TEXT("GetCurrentLevel. Can't Get the current level because there is no world."));
+		return nullptr;
+	}
+
+	return World->GetCurrentLevel();
 }
\ No newline at end of file
diff --git a/Engine/Source/Editor/LevelEditor/Public/LevelEditorSubsystem.h b/Engine/Source/Editor/LevelEditor/Public/LevelEditorSubsystem.h
--- a/Engine/Source/Editor/LevelEditor/Public/LevelEditorSubsystem.h
+++ b/Engine/Source/Editor/LevelEditor/Public/LevelEditorSubsystem.h
@@ -84,4 +84,11 @@ public:
 	 */
 	UFUNCTION(BlueprintCallable, Category = "Editor Scripting | Level Utility")
 	bool SetCurrentLevelByName(FName LevelName);
+
+	/**
+	* Get the current level used by the world editor.
+	* @return	The current level
+	*/
+	UFUNCTION(BlueprintCallable, Category = "Editor Scripting | Level Utility")
+	ULevel* GetCurrentLevel();
 };
\ No newline at end of file
