Move compile off the main toolbar and into the status bar

diff --git a/Engine/Source/Editor/LevelEditor/Private/LevelEditorToolBar.cpp b/Engine/Source/Editor/LevelEditor/Private/LevelEditorToolBar.cpp
--- a/Engine/Source/Editor/LevelEditor/Private/LevelEditorToolBar.cpp
+++ b/Engine/Source/Editor/LevelEditor/Private/LevelEditorToolBar.cpp
@@ -1176,10 +1176,6 @@ void FLevelEditorToolBar::RegisterLevelEditorToolBar( const TSharedRef<FUIComman
 	RegisterSourceControlMenu();
 	RegisterCinematicsMenu();
 
-#if WITH_LIVE_CODING
-	RegisterCompileMenu();
-#endif
-
 	RegisterQuickSettingsMenu();
 	RegisterOpenBlueprintMenu();
 
@@ -1226,48 +1222,6 @@ void FLevelEditorToolBar::RegisterLevelEditorToolBar( const TSharedRef<FUIComman
 			Section.AddEntry(FToolMenuEntry::InitToolBarButton(FLevelEditorCommands::Get().ToggleVR, LOCTEXT("ToggleVR", "VR Mode")));
 		}
 
-		//if(FCoreStyle::IsStarshipStyle())
-		{
-			FToolMenuSection& Section = AssetsToolBar->AddSection("Compile");
-
-			Section.AddDynamicEntry("CompilerAvailable", FNewToolMenuSectionDelegate::CreateLambda([](FToolMenuSection& InSection)
-				{
-					// Only show the compile options on machines with the solution (assuming they can build it)
-					if (FSourceCodeNavigation::IsCompilerAvailable())
-					{
-						// Since we can always add new code to the project, only hide these buttons if we haven't done so yet
-						InSection.AddEntry(FToolMenuEntry::InitToolBarButton(
-							"CompileButton",
-							FUIAction(
-								FExecuteAction::CreateStatic(&FLevelEditorActionCallbacks::RecompileGameCode_Clicked),
-								FCanExecuteAction::CreateStatic(&FLevelEditorActionCallbacks::Recompile_CanExecute),
-								FIsActionChecked(),
-								FIsActionButtonVisible::CreateStatic(FLevelEditorActionCallbacks::CanShowSourceCodeActions)),
-							LOCTEXT("CompileMenuButton", "Compile"),
-							FLevelEditorCommands::Get().RecompileGameCode->GetDescription(),
-							FSlateIcon(FEditorStyle::GetStyleSetName(), "LevelEditor.Recompile")
-						));
-
-#if WITH_LIVE_CODING
-						InSection.AddEntry(FToolMenuEntry::InitComboButton(
-							"CompileComboButton",
-							FUIAction(
-								FExecuteAction(),
-								FCanExecuteAction(),
-								FIsActionChecked(),
-								FIsActionButtonVisible::CreateStatic(FLevelEditorActionCallbacks::CanShowSourceCodeActions)),
-							FNewToolMenuChoice(),
-							LOCTEXT("CompileCombo_Label", "Compile Options"),
-							LOCTEXT("CompileComboToolTip", "Compile options menu"),
-							FSlateIcon(FEditorStyle::GetStyleSetName(), "LevelEditor.Recompile"),
-							true
-						));
-#endif
-					}
-				}));
-		}
-		
-
 	}
 
 	UToolMenu* ModesToolbar = UToolMenus::Get()->RegisterMenu("LevelEditor.LevelEditorToolBar.ModesToolBar", NAME_None, EMultiBoxType::SlimHorizontalToolBar);
@@ -1600,28 +1554,7 @@ static void MakePreviewSettingsMenu( UToolMenu* InMenu )
 #undef LOCTEXT_NAMESPACE
 }
 
-#if WITH_LIVE_CODING
-void FLevelEditorToolBar::RegisterCompileMenu()
-{
-#define LOCTEXT_NAMESPACE "LevelToolBarCompileMenu"
 
-	UToolMenu* Menu = UToolMenus::Get()->RegisterMenu("LevelEditor.LevelEditorToolBar.AssetsToolBar.CompileComboButton");
-
-	{
-		FToolMenuSection& Section = Menu->AddSection("LiveCodingMode", LOCTEXT( "LiveCodingMode", "General" ) );
-		Section.AddMenuEntry( FLevelEditorCommands::Get().LiveCoding_Enable );
-	}
-
-	{
-		FToolMenuSection& Section = Menu->AddSection("LiveCodingActions", LOCTEXT( "LiveCodingActions", "Actions" ) );
-		Section.AddMenuEntry( FLevelEditorCommands::Get().LiveCoding_StartSession );
-		Section.AddMenuEntry( FLevelEditorCommands::Get().LiveCoding_ShowConsole );
-		Section.AddMenuEntry( FLevelEditorCommands::Get().LiveCoding_Settings );
-	}
-
-#undef LOCTEXT_NAMESPACE
-}
-#endif
 
 TSharedRef< SWidget > FLevelEditorToolBar::GenerateQuickSettingsMenu(TSharedRef<FUICommandList> InCommandList, TWeakPtr<SLevelEditor> InLevelEditor)
 {
diff --git a/Engine/Source/Editor/LevelEditor/Private/LevelEditorToolBar.h b/Engine/Source/Editor/LevelEditor/Private/LevelEditorToolBar.h
--- a/Engine/Source/Editor/LevelEditor/Private/LevelEditorToolBar.h
+++ b/Engine/Source/Editor/LevelEditor/Private/LevelEditorToolBar.h
@@ -50,14 +50,6 @@ protected:
 	 */
 	static TSharedRef< SWidget > GenerateSourceControlMenu(TSharedRef<FUICommandList> InCommandList, TWeakPtr<SLevelEditor> InLevelEditor);
 
-	
-	/**
-	 * Generates menu content for the modes combo button drop down menu
-	 *
-	 * @return	Menu content widget
-	 */
-	static TSharedRef< SWidget > GenerateEditorModesMenu(TSharedRef<FUICommandList> InCommandList, TWeakPtr<SLevelEditor> InLevelEditor);
-
 	/**
 	 * Generates menu content for the compile combo button drop down menu
 	 *
@@ -89,15 +81,6 @@ private:
 
 	static void RegisterSourceControlMenu();
 	static void RegisterCinematicsMenu();
-	static void RegisterEditorModesMenu();
-	static void RegisterBuildMenu();
-
-#if WITH_LIVE_CODING
-	/**
-	 * Generates menu content for the compile combo button drop down menu
-	 */
-	static void RegisterCompileMenu();
-#endif
 
 	static void RegisterQuickSettingsMenu();
 	static void RegisterOpenBlueprintMenu();
diff --git a/Engine/Source/Editor/LevelEditor/Private/SLevelEditor.cpp b/Engine/Source/Editor/LevelEditor/Private/SLevelEditor.cpp
--- a/Engine/Source/Editor/LevelEditor/Private/SLevelEditor.cpp
+++ b/Engine/Source/Editor/LevelEditor/Private/SLevelEditor.cpp
@@ -54,6 +54,7 @@
 #include "Classes/EditorStyleSettings.h"
 #include "StatusBarSubsystem.h"
 #include "Widgets/Colors/SColorPicker.h"
+#include "SourceCodeNavigation.h"
 
 static const FName MainFrameModuleName("MainFrame");
 static const FName LevelEditorModuleName("LevelEditor");
@@ -248,6 +249,8 @@ void SLevelEditor::Initialize( const TSharedRef<SDockTab>& OwnerTab, const TShar
 		]
 	];
 	
+	RegisterStatusBarTools();
+
 	TtileBarMessageBox = SNew(SHorizontalBox);
 
 	ConstructTitleBarMessages();
@@ -1648,6 +1651,72 @@ void SLevelEditor::OnLevelActorOuterChanged(AActor* InActor, UObject* InOldOuter
 	bNeedsRefresh = true;
 }
 
+void SLevelEditor::RegisterStatusBarTools()
+{
+#define LOCTEXT_NAMESPACE "LevelToolBarCompileMenu"
+
+#if WITH_LIVE_CODING
+	{
+		UToolMenu* Menu = UToolMenus::Get()->RegisterMenu("LevelEditor.StatusBar.ToolBar.CompileComboButton");
+		{
+			FToolMenuSection& Section = Menu->AddSection("LiveCodingMode", LOCTEXT("LiveCodingMode", "General"));
+			Section.AddMenuEntry(FLevelEditorCommands::Get().LiveCoding_Enable);
+		}
+
+		{
+			FToolMenuSection& Section = Menu->AddSection("LiveCodingActions", LOCTEXT("LiveCodingActions", "Actions"));
+			Section.AddMenuEntry(FLevelEditorCommands::Get().LiveCoding_StartSession);
+			Section.AddMenuEntry(FLevelEditorCommands::Get().LiveCoding_ShowConsole);
+			Section.AddMenuEntry(FLevelEditorCommands::Get().LiveCoding_Settings);
+		}
+	}
+
+#endif
+
+	UToolMenu* Menu = UToolMenus::Get()->ExtendMenu("LevelEditor.StatusBar.ToolBar");
+	
+	FToolMenuSection& Section = Menu->AddSection("Compile", FText::GetEmpty(), FToolMenuInsert("SourceControl", EToolMenuInsertType::Before));
+
+	Section.AddDynamicEntry("CompilerAvailable", 
+		FNewToolMenuSectionDelegate::CreateLambda([](FToolMenuSection& InSection)
+		{
+			// Only show the compile options on machines with the solution (assuming they can build it)
+			if (FSourceCodeNavigation::IsCompilerAvailable())
+			{
+				// Since we can always add new code to the project, only hide these buttons if we haven't done so yet
+				InSection.AddEntry(FToolMenuEntry::InitToolBarButton(
+					"CompileButton",
+					FUIAction(
+						FExecuteAction::CreateStatic(&FLevelEditorActionCallbacks::RecompileGameCode_Clicked),
+						FCanExecuteAction::CreateStatic(&FLevelEditorActionCallbacks::Recompile_CanExecute),
+						FIsActionChecked(),
+						FIsActionButtonVisible::CreateStatic(FLevelEditorActionCallbacks::CanShowSourceCodeActions)),
+					LOCTEXT("CompileMenuButton", "Compile"),
+					FLevelEditorCommands::Get().RecompileGameCode->GetDescription(),
+					FSlateIcon(FEditorStyle::GetStyleSetName(), "LevelEditor.Recompile")
+				));
+
+#if WITH_LIVE_CODING
+				InSection.AddEntry(FToolMenuEntry::InitComboButton(
+					"CompileComboButton",
+					FUIAction(
+						FExecuteAction(),
+						FCanExecuteAction(),
+						FIsActionChecked(),
+						FIsActionButtonVisible::CreateStatic(FLevelEditorActionCallbacks::CanShowSourceCodeActions)),
+					FNewToolMenuChoice(),
+					LOCTEXT("CompileCombo_Label", "Compile Options"),
+					LOCTEXT("CompileComboToolTip", "Compile options menu"),
+					FSlateIcon(FEditorStyle::GetStyleSetName(), "LevelEditor.Recompile"),
+					true
+				));
+#endif
+			}
+		}));
+	
+#undef LOCTEXT_NAMESPACE
+}
+
 void SLevelEditor::AddStandaloneLevelViewport( const TSharedRef<SLevelViewport>& LevelViewport )
 {
 	CleanupPointerArray( StandaloneViewports );
diff --git a/Engine/Source/Editor/LevelEditor/Private/SLevelEditor.h b/Engine/Source/Editor/LevelEditor/Private/SLevelEditor.h
--- a/Engine/Source/Editor/LevelEditor/Private/SLevelEditor.h
+++ b/Engine/Source/Editor/LevelEditor/Private/SLevelEditor.h
@@ -224,6 +224,9 @@ private:
 
 	/** Called when an actor changes outer */
 	void OnLevelActorOuterChanged(AActor* InActor = nullptr, UObject* InOldOuter = nullptr);
+
+	/** Registers toolbar options for the level editor status bar */
+	void RegisterStatusBarTools();
 private:
 
 	// Tracking the active viewports in this level editor.
