Fix for level viewport layout not saving/loading from user settings post CL 14066626.

#rnx

diff --git a/Engine/Source/Editor/LevelEditor/Private/LevelViewportTabContent.cpp b/Engine/Source/Editor/LevelEditor/Private/LevelViewportTabContent.cpp
--- a/Engine/Source/Editor/LevelEditor/Private/LevelViewportTabContent.cpp
+++ b/Engine/Source/Editor/LevelEditor/Private/LevelViewportTabContent.cpp
@@ -19,6 +19,20 @@ TSharedPtr< FEditorViewportLayout > FLevelViewportTabContent::FactoryViewportLay
 	return ViewportLayout;
 }
 
+FName FLevelViewportTabContent::GetLayoutTypeNameFromLayoutString() const
+{
+	const FString& IniSection = FLayoutSaveRestore::GetAdditionalLayoutConfigIni();
+
+	FString LayoutTypeString;
+	if (LayoutString.IsEmpty() ||
+		!GConfig->GetString(*IniSection, *(LayoutString + TEXT(".LayoutType")), LayoutTypeString, GEditorPerProjectIni))
+	{
+		return EditorViewportConfigurationNames::FourPanes2x2;
+	}
+
+	return *LayoutTypeString;
+}
+
 FLevelViewportTabContent::~FLevelViewportTabContent()
 {
 	if (GEditor)
@@ -29,21 +43,13 @@ FLevelViewportTabContent::~FLevelViewportTabContent()
 
 void FLevelViewportTabContent::Initialize(AssetEditorViewportFactoryFunction Func, TSharedPtr<SDockTab> InParentTab, const FString& InLayoutString)
 {
+	check(InParentTab.IsValid());
 	InParentTab->SetOnPersistVisualState( SDockTab::FOnPersistVisualState::CreateSP(this, &FLevelViewportTabContent::SaveConfig) );
 
-	const FString& IniSection = FLayoutSaveRestore::GetAdditionalLayoutConfigIni();
-
-	FString LayoutTypeString;
-	if(InLayoutString.IsEmpty() ||
-		!GConfig->GetString(*IniSection, *(InLayoutString + TEXT(".LayoutType")), LayoutTypeString, GEditorPerProjectIni))
-	{
-		LayoutTypeString = EditorViewportConfigurationNames::FourPanes2x2.ToString();
-	}
-
 	OnViewportTabContentLayoutStartChangeEvent.AddSP(this, &FLevelViewportTabContent::OnLayoutStartChange);
 	OnViewportTabContentLayoutChangedEvent.AddSP(this, &FLevelViewportTabContent::OnLayoutChanged);
 
-	FEditorViewportTabContent::Initialize(Func, InParentTab, LayoutTypeString);
+	FEditorViewportTabContent::Initialize(Func, InParentTab, InLayoutString);
 }
 
 void FLevelViewportTabContent::OnLayoutStartChange(bool bSwitchingLayouts)
diff --git a/Engine/Source/Editor/LevelEditor/Public/LevelViewportTabContent.h b/Engine/Source/Editor/LevelEditor/Public/LevelViewportTabContent.h
--- a/Engine/Source/Editor/LevelEditor/Public/LevelViewportTabContent.h
+++ b/Engine/Source/Editor/LevelEditor/Public/LevelViewportTabContent.h
@@ -23,6 +23,7 @@ public:
 
 protected:
 	virtual TSharedPtr<FEditorViewportLayout> FactoryViewportLayout(bool bIsSwitchingLayouts) override;
+	virtual FName GetLayoutTypeNameFromLayoutString() const override;
 
 	void OnLayoutStartChange(bool bSwitchingLayouts);
 	void OnLayoutChanged();
diff --git a/Engine/Source/Editor/UnrealEd/Private/AssetEditorViewportLayout.cpp b/Engine/Source/Editor/UnrealEd/Private/AssetEditorViewportLayout.cpp
--- a/Engine/Source/Editor/UnrealEd/Private/AssetEditorViewportLayout.cpp
+++ b/Engine/Source/Editor/UnrealEd/Private/AssetEditorViewportLayout.cpp
@@ -215,6 +215,7 @@ TSharedRef<SWidget> FAssetEditorViewportLayout::BuildViewportLayout(TSharedPtr<S
 
 	// Don't set the content until the OverlayPtr has been set, because it access this when we want to start with the viewports maximized.
 	TSharedRef<SWidget> ViewportLayoutWidget = LayoutConfiguration->MakeViewportLayout(this->AsShared(), LayoutString);
+	LoadConfig(LayoutString);
 	ViewportsBorder->SetContent(ViewportLayoutWidget);
 
 	return ViewportsOverlay;
diff --git a/Engine/Source/Editor/UnrealEd/Private/EditorViewportLayout2x2.cpp b/Engine/Source/Editor/UnrealEd/Private/EditorViewportLayout2x2.cpp
--- a/Engine/Source/Editor/UnrealEd/Private/EditorViewportLayout2x2.cpp
+++ b/Engine/Source/Editor/UnrealEd/Private/EditorViewportLayout2x2.cpp
@@ -147,10 +147,8 @@ const FName& FEditorViewportLayout2x2::GetLayoutTypeName() const
 	return EditorViewportConfigurationNames::FourPanes2x2;
 }
 
-void FEditorViewportLayout2x2::SaveLayoutString(const FString& LayoutString) const
+void FEditorViewportLayout2x2::SaveLayoutString(const FString& SpecificLayoutString) const
 {
-	FString SpecificLayoutString = GetTypeSpecificLayoutString(LayoutString);
-
 	const FString& IniSection = FLayoutSaveRestore::GetAdditionalLayoutConfigIni();
 
 	TArray<FVector2D> Percentages;
diff --git a/Engine/Source/Editor/UnrealEd/Private/EditorViewportLayoutFourPanes.cpp b/Engine/Source/Editor/UnrealEd/Private/EditorViewportLayoutFourPanes.cpp
--- a/Engine/Source/Editor/UnrealEd/Private/EditorViewportLayoutFourPanes.cpp
+++ b/Engine/Source/Editor/UnrealEd/Private/EditorViewportLayoutFourPanes.cpp
@@ -92,10 +92,8 @@ TSharedRef<SWidget> FEditorViewportLayoutFourPanes::MakeViewportLayout(TSharedPt
 	return LayoutWidget;
 }
 
-void FEditorViewportLayoutFourPanes::SaveLayoutString(const FString& LayoutString) const
+void FEditorViewportLayoutFourPanes::SaveLayoutString(const FString& SpecificLayoutString) const
 {
-	FString SpecificLayoutString = GetTypeSpecificLayoutString(LayoutString);
-
 	const FString& IniSection = FLayoutSaveRestore::GetAdditionalLayoutConfigIni();
 
 	check(PrimarySplitterWidget->GetChildren()->Num() == 2);
diff --git a/Engine/Source/Editor/UnrealEd/Private/EditorViewportLayoutOnePane.cpp b/Engine/Source/Editor/UnrealEd/Private/EditorViewportLayoutOnePane.cpp
--- a/Engine/Source/Editor/UnrealEd/Private/EditorViewportLayoutOnePane.cpp
+++ b/Engine/Source/Editor/UnrealEd/Private/EditorViewportLayoutOnePane.cpp
@@ -53,7 +53,7 @@ const FName& FEditorViewportLayoutOnePane::GetLayoutTypeName() const
 	return EditorViewportConfigurationNames::OnePane;
 }
 
-void FEditorViewportLayoutOnePane::SaveLayoutString(const FString& LayoutString) const
+void FEditorViewportLayoutOnePane::SaveLayoutString(const FString& SpecificLayoutString) const
 {
 }
 
diff --git a/Engine/Source/Editor/UnrealEd/Private/EditorViewportLayoutTwoPanes.cpp b/Engine/Source/Editor/UnrealEd/Private/EditorViewportLayoutTwoPanes.cpp
--- a/Engine/Source/Editor/UnrealEd/Private/EditorViewportLayoutTwoPanes.cpp
+++ b/Engine/Source/Editor/UnrealEd/Private/EditorViewportLayoutTwoPanes.cpp
@@ -96,14 +96,13 @@ void TEditorViewportLayoutTwoPanes<TOrientation>::ReplaceWidget(TSharedRef<SWidg
 }
 
 template <EOrientation TOrientation>
-void TEditorViewportLayoutTwoPanes<TOrientation>::SaveLayoutString(const FString& LayoutString) const
+void TEditorViewportLayoutTwoPanes<TOrientation>::SaveLayoutString(const FString& SpecificLayoutString) const
 {
 	const FString& IniSection = FLayoutSaveRestore::GetAdditionalLayoutConfigIni();
 
 	check(SplitterWidget->GetChildren()->Num() == 2);
 	float Percentage = SplitterWidget->SlotAt(0).SizeValue.Get();
 
-	FString SpecificLayoutString = GetTypeSpecificLayoutString(LayoutString);
 	GConfig->SetString(*IniSection, *(SpecificLayoutString + TEXT(".Percentage")), *TTypeToString<float>::ToString(Percentage), GEditorPerProjectIni);
 }
 
diff --git a/Engine/Source/Editor/UnrealEd/Private/EditorViewportTabContent.cpp b/Engine/Source/Editor/UnrealEd/Private/EditorViewportTabContent.cpp
--- a/Engine/Source/Editor/UnrealEd/Private/EditorViewportTabContent.cpp
+++ b/Engine/Source/Editor/UnrealEd/Private/EditorViewportTabContent.cpp
@@ -21,7 +21,7 @@ void FEditorViewportTabContent::Initialize(AssetEditorViewportFactoryFunction Fu
 	ParentTab = InParentTab;
 	LayoutString = InLayoutString;
 
-	FName LayoutType(*LayoutString);
+	FName LayoutType = GetLayoutTypeNameFromLayoutString();
 	ViewportCreationFactories.Add(AssetEditorViewportCreationFactories::ElementType(NAME_None, Func));
 	SetViewportConfiguration(LayoutType);
 }
@@ -131,6 +131,11 @@ const AssetEditorViewportFactoryFunction* FEditorViewportTabContent::FindViewpor
 	return ViewportCreationFactories.Find(InTypeName);
 }
 
+FName FEditorViewportTabContent::GetLayoutTypeNameFromLayoutString() const
+{
+	return *LayoutString;
+}
+
 TSharedPtr<FEditorViewportLayout> FEditorViewportTabContent::FactoryViewportLayout(bool bIsSwitchingLayouts)
 {
 	return MakeShareable(new FAssetEditorViewportLayout);
diff --git a/Engine/Source/Editor/UnrealEd/Public/EditorViewportLayout2x2.h b/Engine/Source/Editor/UnrealEd/Public/EditorViewportLayout2x2.h
--- a/Engine/Source/Editor/UnrealEd/Public/EditorViewportLayout2x2.h
+++ b/Engine/Source/Editor/UnrealEd/Public/EditorViewportLayout2x2.h
@@ -16,7 +16,7 @@ public:
 	virtual TSharedRef<SWidget> MakeViewportLayout(TSharedPtr<FAssetEditorViewportLayout> InParentLayout, const FString& LayoutString) override;
 	virtual void ReplaceWidget(TSharedRef<SWidget> OriginalWidget, TSharedRef<SWidget> ReplacementWidget) override;
 	virtual const FName& GetLayoutTypeName() const override;
-	virtual void SaveLayoutString(const FString& LayoutString) const override;
+	virtual void SaveLayoutString(const FString& SpecificLayoutString) const override;
 
 protected:
 	/** The splitter widget */
diff --git a/Engine/Source/Editor/UnrealEd/Public/EditorViewportLayoutFourPanes.h b/Engine/Source/Editor/UnrealEd/Public/EditorViewportLayoutFourPanes.h
--- a/Engine/Source/Editor/UnrealEd/Public/EditorViewportLayoutFourPanes.h
+++ b/Engine/Source/Editor/UnrealEd/Public/EditorViewportLayoutFourPanes.h
@@ -13,7 +13,7 @@ class FEditorViewportLayoutFourPanes : public FAssetEditorViewportPaneLayout
 public:
 	// IAssetEditorViewportLayoutConfiguration overrides
 	virtual TSharedRef<SWidget> MakeViewportLayout(TSharedPtr<FAssetEditorViewportLayout> InParentLayout, const FString& LayoutString) override;
-	virtual void SaveLayoutString(const FString& LayoutString) const override;
+	virtual void SaveLayoutString(const FString& SpecificLayoutString) const override;
 	virtual void ReplaceWidget(TSharedRef<SWidget> OriginalWidget, TSharedRef<SWidget> ReplacementWidget) override;
 
 protected:
diff --git a/Engine/Source/Editor/UnrealEd/Public/EditorViewportLayoutOnePane.h b/Engine/Source/Editor/UnrealEd/Public/EditorViewportLayoutOnePane.h
--- a/Engine/Source/Editor/UnrealEd/Public/EditorViewportLayoutOnePane.h
+++ b/Engine/Source/Editor/UnrealEd/Public/EditorViewportLayoutOnePane.h
@@ -16,7 +16,7 @@ public:
 	// IAssetEditorViewportLayoutConfiguration overrides
 	virtual TSharedRef<SWidget> MakeViewportLayout(TSharedPtr<FAssetEditorViewportLayout> InParentLayout, const FString& LayoutString) override;
 	virtual const FName& GetLayoutTypeName() const override;
-	virtual void SaveLayoutString(const FString& LayoutString) const override;
+	virtual void SaveLayoutString(const FString& SpecificLayoutString) const override;
 	virtual void ReplaceWidget(TSharedRef<SWidget> OriginalWidget, TSharedRef<SWidget> ReplacementWidget) override;
 
 protected:
diff --git a/Engine/Source/Editor/UnrealEd/Public/EditorViewportLayoutThreePanes.h b/Engine/Source/Editor/UnrealEd/Public/EditorViewportLayoutThreePanes.h
--- a/Engine/Source/Editor/UnrealEd/Public/EditorViewportLayoutThreePanes.h
+++ b/Engine/Source/Editor/UnrealEd/Public/EditorViewportLayoutThreePanes.h
@@ -15,7 +15,7 @@ public:
 	 * Creates the viewports and splitter for the two panes vertical layout                   
 	 */
 	virtual TSharedRef<SWidget> MakeViewportLayout(TSharedPtr<FAssetEditorViewportLayout> InParentLayout, const FString& LayoutString) override;
-	virtual void SaveLayoutString(const FString& LayoutString) const override;
+	virtual void SaveLayoutString(const FString& SpecificLayoutString) const override;
 	virtual void ReplaceWidget(TSharedRef<SWidget> OriginalWidget, TSharedRef<SWidget> ReplacementWidget) override;
 
 protected:
diff --git a/Engine/Source/Editor/UnrealEd/Public/EditorViewportLayoutTwoPanes.h b/Engine/Source/Editor/UnrealEd/Public/EditorViewportLayoutTwoPanes.h
--- a/Engine/Source/Editor/UnrealEd/Public/EditorViewportLayoutTwoPanes.h
+++ b/Engine/Source/Editor/UnrealEd/Public/EditorViewportLayoutTwoPanes.h
@@ -16,7 +16,7 @@ class TEditorViewportLayoutTwoPanes : public FAssetEditorViewportPaneLayout
 public:
 	virtual TSharedRef<SWidget> MakeViewportLayout(TSharedPtr<FAssetEditorViewportLayout> InParentLayout, const FString& LayoutString) override;
 	virtual void ReplaceWidget(TSharedRef<SWidget> OriginalWidget, TSharedRef<SWidget> ReplacementWidget) override;
-	virtual void SaveLayoutString(const FString& LayoutString) const override;
+	virtual void SaveLayoutString(const FString& SpecificLayoutString) const override;
 
 private:
 	/** The splitter widget */
diff --git a/Engine/Source/Editor/UnrealEd/Public/EditorViewportTabContent.h b/Engine/Source/Editor/UnrealEd/Public/EditorViewportTabContent.h
--- a/Engine/Source/Editor/UnrealEd/Public/EditorViewportTabContent.h
+++ b/Engine/Source/Editor/UnrealEd/Public/EditorViewportTabContent.h
@@ -64,6 +64,7 @@ public:
 
 protected:
 	virtual TSharedPtr<FEditorViewportLayout> FactoryViewportLayout(bool bIsSwitchingLayouts);
+	virtual FName GetLayoutTypeNameFromLayoutString() const;
 
 private:
 	using AssetEditorViewportCreationFactories = TMap<FName, AssetEditorViewportFactoryFunction>;
