Fix for maximize viewport button gone missing in CL 14066626.

#rnx

diff --git a/Engine/Source/Editor/LevelEditor/Private/LevelViewportLayout.cpp b/Engine/Source/Editor/LevelEditor/Private/LevelViewportLayout.cpp
--- a/Engine/Source/Editor/LevelEditor/Private/LevelViewportLayout.cpp
+++ b/Engine/Source/Editor/LevelEditor/Private/LevelViewportLayout.cpp
@@ -114,7 +114,7 @@ TSharedRef<SWidget> FLevelViewportLayout::FactoryViewport(FName InTypeName, cons
 	// Both the legacy level editor factory viewport, and the viewport creation functions should fall back to a valid default
 	check(ViewportLayoutEntity.IsValid());
 	
-	Viewports.Add(InTypeName, ViewportLayoutEntity);
+	Viewports.Add(ConstructionArgs.ConfigKey, ViewportLayoutEntity);
 	return ViewportLayoutEntity->AsWidget();
 }
 
