Adjusted editor default layout

diff --git a/Engine/Source/Editor/LevelEditor/Private/SLevelEditor.cpp b/Engine/Source/Editor/LevelEditor/Private/SLevelEditor.cpp
--- a/Engine/Source/Editor/LevelEditor/Private/SLevelEditor.cpp
+++ b/Engine/Source/Editor/LevelEditor/Private/SLevelEditor.cpp
@@ -1300,7 +1300,7 @@ TSharedRef<SWidget> SLevelEditor::RestoreContentArea( const TSharedRef<SDockTab>
 	// 10. Push the new "DefaultLayout.ini" together with your new code.
 	// 11. Also update these instructions if you change the version number (e.g., from "UnrealEd_Layout_v1.4" to "UnrealEd_Layout_v1.5").
 	const TSharedRef<FTabManager::FLayout> DefaultLayout = FLayoutSaveRestore::LoadFromConfig(GEditorLayoutIni,
-		FTabManager::NewLayout( "LevelEditor_Layout_v1.3" )
+		FTabManager::NewLayout( "LevelEditor_Layout_v1.5" )
 		->AddArea
 		(
 			FTabManager::NewPrimaryArea()
@@ -1320,9 +1320,8 @@ TSharedRef<SWidget> SLevelEditor::RestoreContentArea( const TSharedRef<SDockTab>
 					(
 						FTabManager::NewStack()
 						->SetSizeCoefficient( 0.3f )
-						->AddTab(LevelEditorTabIds::PlacementBrowser, ETabState::OpenedTab)
+						->AddTab(LevelEditorTabIds::PlacementBrowser, ETabState::SidebarTab, ESidebarLocation::Left, 0.12f)
 						->AddTab(LevelEditorTabIds::LevelEditorToolBox, ETabState::ClosedTab)
-						->SetForegroundTab(LevelEditorTabIds::PlacementBrowser)
 					)
 					->Split
 					(
