Fix compile button options not functioning on status bar

diff --git a/Engine/Source/Editor/LevelEditor/Private/SLevelEditor.cpp b/Engine/Source/Editor/LevelEditor/Private/SLevelEditor.cpp
--- a/Engine/Source/Editor/LevelEditor/Private/SLevelEditor.cpp
+++ b/Engine/Source/Editor/LevelEditor/Private/SLevelEditor.cpp
@@ -1697,14 +1697,14 @@ void SLevelEditor::RegisterStatusBarTools()
 		UToolMenu* Menu = UToolMenus::Get()->RegisterMenu("LevelEditor.StatusBar.ToolBar.CompileComboButton");
 		{
 			FToolMenuSection& Section = Menu->AddSection("LiveCodingMode", LOCTEXT("LiveCodingMode", "General"));
-			Section.AddMenuEntry(FLevelEditorCommands::Get().LiveCoding_Enable);
+			Section.AddMenuEntryWithCommandList(FLevelEditorCommands::Get().LiveCoding_Enable, LevelEditorCommands);
 		}
 
 		{
 			FToolMenuSection& Section = Menu->AddSection("LiveCodingActions", LOCTEXT("LiveCodingActions", "Actions"));
-			Section.AddMenuEntry(FLevelEditorCommands::Get().LiveCoding_StartSession);
-			Section.AddMenuEntry(FLevelEditorCommands::Get().LiveCoding_ShowConsole);
-			Section.AddMenuEntry(FLevelEditorCommands::Get().LiveCoding_Settings);
+			Section.AddMenuEntryWithCommandList(FLevelEditorCommands::Get().LiveCoding_StartSession, LevelEditorCommands);
+			Section.AddMenuEntryWithCommandList(FLevelEditorCommands::Get().LiveCoding_ShowConsole, LevelEditorCommands);
+			Section.AddMenuEntryWithCommandList(FLevelEditorCommands::Get().LiveCoding_Settings, LevelEditorCommands);
 		}
 	}
 
@@ -1728,7 +1728,7 @@ void SLevelEditor::RegisterStatusBarTools()
 						FCanExecuteAction::CreateStatic(&FLevelEditorActionCallbacks::Recompile_CanExecute),
 						FIsActionChecked(),
 						FIsActionButtonVisible::CreateStatic(FLevelEditorActionCallbacks::CanShowSourceCodeActions)),
-					LOCTEXT("CompileMenuButton", "Compile"),
+					FText::GetEmpty(),
 					FLevelEditorCommands::Get().RecompileGameCode->GetDescription(),
 					FSlateIcon(FEditorStyle::GetStyleSetName(), "LevelEditor.Recompile")
 				));
