Switching USelection to use element lists by default for actors and components

#rnx

diff --git a/Engine/Source/Editor/LevelEditor/LevelEditor.Build.cs b/Engine/Source/Editor/LevelEditor/LevelEditor.Build.cs
--- a/Engine/Source/Editor/LevelEditor/LevelEditor.Build.cs
+++ b/Engine/Source/Editor/LevelEditor/LevelEditor.Build.cs
@@ -58,6 +58,7 @@ public class LevelEditor : ModuleRules
 				"RHI",
 				"Projects",
 				"TargetPlatform",
+				"TypedElementFramework",
 				"EngineSettings",
 				"PropertyEditor",
 				"Kismet",
diff --git a/Engine/Source/Editor/LevelEditor/Private/SLevelEditor.cpp b/Engine/Source/Editor/LevelEditor/Private/SLevelEditor.cpp
--- a/Engine/Source/Editor/LevelEditor/Private/SLevelEditor.cpp
+++ b/Engine/Source/Editor/LevelEditor/Private/SLevelEditor.cpp
@@ -57,6 +57,7 @@
 #include "Widgets/Colors/SColorPicker.h"
 #include "Editor/EnvironmentLightingViewer/Public/EnvironmentLightingModule.h"
 #include "SourceCodeNavigation.h"
+#include "TypedElementRegistry.h"
 
 static const FName MainFrameModuleName("MainFrame");
 static const FName LevelEditorModuleName("LevelEditor");
@@ -210,6 +211,12 @@ void SLevelEditor::Construct( const SLevelEditor::FArguments& InArgs)
 
 void SLevelEditor::Initialize( const TSharedRef<SDockTab>& OwnerTab, const TSharedRef<SWindow>& OwnerWindow )
 {
+	SelectedElements = UTypedElementRegistry::GetInstance()->CreateElementList();
+
+	// Allow USelection to bridge to our selected element list
+	GUnrealEd->GetSelectedActors()->SetElementList(SelectedElements.Get());
+	GUnrealEd->GetSelectedComponents()->SetElementList(SelectedElements.Get());
+
 	// Bind the level editor tab's label to the currently loaded level name string in the main frame
 	OwnerTab->SetLabel( TAttribute<FText>( this, &SLevelEditor::GetTabTitle) );
 	OwnerTab->SetTabLabelSuffix(TAttribute<FText>(this, &SLevelEditor::GetTabSuffix));
@@ -312,6 +319,14 @@ SLevelEditor::~SLevelEditor()
 
 		GEditor->GetEditorSubsystem<UAssetEditorSubsystem>()->OnEditorModesChanged().RemoveAll(this);
 	}
+
+	// Clear USelection from using our selected element list
+	if (GUnrealEd)
+	{
+		GUnrealEd->GetSelectedActors()->SetElementList(nullptr);
+		GUnrealEd->GetSelectedComponents()->SetElementList(nullptr);
+	}
+	SelectedElements.Reset();
 }
 
 FText SLevelEditor::GetTabTitle() const
diff --git a/Engine/Source/Editor/LevelEditor/Private/SLevelEditor.h b/Engine/Source/Editor/LevelEditor/Private/SLevelEditor.h
--- a/Engine/Source/Editor/LevelEditor/Private/SLevelEditor.h
+++ b/Engine/Source/Editor/LevelEditor/Private/SLevelEditor.h
@@ -15,6 +15,7 @@
 #include "AssetThumbnail.h"
 #include "ILevelEditor.h"
 #include "LevelViewportTabContent.h"
+#include "TypedElementListPtr.h"
 
 class IAssetEditorInstance;
 class IDetailsView;
@@ -252,6 +253,9 @@ private:
 	// The UWorld that this level editor is viewing and allowing the user to interact with through.
 	UWorld* World;
 
+	// The list of selected elements (also set on the global USelection for actors and components).
+	FTypedElementListPtr SelectedElements;
+
 	// The box that holds the title bar messages.
 	TSharedPtr<SHorizontalBox> TtileBarMessageBox;
 
diff --git a/Engine/Source/Editor/UnrealEd/Private/Selection.cpp b/Engine/Source/Editor/UnrealEd/Private/Selection.cpp
--- a/Engine/Source/Editor/UnrealEd/Private/Selection.cpp
+++ b/Engine/Source/Editor/UnrealEd/Private/Selection.cpp
@@ -5,7 +5,7 @@
 #include "UObject/UObjectAnnotation.h"
 #include "GameFramework/Actor.h"
 
-#define UE_USE_ELEMENT_LIST_SELECTION (0)
+#define UE_USE_ELEMENT_LIST_SELECTION (1)
 
 #if UE_USE_ELEMENT_LIST_SELECTION
 #include "TypedElementList.h"
