Fixing a shutdown crash in the global editor mode manager.

#rnx

diff --git a/Engine/Source/Editor/LevelEditor/Private/SLevelEditor.cpp b/Engine/Source/Editor/LevelEditor/Private/SLevelEditor.cpp
--- a/Engine/Source/Editor/LevelEditor/Private/SLevelEditor.cpp
+++ b/Engine/Source/Editor/LevelEditor/Private/SLevelEditor.cpp
@@ -1527,7 +1527,7 @@ void SLevelEditor::RefreshEditorModeCommands()
 				Commands.EditorModeCommands[CommandIndex],
 				FExecuteAction::CreateStatic( &SLevelEditor::ToggleEditorMode, Mode.ID ),
 				FCanExecuteAction(),
-				FIsActionChecked::CreateSP( GetEditorModeManager().AsShared(), &FEditorModeTools::IsModeActive, Mode.ID ));
+				FIsActionChecked::CreateStatic( &SLevelEditor::IsModeActive, Mode.ID ));
 		}
 
 		CommandIndex++;
diff --git a/Engine/Source/Editor/LevelEditor/Private/SLevelEditor.h b/Engine/Source/Editor/LevelEditor/Private/SLevelEditor.h
--- a/Engine/Source/Editor/LevelEditor/Private/SLevelEditor.h
+++ b/Engine/Source/Editor/LevelEditor/Private/SLevelEditor.h
@@ -181,7 +181,6 @@ private:
 	static void ToggleEditorMode( FEditorModeID ModeID );
 
 	/** Checks if the editor mode is active for the auto-generated editor mode command. */
-	UE_DEPRECATED(5.0, "Use GLevelEditorModeTools().IsModeActive.")
 	static bool IsModeActive( FEditorModeID ModeID );
 
 	/**
diff --git a/Engine/Source/Editor/UnrealEd/Private/UnrealEdGlobals.cpp b/Engine/Source/Editor/UnrealEd/Private/UnrealEdGlobals.cpp
--- a/Engine/Source/Editor/UnrealEd/Private/UnrealEdGlobals.cpp
+++ b/Engine/Source/Editor/UnrealEd/Private/UnrealEdGlobals.cpp
@@ -46,7 +46,7 @@ DEFINE_LOG_CATEGORY_STATIC(LogUnrealEd, Log, All);
  */
 namespace Internal
 {
-	static TSharedPtr<FEditorModeTools> GetGlobalModeManager()
+	static TSharedPtr<FEditorModeTools>& GetGlobalModeManager()
 	{
 		static TSharedPtr<FEditorModeTools> EditorModeToolsSingleton = MakeShared<FEditorModeTools>();
 		return EditorModeToolsSingleton;
