Added missing null check when walking parent component tree

#rnx

diff --git a/Engine/Source/Editor/LevelEditor/Private/Elements/Component/ComponentElementLevelEditorSelectionProxy.cpp b/Engine/Source/Editor/LevelEditor/Private/Elements/Component/ComponentElementLevelEditorSelectionProxy.cpp
--- a/Engine/Source/Editor/LevelEditor/Private/Elements/Component/ComponentElementLevelEditorSelectionProxy.cpp
+++ b/Engine/Source/Editor/LevelEditor/Private/Elements/Component/ComponentElementLevelEditorSelectionProxy.cpp
@@ -41,7 +41,7 @@ FTypedElementHandle UComponentElementLevelEditorSelectionProxy::GetSelectionElem
 	{
 		const AActor* ConsideredActor = ComponentData->Component->GetOwner();
 		const USceneComponent* ConsideredComponent = Cast<USceneComponent>(ComponentData->Component);
-		if (ConsideredActor && ConsideredComponent)
+		if (ConsideredActor)
 		{
 			while (ConsideredActor->IsChildActor())
 			{
@@ -50,7 +50,7 @@ FTypedElementHandle UComponentElementLevelEditorSelectionProxy::GetSelectionElem
 			}
 
 			// If the component selected is a visualization component, we want to select the non-visualization component it's attached to
-			while (ConsideredComponent->IsVisualizationComponent())
+			while (ConsideredComponent && ConsideredComponent->IsVisualizationComponent())
 			{
 				ConsideredComponent = ConsideredComponent->GetAttachParent();
 			}
@@ -69,7 +69,7 @@ FTypedElementHandle UComponentElementLevelEditorSelectionProxy::GetSelectionElem
 
 			const bool bSelectComponent = bActorAlreadySelectedExclusively && bActorIsBlueprintable && (bComponentAlreadySelected != bWasDoubleClick);
 
-			if (bSelectComponent)
+			if (bSelectComponent && ConsideredComponent)
 			{
 				return ConsideredComponent->AcquireEditorElementHandle();
 			}
diff --git a/Engine/Source/Editor/UnrealEd/Private/Elements/Component/ComponentElementEditorSelectionInterface.cpp b/Engine/Source/Editor/UnrealEd/Private/Elements/Component/ComponentElementEditorSelectionInterface.cpp
--- a/Engine/Source/Editor/UnrealEd/Private/Elements/Component/ComponentElementEditorSelectionInterface.cpp
+++ b/Engine/Source/Editor/UnrealEd/Private/Elements/Component/ComponentElementEditorSelectionInterface.cpp
@@ -28,7 +28,7 @@ bool UComponentElementEditorSelectionInterface::IsComponentSelected(const UActor
 	{
 		const AActor* ConsideredActor = InComponent->GetOwner();
 		const USceneComponent* ConsideredComponent = Cast<USceneComponent>(InComponent);
-		if (ConsideredActor && ConsideredComponent)
+		if (ConsideredActor)
 		{
 			while (ConsideredActor->IsChildActor())
 			{
@@ -36,13 +36,16 @@ bool UComponentElementEditorSelectionInterface::IsComponentSelected(const UActor
 				ConsideredComponent = ConsideredActor->GetParentComponent();
 			}
 
-			while (ConsideredComponent->IsVisualizationComponent())
+			while (ConsideredComponent && ConsideredComponent->IsVisualizationComponent())
 			{
 				ConsideredComponent = ConsideredComponent->GetAttachParent();
 			}
 		}
 
-		return InSelectionSet->Contains(ConsideredComponent->AcquireEditorElementHandle());
+		if (ConsideredComponent)
+		{
+			return InSelectionSet->Contains(ConsideredComponent->AcquireEditorElementHandle());
+		}
 	}
 
 	return false;
