Don't allow "save world as" for partitionned worlds as it is unsupported for now.

diff --git a/Engine/Source/Editor/LevelEditor/Private/LevelEditorActions.cpp b/Engine/Source/Editor/LevelEditor/Private/LevelEditorActions.cpp
--- a/Engine/Source/Editor/LevelEditor/Private/LevelEditorActions.cpp
+++ b/Engine/Source/Editor/LevelEditor/Private/LevelEditorActions.cpp
@@ -36,6 +36,7 @@
 #include "FileHelpers.h"
 #include "EditorModeInterpolation.h"
 #include "UnrealEdGlobals.h"
+#include "WorldPartition/WorldPartitionSubsystem.h"
 
 #include "LevelEditor.h"
 #include "Matinee/MatineeActor.h"
@@ -405,6 +406,12 @@ void FLevelEditorActionCallbacks::SaveCurrentAs()
 {
 	UWorld* World = GetWorld();
 	ULevel* CurrentLevel = World->GetCurrentLevel();
+
+	if (World->HasSubsystem<UWorldPartitionSubsystem>())
+	{
+		FMessageDialog::Open(EAppMsgType::Ok, NSLOCTEXT("UnrealEd", "SaveWorldAsFailed", "Save as not supported for partitioned worlds"));
+		return;
+	}
 	
 	UClass* CurrentStreamingLevelClass = ULevelStreamingDynamic::StaticClass();
 
