Add support for creating a new world using external actors and world partition in the editor.

diff --git a/Engine/Source/Editor/LevelEditor/Private/LevelEditorActions.cpp b/Engine/Source/Editor/LevelEditor/Private/LevelEditorActions.cpp
--- a/Engine/Source/Editor/LevelEditor/Private/LevelEditorActions.cpp
+++ b/Engine/Source/Editor/LevelEditor/Private/LevelEditorActions.cpp
@@ -37,6 +37,10 @@
 #include "EditorModeInterpolation.h"
 #include "UnrealEdGlobals.h"
 #include "WorldPartition/WorldPartitionSubsystem.h"
+#include "WorldPartition/WorldPartition.h"
+#include "WorldPartition/WorldPartitionEditorHash.h"
+#include "WorldPartition/WorldPartitionRuntimeHash.h"
+#include "Editor/WorldBrowser/Public/WorldBrowserModule.h"
 
 #include "LevelEditor.h"
 #include "Matinee/MatineeActor.h"
@@ -195,12 +199,20 @@ void FLevelEditorActionCallbacks::NewLevel()
 
 	FString TemplateMapPackageName;
 	FNewLevelDialogModule& NewLevelDialogModule = FModuleManager::LoadModuleChecked<FNewLevelDialogModule>("NewLevelDialog");
-	if (NewLevelDialogModule.CreateAndShowNewLevelDialog(MainFrameModule.GetParentWindow(), TemplateMapPackageName))
+	
+	FNewLevelDialogModule::FNewLevelOptions NewLevelOptions;
+	NewLevelOptions.bExternalActors = true;
+	NewLevelOptions.bPartitionedWorld = true;
+
+	if (NewLevelDialogModule.CreateAndShowNewLevelDialog(MainFrameModule.GetParentWindow(), TemplateMapPackageName, &NewLevelOptions))
 	{
+		bool bChangedMap = false;
+
 		// The new map screen will return a blank TemplateName if the user has selected to begin a new blank map
 		if (TemplateMapPackageName.IsEmpty())
 		{
 			GEditor->CreateNewMapForEditing();
+			bChangedMap = true;
 		}
 		else
 		{
@@ -218,14 +230,11 @@ void FLevelEditorActionCallbacks::NewLevel()
 					TemplateFound = true;
 
 					// If there are any unsaved changes to the current level, see if the user wants to save those first.
-					bool bPromptUserToSave = true;
-					bool bSaveMapPackages = true;
-					bool bSaveContentPackages = false;
-					if ( FEditorFileUtils::SaveDirtyPackages(bPromptUserToSave, bSaveMapPackages, bSaveContentPackages) )
+					if ( FEditorFileUtils::SaveDirtyPackages(/*bPromptUserToSave*/true, /*bSaveMapPackages*/true, /*bSaveContentPackages*/false) )
 					{
 						// Load the template map file - passes LoadAsTemplate==true making the
 						// level load into an untitled package that won't save over the template
-						FEditorFileUtils::LoadMap(*MapPackageFilename, /*bLoadAsTemplate=*/true);
+						bChangedMap = FEditorFileUtils::LoadMap(*MapPackageFilename, /*bLoadAsTemplate=*/true);
 					}
 				}
 			}
@@ -236,6 +245,80 @@ void FLevelEditorActionCallbacks::NewLevel()
 				GEditor->CreateNewMapForEditing();
 			}
 		}
+
+		if (bChangedMap)
+		{
+			// if target world needs to be saved (user asked for external actors or world partition)
+			if (NewLevelOptions.bExternalActors || NewLevelOptions.bPartitionedWorld)
+			{
+				UWorld* World = GetWorld();
+
+				if (World->GetStreamingLevels().Num())
+				{
+					FMessageDialog::Open(EAppMsgType::Ok, NSLOCTEXT("UnrealEd", "NewLevelContainsStreamingLevels", "The provided template contains streaming levels so it can't be converted with the provided options"));
+				}
+				else if (World->WorldComposition)
+				{
+					FMessageDialog::Open(EAppMsgType::Ok, NSLOCTEXT("UnrealEd", "NewLevelContainsExternalActors", "The provided template is using world composition so it can't be converted with the provided options"));
+				}
+				else if (World->PersistentLevel->bUseExternalActors)
+				{
+					FMessageDialog::Open(EAppMsgType::Ok, NSLOCTEXT("UnrealEd", "NewLevelContainsExternalActors", "The provided template contains external actors so it can't be converted with the provided options"));
+				}
+				else if (World->PersistentLevel->bIsPartitioned)
+				{
+					FMessageDialog::Open(EAppMsgType::Ok, NSLOCTEXT("UnrealEd", "NewLevelContainsExternalActors", "The provided template is already partitioned so it can't be converted with the provided options"));
+				}
+				else
+				{
+					bool bCanConvert = true;
+					FString SavedFilename;
+
+					while (!FEditorFileUtils::SaveLevelAs(GetWorld()->PersistentLevel, &SavedFilename))
+					{
+						if (FMessageDialog::Open(EAppMsgType::YesNo, EAppReturnType::Yes, LOCTEXT("NewLevelNeedsSave", "The new level needs to be saved prior to converting to external actors, do you want to save it?")) == EAppReturnType::No)
+						{
+							bCanConvert = false;
+							break;
+						}
+					}
+
+					if (bCanConvert)
+					{
+						FEditorFileUtils::LoadMap(SavedFilename);
+
+						// Update the world pointer as we reloaded the map
+						World = GetWorld();
+				
+						if (NewLevelOptions.bExternalActors)
+						{
+							World->PersistentLevel->ConvertAllActorsToPackaging(true);
+							World->PersistentLevel->bUseExternalActors = true;
+						}
+
+						if (NewLevelOptions.bPartitionedWorld)
+						{
+							check(World->PersistentLevel->bUseExternalActors);
+							check(!World->GetStreamingLevels().Num());
+
+							AWorldSettings* WorldSettings = World->GetWorldSettings();
+
+							UWorldPartition* WorldPartition = NewObject<UWorldPartition>(WorldSettings);
+							WorldSettings->SetWorldPartition(WorldPartition);
+							WorldSettings->MarkPackageDirty();
+
+							World->PersistentLevel->bIsPartitioned = true;
+							World->ReinitializeSubSystems();
+
+							WorldPartition->Initialize(World, FTransform::Identity);					
+
+							FWorldBrowserModule& WorldBrowserModule = FModuleManager::GetModuleChecked<FWorldBrowserModule>("WorldBrowser");
+							WorldBrowserModule.OnBrowseWorld.Broadcast(World);
+						}
+					}
+				}
+			}
+		}
 	}
 }
 
diff --git a/Engine/Source/Editor/NewLevelDialog/Private/NewLevelDialogModule.cpp b/Engine/Source/Editor/NewLevelDialog/Private/NewLevelDialogModule.cpp
--- a/Engine/Source/Editor/NewLevelDialog/Private/NewLevelDialogModule.cpp
+++ b/Engine/Source/Editor/NewLevelDialog/Private/NewLevelDialogModule.cpp
@@ -93,6 +93,8 @@ public:
 		/** A pointer to the parent window */
 		SLATE_ATTRIBUTE(TSharedPtr<SWindow>, ParentWindow)
 		SLATE_ATTRIBUTE(TArray<FTemplateMapInfo>, Templates)
+		SLATE_ATTRIBUTE(bool, bExternalActors)
+		SLATE_ATTRIBUTE(bool, bPartitionedWorld)
 
 	SLATE_END_ARGS()
 
@@ -102,6 +104,8 @@ public:
 
 		OutTemplateMapPackageName = TEXT("");
 		bUserClickedOkay = false;
+		bExternalActors = false;
+		bPartitionedWorld = false;
 
 		const TArray<FTemplateMapInfo>& TemplateInfoList = InArgs._Templates.Get();
 
@@ -124,6 +128,9 @@ public:
 			.Text(LOCTEXT("Cancel", "Cancel"))
 			.OnClicked(this, &SNewLevelDialog::OnCancelClicked);
 
+		const bool bExternalActorsEnabled = InArgs._bExternalActors.Get();
+		const bool bPartitionedWorldEnabled = InArgs._bPartitionedWorld.Get();
+
 		this->ChildSlot
 		[
 			SNew(SBorder)
@@ -131,6 +138,34 @@ public:
 			[
 				SNew(SVerticalBox)
 				+SVerticalBox::Slot()
+				.AutoHeight()
+				[
+					SNew(SVerticalBox)
+					+SVerticalBox::Slot()
+					[
+						SNew(SCheckBox)
+						.IsChecked_Lambda([=]() { return (bExternalActors || bPartitionedWorld) ? ECheckBoxState::Checked : ECheckBoxState::Unchecked; })
+						.IsEnabled_Lambda([=]() { return bExternalActorsEnabled && !bPartitionedWorld; })
+						.OnCheckStateChanged(FOnCheckStateChanged::CreateLambda([=](ECheckBoxState State) { bExternalActors = !bExternalActors; }))
+						[
+							SNew(STextBlock)
+							.Text(LOCTEXT("ExternalActors", "External Actors"))
+							.IsEnabled_Lambda([=]() { return bExternalActorsEnabled && !bPartitionedWorld; })
+						]
+					]
+					+SVerticalBox::Slot()
+					[
+						SNew(SCheckBox)
+						.IsChecked(ECheckBoxState::Unchecked)
+						.IsEnabled(bPartitionedWorldEnabled)
+						.OnCheckStateChanged(FOnCheckStateChanged::CreateLambda([=](ECheckBoxState State) { bPartitionedWorld = !bPartitionedWorld; }))
+						[
+							SNew(STextBlock)
+							.Text(LOCTEXT("PartitionedWorld", "Partitioned World"))
+						]
+					]
+				]
+				+SVerticalBox::Slot()
 				.FillHeight(1)
 				[
 					SNew(SScrollBox)
@@ -166,6 +201,8 @@ public:
 
 	FString GetChosenTemplate() const { return OutTemplateMapPackageName; }
 	bool IsTemplateChosen() const { return bUserClickedOkay; }
+	bool GetExternalActors() const { return bExternalActors; }
+	bool GetPartitionedWorld() const { return bPartitionedWorld; }
 
 private:
 	void AddItemsToWrapBox()
@@ -250,6 +287,7 @@ private:
 			OutTemplateMapPackageName = Template->TemplateMapInfo.Map;
 		}
 		bUserClickedOkay = true;
+		bExternalActors |= bPartitionedWorld;
 
 		ParentWindowPtr.Pin()->RequestDestroyWindow();
 		return FReply::Handled();
@@ -292,6 +330,8 @@ private:
 	TSharedPtr<SWrapBox> TemplatesWrapBox;
 	FString OutTemplateMapPackageName;
 	bool bUserClickedOkay;
+	bool bExternalActors;
+	bool bPartitionedWorld;
 };
 
 const FVector2D SNewLevelDialog::DEFAULT_WINDOW_SIZE = FVector2D(527, 418);
@@ -309,13 +349,13 @@ void FNewLevelDialogModule::ShutdownModule()
 {
 }
 
-bool FNewLevelDialogModule::CreateAndShowNewLevelDialog(const TSharedPtr<const SWidget> ParentWidget, FString& OutTemplateMapPackageName)
+bool FNewLevelDialogModule::CreateAndShowNewLevelDialog( const TSharedPtr<const SWidget> ParentWidget, FString& OutTemplateMapPackageName, FNewLevelOptions* InOutNewLevelOptions )
 {
 	TArray<FTemplateMapInfo> EmptyTemplates;
-	return CreateAndShowTemplateDialog(ParentWidget, LOCTEXT("WindowHeader", "New Level"), GUnrealEd ? GUnrealEd->TemplateMapInfos : EmptyTemplates, OutTemplateMapPackageName);
+	return CreateAndShowTemplateDialog(ParentWidget, LOCTEXT("WindowHeader", "New Level"), GUnrealEd ? GUnrealEd->TemplateMapInfos : EmptyTemplates, OutTemplateMapPackageName, InOutNewLevelOptions);
 }
 
-bool FNewLevelDialogModule::CreateAndShowTemplateDialog( const TSharedPtr<const SWidget> ParentWidget, const FText& Title, const TArray<FTemplateMapInfo>& Templates, FString& OutTemplateMapPackageName )
+bool FNewLevelDialogModule::CreateAndShowTemplateDialog( const TSharedPtr<const SWidget> ParentWidget, const FText& Title, const TArray<FTemplateMapInfo>& Templates, FString& OutTemplateMapPackageName, FNewLevelOptions* InOutNewLevelOptions )
 {
 	
 	TSharedPtr<SWindow> NewLevelWindow =
@@ -329,13 +369,20 @@ bool FNewLevelDialogModule::CreateAndShowTemplateDialog( const TSharedPtr<const
 	TSharedRef<SNewLevelDialog> NewLevelDialog =
 		SNew(SNewLevelDialog)
 		.ParentWindow(NewLevelWindow)
-		.Templates(Templates);
+		.Templates(Templates)
+		.bExternalActors(InOutNewLevelOptions ? InOutNewLevelOptions->bExternalActors : false)
+		.bPartitionedWorld(InOutNewLevelOptions ? InOutNewLevelOptions->bPartitionedWorld : false);
 
 	NewLevelWindow->SetContent(NewLevelDialog);
 
 	FSlateApplication::Get().AddModalWindow(NewLevelWindow.ToSharedRef(), ParentWidget);
 
 	OutTemplateMapPackageName = NewLevelDialog->GetChosenTemplate();
+	if (InOutNewLevelOptions)
+	{
+		InOutNewLevelOptions->bExternalActors = NewLevelDialog->GetExternalActors();
+		InOutNewLevelOptions->bPartitionedWorld = NewLevelDialog->GetPartitionedWorld();
+	}
 	return NewLevelDialog->IsTemplateChosen();
 }
 
diff --git a/Engine/Source/Editor/NewLevelDialog/Public/NewLevelDialogModule.h b/Engine/Source/Editor/NewLevelDialog/Public/NewLevelDialogModule.h
--- a/Engine/Source/Editor/NewLevelDialog/Public/NewLevelDialogModule.h
+++ b/Engine/Source/Editor/NewLevelDialog/Public/NewLevelDialogModule.h
@@ -13,6 +13,16 @@
 class FNewLevelDialogModule : public IModuleInterface
 {
 public:
+	struct FNewLevelOptions
+	{
+		bool bExternalActors;
+		bool bPartitionedWorld;
+
+		FNewLevelOptions()
+			: bExternalActors(false)
+			, bPartitionedWorld(false)
+		{}
+	};
 
 	/**
 	 * Called right after the plugin DLL has been loaded and the plugin object has been created
@@ -29,9 +39,10 @@ public:
 	 * 
 	 * @param	ParentWidget - The parent widget for the modal window showing the dialog
 	 * @param	OutTemplateName	- (out) The package name of the template map selected by the user. Empty if blank map selected.
+	 * @param	InOutNewLevelOptions - (in/out) The new level options: input is the options we can support, out is the selected options.
 	 * @return	true if the user selected a valid item, false if the user canceled
 	 */
-	virtual bool CreateAndShowNewLevelDialog( const TSharedPtr<const SWidget> ParentWidget, FString& OutTemplateMapPackageName );
+	virtual bool CreateAndShowNewLevelDialog( const TSharedPtr<const SWidget> ParentWidget, FString& OutTemplateMapPackageName, FNewLevelOptions* InOutNewLevelOptions = nullptr );
 
 	/**
 	 * Creates and show a window with an SNewLevelDialog
@@ -40,9 +51,10 @@ public:
 	 * @param   Title - The dialog's title
 	 * @param   Templates - The list of template to be shown in the dialog
 	 * @param	OutTemplateName	- (out) The package name of the template map selected by the user. Empty if blank map selected.
+	 * @param	InOutNewLevelOptions - (in/out) The new level options: input is the options we can support, out is the selected options.
 	 * @return	true if the user selected a valid item, false if the user canceled
 	 */
-	virtual bool CreateAndShowTemplateDialog(const TSharedPtr<const SWidget> ParentWidget, const FText& Title, const TArray<FTemplateMapInfo>& Templates, FString& OutTemplateMapPackageName);
+	virtual bool CreateAndShowTemplateDialog(const TSharedPtr<const SWidget> ParentWidget, const FText& Title, const TArray<FTemplateMapInfo>& Templates, FString& OutTemplateMapPackageName, FNewLevelOptions* InOutNewLevelOptions = nullptr );
 
 	/** New Level Dialog app identifier string */
 	static const FName NewLevelDialogAppIdentifier;
diff --git a/Engine/Source/Runtime/Engine/Classes/Engine/World.h b/Engine/Source/Runtime/Engine/Classes/Engine/World.h
--- a/Engine/Source/Runtime/Engine/Classes/Engine/World.h
+++ b/Engine/Source/Runtime/Engine/Classes/Engine/World.h
@@ -3655,6 +3655,15 @@ public:
 		return (OwningGameInstance ? OwningGameInstance->GetLatentActionManager() : LatentActionManager);
 	}
 
+	/**
+	 * Reinitialize all subsystems
+	 */
+	void ReinitializeSubSystems()
+	{
+		SubsystemCollection.Deinitialize();
+		SubsystemCollection.Initialize(this);
+	}
+
 	/**
 	 * Get a Subsystem of specified type
 	 */
