Formalized the process for creating higher-level element interfaces that can be customized by specific asset editors

This takes the framework used for UTypedElementSelectionSet and moves to it a utility base class that other higher-level interfaces can also make use of.

This also changes the interface implementations for these systems to be non-UObject types, as the script exposure is handled by the owner object, and this can solve some cyclic parsing issues with UHT (eg, between UnrealEd and LevelEditor).

#rnx

diff --git a/Engine/Source/Editor/EditorFramework/Private/Elements/Framework/TypedElementAssetEditorCustomization.cpp b/Engine/Source/Editor/EditorFramework/Private/Elements/Framework/TypedElementAssetEditorCustomization.cpp
new file mode 100644
--- /dev/null
+++ b/Engine/Source/Editor/EditorFramework/Private/Elements/Framework/TypedElementAssetEditorCustomization.cpp
@@ -0,0 +1,16 @@
+// Copyright Epic Games, Inc. All Rights Reserved.
+
+#include "Elements/Framework/TypedElementAssetEditorCustomization.h"
+#include "Elements/Framework/TypedElementRegistry.h"
+
+FTypedHandleTypeId FTypedElementAssetEditorCustomizationRegistryBase::GetElementTypeIdFromName(const FName InElementTypeName) const
+{
+	return UTypedElementRegistry::GetInstance()->GetRegisteredElementTypeId(InElementTypeName);
+}
+
+FTypedHandleTypeId FTypedElementAssetEditorCustomizationRegistryBase::GetElementTypeIdFromNameChecked(const FName InElementTypeName) const
+{
+	const FTypedHandleTypeId ElementTypeId = GetElementTypeIdFromName(InElementTypeName);
+	checkf(ElementTypeId > 0, TEXT("Element type '%s' has not been registered!"), *InElementTypeName.ToString());
+	return ElementTypeId;
+}
diff --git a/Engine/Source/Editor/EditorFramework/Private/Elements/Framework/TypedElementSelectionSet.cpp b/Engine/Source/Editor/EditorFramework/Private/Elements/Framework/TypedElementSelectionSet.cpp
--- a/Engine/Source/Editor/EditorFramework/Private/Elements/Framework/TypedElementSelectionSet.cpp
+++ b/Engine/Source/Editor/EditorFramework/Private/Elements/Framework/TypedElementSelectionSet.cpp
@@ -209,21 +209,11 @@ FTypedElementHandle UTypedElementSelectionSet::GetSelectionElement(const FTypedE
 	return SelectionSetElement ? SelectionSetElement.GetSelectionElement(InSelectionMethod) : FTypedElementHandle();
 }
 
-void UTypedElementSelectionSet::RegisterAssetEditorSelectionProxy(const FName InElementTypeName, UTypedElementAssetEditorSelectionProxy* InAssetEditorSelectionProxy)
-{
-	const FTypedHandleTypeId ElementTypeId = UTypedElementRegistry::GetInstance()->GetRegisteredElementTypeId(InElementTypeName);
-	checkf(ElementTypeId > 0, TEXT("Element type '%s' has not been registered!"), *InElementTypeName.ToString());
-	RegisteredAssetEditorSelectionProxies[ElementTypeId - 1] = InAssetEditorSelectionProxy;
-}
-
 FTypedElementSelectionSetElement UTypedElementSelectionSet::ResolveSelectionSetElement(const FTypedElementHandle& InElementHandle) const
 {
-	if (InElementHandle)
-	{
-		UTypedElementAssetEditorSelectionProxy* AssetEditorSelectionProxy = RegisteredAssetEditorSelectionProxies[InElementHandle.GetId().GetTypeId() - 1];
-		return FTypedElementSelectionSetElement(ElementList->GetElement<UTypedElementSelectionInterface>(InElementHandle), ElementList, AssetEditorSelectionProxy ? AssetEditorSelectionProxy : GetMutableDefault<UTypedElementAssetEditorSelectionProxy>());
-	}
-	return FTypedElementSelectionSetElement();
+	return InElementHandle
+		? FTypedElementSelectionSetElement(ElementList->GetElement<UTypedElementSelectionInterface>(InElementHandle), ElementList, GetAssetEditorCustomizationByTypeId(InElementHandle.GetId().GetTypeId()))
+		: FTypedElementSelectionSetElement();
 }
 
 void UTypedElementSelectionSet::OnElementListPreChange(const UTypedElementList* InElementList)
diff --git a/Engine/Source/Editor/EditorFramework/Public/Elements/Framework/TypedElementAssetEditorCustomization.h b/Engine/Source/Editor/EditorFramework/Public/Elements/Framework/TypedElementAssetEditorCustomization.h
new file mode 100644
--- /dev/null
+++ b/Engine/Source/Editor/EditorFramework/Public/Elements/Framework/TypedElementAssetEditorCustomization.h
@@ -0,0 +1,122 @@
+// Copyright Epic Games, Inc. All Rights Reserved.
+
+#pragma once
+
+#include "CoreMinimal.h"
+#include "Templates/UniquePtr.h"
+#include "Elements/Framework/TypedElementLimits.h"
+
+/**
+ * Non-templated base class for the asset editor customization registry.
+ */
+class EDITORFRAMEWORK_API FTypedElementAssetEditorCustomizationRegistryBase
+{
+public:
+	virtual ~FTypedElementAssetEditorCustomizationRegistryBase() = default;
+
+protected:
+	/**
+	 * Given an element name, attempt to get its registered type ID from the global registry.
+	 * @return The registered type ID, or 0 if the element name is not registered.
+	 */
+	FTypedHandleTypeId GetElementTypeIdFromName(const FName InElementTypeName) const;
+
+	/**
+	 * Given an element name, attempt to get its registered type ID from the global registry.
+	 * @return The registered type ID, or asserts if the element name is not registered.
+	 */
+	FTypedHandleTypeId GetElementTypeIdFromNameChecked(const FName InElementTypeName) const;
+};
+
+/**
+ * Utility to register and retrieve asset editor customizations for a given type.
+ */
+template <typename CustomizationBaseType, typename DefaultCustomizationType = CustomizationBaseType>
+class TTypedElementAssetEditorCustomizationRegistry : public FTypedElementAssetEditorCustomizationRegistryBase
+{
+public:
+	template <typename... TDefaultArgs>
+	TTypedElementAssetEditorCustomizationRegistry(TDefaultArgs&&... DefaultArgs)
+		: DefaultAssetEditorCustomization(MakeUnique<DefaultCustomizationType>(Forward<TDefaultArgs>(DefaultArgs)...))
+	{
+	}
+
+	virtual ~TTypedElementAssetEditorCustomizationRegistry() = default;
+
+	TTypedElementAssetEditorCustomizationRegistry(const TTypedElementAssetEditorCustomizationRegistry&) = delete;
+	TTypedElementAssetEditorCustomizationRegistry& operator=(const TTypedElementAssetEditorCustomizationRegistry&) = delete;
+
+	TTypedElementAssetEditorCustomizationRegistry(TTypedElementAssetEditorCustomizationRegistry&&) = delete;
+	TTypedElementAssetEditorCustomizationRegistry& operator=(TTypedElementAssetEditorCustomizationRegistry&&) = delete;
+
+	/**
+	 * Set the default asset editor customization instance.
+	 */
+	void SetDefaultAssetEditorCustomization(TUniquePtr<CustomizationBaseType>&& InAssetEditorCustomization)
+	{
+		checkf(InAssetEditorCustomization, TEXT("Default asset editor customization cannot be null!"));
+		DefaultAssetEditorCustomization = MoveTemp(InAssetEditorCustomization);
+	}
+
+	/**
+	 * Register an asset editor customization for the given element type.
+	 */
+	void RegisterAssetEditorCustomizationByTypeName(const FName InElementTypeName, TUniquePtr<CustomizationBaseType>&& InAssetEditorCustomization)
+	{
+		RegisterAssetEditorCustomizationByTypeId(GetElementTypeIdFromNameChecked(InElementTypeName), MoveTemp(InAssetEditorCustomization));
+	}
+
+	/**
+	 * Register an asset editor customization for the given element type.
+	 */
+	void RegisterAssetEditorCustomizationByTypeId(const FTypedHandleTypeId InElementTypeId, TUniquePtr<CustomizationBaseType>&& InAssetEditorCustomization)
+	{
+		RegisteredAssetEditorCustomizations[InElementTypeId - 1] = MoveTemp(InAssetEditorCustomization);
+	}
+
+	/**
+	 * Unregister an asset editor customization for the given element type.
+	 */
+	void UnregisterAssetEditorCustomizationByTypeName(const FName InElementTypeName)
+	{
+		UnregisterAssetEditorCustomizationByTypeId(GetElementTypeIdFromNameChecked(InElementTypeName));
+	}
+
+	/**
+	 * Unregister an asset editor customization for the given element type.
+	 */
+	void UnregisterAssetEditorCustomizationByTypeId(const FTypedHandleTypeId InElementTypeId)
+	{
+		RegisteredAssetEditorCustomizations[InElementTypeId - 1].Reset();
+	}
+
+	/**
+	 * Get the asset editor customization for the given element type.
+	 * @note If bAllowFallback is true, then this will return the default asset editor customization if no override is present, otherwise it will return null.
+	 */
+	CustomizationBaseType* GetAssetEditorCustomizationByTypeName(const FName InElementTypeName, const bool bAllowFallback = true) const
+	{
+		return GetAssetEditorCustomizationByTypeId(GetElementTypeIdFromNameChecked(InElementTypeName), bAllowFallback);
+	}
+
+	/**
+	 * Get the asset editor customization for the given element type.
+	 * @note If bAllowFallback is true, then this will return the default asset editor customization if no override is present, otherwise it will return null.
+	 */
+	CustomizationBaseType* GetAssetEditorCustomizationByTypeId(const FTypedHandleTypeId InElementTypeId, const bool bAllowFallback = true) const
+	{
+		CustomizationBaseType* AssetEditorCustomization = RegisteredAssetEditorCustomizations[InElementTypeId - 1].Get();
+		return AssetEditorCustomization
+			? AssetEditorCustomization
+			: bAllowFallback 
+				? DefaultAssetEditorCustomization.Get()
+				: nullptr;
+	}
+
+private:
+	/** Default asset editor customization, used if no type-specific override is present. */
+	TUniquePtr<CustomizationBaseType> DefaultAssetEditorCustomization;
+
+	/** Array of registered asset editor customizations, indexed by ElementTypeId-1. */
+	TUniquePtr<CustomizationBaseType> RegisteredAssetEditorCustomizations[TypedHandleMaxTypeId - 1];
+};
diff --git a/Engine/Source/Editor/EditorFramework/Public/Elements/Framework/TypedElementSelectionSet.h b/Engine/Source/Editor/EditorFramework/Public/Elements/Framework/TypedElementSelectionSet.h
--- a/Engine/Source/Editor/EditorFramework/Public/Elements/Framework/TypedElementSelectionSet.h
+++ b/Engine/Source/Editor/EditorFramework/Public/Elements/Framework/TypedElementSelectionSet.h
@@ -3,19 +3,19 @@
 #pragma once
 
 #include "Elements/Framework/TypedElementList.h"
+#include "Elements/Framework/TypedElementAssetEditorCustomization.h"
 #include "Elements/Interfaces/TypedElementSelectionInterface.h"
 #include "TypedElementSelectionSet.generated.h"
 
 /**
- * Proxy type used to allow asset editors (such as the level editor) to override the base behavior of element selection,
+ * Customization type used to allow asset editors (such as the level editor) to override the base behavior of element selection,
  * by injecting extra pre/post selection logic around the call into the selection interface for an element type.
  */
-UCLASS()
-class EDITORFRAMEWORK_API UTypedElementAssetEditorSelectionProxy : public UObject
+class EDITORFRAMEWORK_API FTypedElementAssetEditorSelectionCustomization
 {
-	GENERATED_BODY()
-
 public:
+	virtual ~FTypedElementAssetEditorSelectionCustomization() = default;
+
 	//~ See UTypedElementSelectionInterface for API docs
 	virtual bool IsElementSelected(const TTypedElement<UTypedElementSelectionInterface>& InElementSelectionHandle, const UTypedElementList* InSelectionSet, const FTypedElementIsSelectedOptions& InSelectionOptions) { return InElementSelectionHandle.IsElementSelected(InSelectionSet, InSelectionOptions); }
 	virtual bool CanSelectElement(const TTypedElement<UTypedElementSelectionInterface>& InElementSelectionHandle, const FTypedElementSelectionOptions& InSelectionOptions) { return InElementSelectionHandle.CanSelectElement(InSelectionOptions); }
@@ -27,17 +27,17 @@ public:
 };
 
 /**
- * Utility to hold a typed element handle and its associated selection interface and asset editor selection proxy.
+ * Utility to hold a typed element handle and its associated selection interface and asset editor selection customization.
  */
 struct EDITORFRAMEWORK_API FTypedElementSelectionSetElement
 {
 public:
 	FTypedElementSelectionSetElement() = default;
 
-	FTypedElementSelectionSetElement(TTypedElement<UTypedElementSelectionInterface> InElementSelectionHandle, UTypedElementList* InElementList, UTypedElementAssetEditorSelectionProxy* InAssetEditorSelectionProxy)
+	FTypedElementSelectionSetElement(TTypedElement<UTypedElementSelectionInterface> InElementSelectionHandle, UTypedElementList* InElementList, FTypedElementAssetEditorSelectionCustomization* InAssetEditorSelectionCustomization)
 		: ElementSelectionHandle(MoveTemp(InElementSelectionHandle))
 		, ElementList(InElementList)
-		, AssetEditorSelectionProxy(InAssetEditorSelectionProxy)
+		, AssetEditorSelectionCustomization(InAssetEditorSelectionCustomization)
 	{
 	}
 
@@ -56,22 +56,22 @@ public:
 	{
 		return ElementSelectionHandle.IsSet()
 			&& ElementList
-			&& AssetEditorSelectionProxy;
+			&& AssetEditorSelectionCustomization;
 	}
 
 	//~ See UTypedElementSelectionInterface for API docs
-	bool IsElementSelected(const FTypedElementIsSelectedOptions& InSelectionOptions) const { return AssetEditorSelectionProxy->IsElementSelected(ElementSelectionHandle, ElementList, InSelectionOptions); }
-	bool CanSelectElement(const FTypedElementSelectionOptions& InSelectionOptions) const { return AssetEditorSelectionProxy->CanSelectElement(ElementSelectionHandle, InSelectionOptions); }
-	bool CanDeselectElement(const FTypedElementSelectionOptions& InSelectionOptions) const { return AssetEditorSelectionProxy->CanDeselectElement(ElementSelectionHandle, InSelectionOptions); }
-	bool SelectElement(const FTypedElementSelectionOptions& InSelectionOptions) const { return AssetEditorSelectionProxy->SelectElement(ElementSelectionHandle, ElementList, InSelectionOptions); }
-	bool DeselectElement(const FTypedElementSelectionOptions& InSelectionOptions) const { return AssetEditorSelectionProxy->DeselectElement(ElementSelectionHandle, ElementList, InSelectionOptions); }
-	bool AllowSelectionModifiers() const { return AssetEditorSelectionProxy->AllowSelectionModifiers(ElementSelectionHandle, ElementList); }
-	FTypedElementHandle GetSelectionElement(const ETypedElementSelectionMethod InSelectionMethod) const { return AssetEditorSelectionProxy->GetSelectionElement(ElementSelectionHandle, ElementList, InSelectionMethod); }
+	bool IsElementSelected(const FTypedElementIsSelectedOptions& InSelectionOptions) const { return AssetEditorSelectionCustomization->IsElementSelected(ElementSelectionHandle, ElementList, InSelectionOptions); }
+	bool CanSelectElement(const FTypedElementSelectionOptions& InSelectionOptions) const { return AssetEditorSelectionCustomization->CanSelectElement(ElementSelectionHandle, InSelectionOptions); }
+	bool CanDeselectElement(const FTypedElementSelectionOptions& InSelectionOptions) const { return AssetEditorSelectionCustomization->CanDeselectElement(ElementSelectionHandle, InSelectionOptions); }
+	bool SelectElement(const FTypedElementSelectionOptions& InSelectionOptions) const { return AssetEditorSelectionCustomization->SelectElement(ElementSelectionHandle, ElementList, InSelectionOptions); }
+	bool DeselectElement(const FTypedElementSelectionOptions& InSelectionOptions) const { return AssetEditorSelectionCustomization->DeselectElement(ElementSelectionHandle, ElementList, InSelectionOptions); }
+	bool AllowSelectionModifiers() const { return AssetEditorSelectionCustomization->AllowSelectionModifiers(ElementSelectionHandle, ElementList); }
+	FTypedElementHandle GetSelectionElement(const ETypedElementSelectionMethod InSelectionMethod) const { return AssetEditorSelectionCustomization->GetSelectionElement(ElementSelectionHandle, ElementList, InSelectionMethod); }
 
 private:
 	TTypedElement<UTypedElementSelectionInterface> ElementSelectionHandle;
 	UTypedElementList* ElementList = nullptr;
-	UTypedElementAssetEditorSelectionProxy* AssetEditorSelectionProxy = nullptr;
+	FTypedElementAssetEditorSelectionCustomization* AssetEditorSelectionCustomization = nullptr;
 };
 
 /**
@@ -79,7 +79,7 @@ private:
  * interfaces, as well as providing some utilities for batching operations.
  */
 UCLASS(Transient)
-class EDITORFRAMEWORK_API UTypedElementSelectionSet : public UObject
+class EDITORFRAMEWORK_API UTypedElementSelectionSet : public UObject, public TTypedElementAssetEditorCustomizationRegistry<FTypedElementAssetEditorSelectionCustomization>
 {
 	GENERATED_BODY()
 
@@ -219,14 +219,9 @@ public:
 		return ElementList;
 	}
 
-	/**
-	 * Register an asset editor selection proxy for the given named element type.
-	 */
-	void RegisterAssetEditorSelectionProxy(const FName InElementTypeName, UTypedElementAssetEditorSelectionProxy* InAssetEditorSelectionProxy);
-
 private:
 	/**
-	 * Attempt to resolve the selection interface and asset editor selection proxy for the given element, if any.
+	 * Attempt to resolve the selection interface and asset editor selection customization for the given element, if any.
 	 */
 	FTypedElementSelectionSetElement ResolveSelectionSetElement(const FTypedElementHandle& InElementHandle) const;
 
@@ -244,10 +239,6 @@ private:
 	UPROPERTY()
 	UTypedElementList* ElementList = nullptr;
 
-	/** Array of registered asset editor selection proxies, indexed by ElementTypeId-1. */
-	UPROPERTY()
-	UTypedElementAssetEditorSelectionProxy* RegisteredAssetEditorSelectionProxies[TypedHandleMaxTypeId - 1];
-
 	/** Delegate that is invoked whenever the underlying element list is potentially about to change. */
 	FOnPreChange OnPreChangeDelegate;
 
diff --git a/Engine/Source/Editor/LevelEditor/Private/Elements/Actor/ActorElementLevelEditorSelectionCustomization.cpp b/Engine/Source/Editor/LevelEditor/Private/Elements/Actor/ActorElementLevelEditorSelectionCustomization.cpp
new file mode 100644
--- /dev/null
+++ b/Engine/Source/Editor/LevelEditor/Private/Elements/Actor/ActorElementLevelEditorSelectionCustomization.cpp
@@ -0,0 +1,308 @@
+// Copyright Epic Games, Inc. All Rights Reserved.
+
+#include "Elements/Actor/ActorElementLevelEditorSelectionCustomization.h"
+#include "Elements/Actor/ActorElementData.h"
+#include "GameFramework/Actor.h"
+#include "Components/ActorComponent.h"
+
+#include "Elements/Framework/TypedElementList.h"
+#include "Elements/Framework/EngineElementsLibrary.h"
+
+#include "Elements/Actor/ActorElementSelectionInterface.h"
+#include "Elements/Component/ComponentElementSelectionInterface.h"
+
+#include "Editor.h"
+#include "LevelUtils.h"
+#include "UnrealEdGlobals.h"
+#include "EditorModeManager.h"
+#include "Editor/GroupActor.h"
+#include "ActorGroupingUtils.h"
+#include "Editor/UnrealEdEngine.h"
+#include "Kismet2/ComponentEditorUtils.h"
+
+DEFINE_LOG_CATEGORY_STATIC(LogActorLevelEditorSelection, Log, All);
+
+bool FActorElementLevelEditorSelectionCustomization::CanSelectElement(const TTypedElement<UTypedElementSelectionInterface>& InElementSelectionHandle, const FTypedElementSelectionOptions& InSelectionOptions)
+{
+	return InElementSelectionHandle.IsDataOfType<FActorElementData>() && CanSelectActorElement(InElementSelectionHandle, InSelectionOptions);
+}
+
+bool FActorElementLevelEditorSelectionCustomization::CanDeselectElement(const TTypedElement<UTypedElementSelectionInterface>& InElementSelectionHandle, const FTypedElementSelectionOptions& InSelectionOptions)
+{
+	return InElementSelectionHandle.IsDataOfType<FActorElementData>() && CanDeselectActorElement(InElementSelectionHandle, InSelectionOptions);
+}
+
+bool FActorElementLevelEditorSelectionCustomization::SelectElement(const TTypedElement<UTypedElementSelectionInterface>& InElementSelectionHandle, UTypedElementList* InSelectionSet, const FTypedElementSelectionOptions& InSelectionOptions)
+{
+	return InElementSelectionHandle.IsDataOfType<FActorElementData>() && SelectActorElement(InElementSelectionHandle, InSelectionSet, InSelectionOptions);
+}
+
+bool FActorElementLevelEditorSelectionCustomization::DeselectElement(const TTypedElement<UTypedElementSelectionInterface>& InElementSelectionHandle, UTypedElementList* InSelectionSet, const FTypedElementSelectionOptions& InSelectionOptions)
+{
+	return InElementSelectionHandle.IsDataOfType<FActorElementData>() && DeselectActorElement(InElementSelectionHandle, InSelectionSet, InSelectionOptions);
+}
+
+bool FActorElementLevelEditorSelectionCustomization::AllowSelectionModifiers(const TTypedElement<UTypedElementSelectionInterface>& InElementSelectionHandle, const UTypedElementList* InSelectionSet)
+{
+	// Ctrl or Shift clicking an actor is the same as regular clicking when components are selected
+	return InElementSelectionHandle.IsDataOfType<FActorElementData>() && !UComponentElementSelectionInterface::HasSelectedComponents(InSelectionSet);
+}
+
+FTypedElementHandle FActorElementLevelEditorSelectionCustomization::GetSelectionElement(const TTypedElement<UTypedElementSelectionInterface>& InElementSelectionHandle, const UTypedElementList* InCurrentSelection, const ETypedElementSelectionMethod InSelectionMethod)
+{
+	if (const FActorElementData* ActorData = InElementSelectionHandle.GetData<FActorElementData>())
+	{
+		AActor* ConsideredActor = ActorData->Actor;
+		while (ConsideredActor->IsChildActor())
+		{
+			ConsideredActor = ConsideredActor->GetParentActor();
+		}
+		return UEngineElementsLibrary::AcquireEditorActorElementHandle(ConsideredActor);
+	}
+	return InElementSelectionHandle;
+}
+
+bool FActorElementLevelEditorSelectionCustomization::CanSelectActorElement(const TTypedElement<UTypedElementSelectionInterface>& InActorSelectionHandle, const FTypedElementSelectionOptions& InSelectionOptions)
+{
+	AActor* Actor = InActorSelectionHandle.GetDataChecked<FActorElementData>().Actor;
+
+	// Bail if global selection is locked, or this actor cannot be edited or selected
+	if (GEdSelectionLock || !Actor->IsEditable() || !Actor->IsSelectable())
+	{
+		return false;
+	}
+
+	// Bail if the actor is hidden, and we're not allowed to select hidden elements
+	if (!InSelectionOptions.AllowHidden() && (Actor->IsHiddenEd() || !FLevelUtils::IsLevelVisible(Actor->GetLevel())))
+	{
+		return false;
+	}
+
+	// Ensure that neither the level nor the actor is being destroyed or is unreachable
+	const EObjectFlags InvalidSelectableFlags = RF_BeginDestroyed;
+	if (Actor->GetLevel()->HasAnyFlags(InvalidSelectableFlags) || Actor->GetLevel()->IsPendingKillOrUnreachable())
+	{
+		UE_LOG(LogActorLevelEditorSelection, Warning, TEXT("SelectActor: %s (%s)"), TEXT("The requested operation could not be completed because the level has invalid flags."), *Actor->GetActorLabel());
+		return false;
+	}
+	if (Actor->HasAnyFlags(InvalidSelectableFlags) || Actor->IsPendingKillOrUnreachable())
+	{
+		UE_LOG(LogActorLevelEditorSelection, Warning, TEXT("SelectActor: %s (%s)"), TEXT("The requested operation could not be completed because the actor has invalid flags."), *Actor->GetActorLabel());
+		return false;
+	}
+
+	if (!Actor->IsTemplate() && FLevelUtils::IsLevelLocked(Actor->GetLevel()))
+	{
+		UE_CLOG(InSelectionOptions.WarnIfLocked(), LogActorLevelEditorSelection, Warning, TEXT("SelectActor: %s (%s)"), TEXT("The requested operation could not be completed because the level is locked."), *Actor->GetActorLabel());
+		return false;
+	}
+
+	// If grouping operations are not currently allowed, don't select groups
+	AGroupActor* SelectedGroupActor = Cast<AGroupActor>(Actor);
+	if (SelectedGroupActor && (!UActorGroupingUtils::IsGroupingActive() || !InSelectionOptions.AllowGroups()))
+	{
+		return false;
+	}
+
+	// Allow active modes to determine whether the selection is allowed
+	return GLevelEditorModeTools().IsSelectionAllowed(Actor, /*bInSelected*/true);
+}
+
+bool FActorElementLevelEditorSelectionCustomization::CanDeselectActorElement(const TTypedElement<UTypedElementSelectionInterface>& InActorSelectionHandle, const FTypedElementSelectionOptions& InSelectionOptions)
+{
+	AActor* Actor = InActorSelectionHandle.GetDataChecked<FActorElementData>().Actor;
+
+	// Bail if global selection is locked
+	if (GEdSelectionLock)
+	{
+		return false;
+	}
+
+	// Allow active modes to determine whether the deselection is allowed
+	return GLevelEditorModeTools().IsSelectionAllowed(Actor, /*bInSelected*/false);
+}
+
+bool FActorElementLevelEditorSelectionCustomization::SelectActorElement(const TTypedElement<UTypedElementSelectionInterface>& InActorSelectionHandle, UTypedElementList* InSelectionSet, const FTypedElementSelectionOptions& InSelectionOptions)
+{
+	AActor* Actor = InActorSelectionHandle.GetDataChecked<FActorElementData>().Actor;
+
+	// Allow active modes to potentially handle the selection
+	// TODO: Should this pass through the selection set?
+	if (GLevelEditorModeTools().IsSelectionHandled(Actor, /*bInSelected*/true))
+	{
+		return true;
+	}
+
+	// If trying to select an actor, use this actors root selection actor instead (if it has one)
+	if (AActor* RootSelection = Actor->GetRootSelectionParent())
+	{
+		Actor = RootSelection;
+	}
+
+	bool bSelectionChanged = false;
+
+	if (UActorGroupingUtils::IsGroupingActive() && InSelectionOptions.AllowGroups())
+	{
+		// If this actor is a group, do a group select
+		if (AGroupActor* SelectedGroupActor = Cast<AGroupActor>(Actor))
+		{
+			bSelectionChanged |= SelectActorGroup(SelectedGroupActor, InSelectionSet, InSelectionOptions, /*bForce*/true);
+		}
+		// Select/Deselect this actor's entire group, starting from the top locked group
+		else if (AGroupActor* ActorLockedRootGroup = AGroupActor::GetRootForActor(Actor, true))
+		{
+			bSelectionChanged |= SelectActorGroup(ActorLockedRootGroup, InSelectionSet, InSelectionOptions, /*bForce*/false);
+		}
+	}
+
+	// Select the desired actor
+	{
+		TTypedElement<UTypedElementSelectionInterface> ActorSelectionHandle = InSelectionSet->GetElement<UTypedElementSelectionInterface>(UEngineElementsLibrary::AcquireEditorActorElementHandle(Actor));
+		if (!ActorSelectionHandle.SelectElement(InSelectionSet, InSelectionOptions))
+		{
+			return bSelectionChanged;
+		}
+	}
+
+	UE_LOG(LogActorLevelEditorSelection, Verbose, TEXT("Selected Actor: %s"), *Actor->GetClass()->GetName());
+
+	// Update the annotation state
+	GSelectedActorAnnotation.Set(Actor);
+	
+	// Bind the override delegates for the components on the selected actor
+	for (UActorComponent* Component : Actor->GetComponents())
+	{
+		if (USceneComponent* SceneComponent = Cast<USceneComponent>(Component))
+		{
+			FComponentEditorUtils::BindComponentSelectionOverride(SceneComponent, /*bBind*/true);
+		}
+	}
+		
+	// Flush some cached data
+	GUnrealEd->PostActorSelectionChanged();
+
+	// A fast path to mark selection rather than reconnecting ALL components for ALL actors that have changed state
+	Actor->PushSelectionToProxies();
+
+	return true;
+}
+
+bool FActorElementLevelEditorSelectionCustomization::DeselectActorElement(const TTypedElement<UTypedElementSelectionInterface>& InActorSelectionHandle, UTypedElementList* InSelectionSet, const FTypedElementSelectionOptions& InSelectionOptions)
+{
+	AActor* Actor = InActorSelectionHandle.GetDataChecked<FActorElementData>().Actor;
+
+	// Allow active modes to potentially handle the deselection
+	// TODO: Should this pass through the selection set?
+	if (GLevelEditorModeTools().IsSelectionHandled(Actor, /*bInSelected*/false))
+	{
+		return true;
+	}
+
+	bool bSelectionChanged = false;
+
+	if (UActorGroupingUtils::IsGroupingActive() && InSelectionOptions.AllowGroups())
+	{
+		// If this actor is a group, do a group select
+		if (AGroupActor* SelectedGroupActor = Cast<AGroupActor>(Actor))
+		{
+			bSelectionChanged |= DeselectActorGroup(SelectedGroupActor, InSelectionSet, InSelectionOptions, /*bForce*/true);
+		}
+		// Select/Deselect this actor's entire group, starting from the top locked group
+		else if (AGroupActor* ActorLockedRootGroup = AGroupActor::GetRootForActor(Actor, true))
+		{
+			bSelectionChanged |= DeselectActorGroup(ActorLockedRootGroup, InSelectionSet, InSelectionOptions, /*bForce*/false);
+		}
+	}
+
+	// Deselect the desired actor
+	{
+		TTypedElement<UTypedElementSelectionInterface> ActorSelectionHandle = InSelectionSet->GetElement<UTypedElementSelectionInterface>(UEngineElementsLibrary::AcquireEditorActorElementHandle(Actor));
+		if (!ActorSelectionHandle.DeselectElement(InSelectionSet, InSelectionOptions))
+		{
+			return bSelectionChanged;
+		}
+	}
+	
+	UE_LOG(LogActorLevelEditorSelection, Verbose, TEXT("Deselected Actor: %s"), *Actor->GetClass()->GetName());
+	
+	// Update the annotation state
+	GSelectedActorAnnotation.Clear(Actor);
+	
+	// Deselect and unbind the override delegates for the components on the selected actor
+	{
+		FTypedElementListLegacySyncScopedBatch LegacySyncBatch(InSelectionSet, InSelectionOptions.AllowLegacyNotifications());
+
+		for (UActorComponent* Component : Actor->GetComponents())
+		{
+			TTypedElement<UTypedElementSelectionInterface> ComponentSelectionHandle = InSelectionSet->GetElement<UTypedElementSelectionInterface>(UEngineElementsLibrary::AcquireEditorComponentElementHandle(Component));
+			ComponentSelectionHandle.DeselectElement(InSelectionSet, InSelectionOptions);
+
+			if (USceneComponent* SceneComponent = Cast<USceneComponent>(Component))
+			{
+				FComponentEditorUtils::BindComponentSelectionOverride(SceneComponent, /*bBind*/false);
+			}
+		}
+	}
+
+	// Flush some cached data
+	GUnrealEd->PostActorSelectionChanged();
+
+	// A fast path to mark selection rather than reconnecting ALL components for ALL actors that have changed state
+	Actor->PushSelectionToProxies();
+
+	return true;
+}
+
+bool FActorElementLevelEditorSelectionCustomization::SelectActorGroup(AGroupActor* InGroupActor, UTypedElementList* InSelectionSet, const FTypedElementSelectionOptions& InSelectionOptions, const bool bForce)
+{
+	bool bSelectionChanged = false;
+
+	// Select all actors within the group (if locked or forced)
+	if (bForce || InGroupActor->IsLocked())
+	{
+		FTypedElementListLegacySyncScopedBatch LegacySyncBatch(InSelectionSet, InSelectionOptions.AllowLegacyNotifications());
+
+		const FTypedElementSelectionOptions GroupSelectionOptions = FTypedElementSelectionOptions(InSelectionOptions)
+			.SetAllowGroups(false);
+
+		TArray<AActor*> GroupActors;
+		InGroupActor->GetGroupActors(GroupActors);
+		for (AActor* Actor : GroupActors)
+		{
+			TTypedElement<UTypedElementSelectionInterface> ActorSelectionHandle = InSelectionSet->GetElement<UTypedElementSelectionInterface>(UEngineElementsLibrary::AcquireEditorActorElementHandle(Actor));
+			if (CanSelectActorElement(ActorSelectionHandle, GroupSelectionOptions))
+			{
+				bSelectionChanged |= SelectActorElement(ActorSelectionHandle, InSelectionSet, GroupSelectionOptions);
+			}
+		}
+	}
+
+	return bSelectionChanged;
+}
+
+bool FActorElementLevelEditorSelectionCustomization::DeselectActorGroup(AGroupActor* InGroupActor, UTypedElementList* InSelectionSet, const FTypedElementSelectionOptions& InSelectionOptions, const bool bForce)
+{
+	bool bSelectionChanged = false;
+
+	// Deselect all actors within the group (if locked or forced)
+	if (bForce || InGroupActor->IsLocked())
+	{
+		FTypedElementListLegacySyncScopedBatch LegacySyncBatch(InSelectionSet, InSelectionOptions.AllowLegacyNotifications());
+
+		const FTypedElementSelectionOptions GroupSelectionOptions = FTypedElementSelectionOptions(InSelectionOptions)
+			.SetAllowGroups(false);
+
+		TArray<AActor*> GroupActors;
+		InGroupActor->GetGroupActors(GroupActors);
+		for (AActor* Actor : GroupActors)
+		{
+			TTypedElement<UTypedElementSelectionInterface> ActorSelectionHandle = InSelectionSet->GetElement<UTypedElementSelectionInterface>(UEngineElementsLibrary::AcquireEditorActorElementHandle(Actor));
+			if (CanDeselectActorElement(ActorSelectionHandle, GroupSelectionOptions))
+			{
+				bSelectionChanged |= DeselectActorElement(ActorSelectionHandle, InSelectionSet, GroupSelectionOptions);
+			}
+		}
+	}
+
+	return bSelectionChanged;
+}
diff --git a/Engine/Source/Editor/LevelEditor/Private/Elements/Actor/ActorElementLevelEditorSelectionCustomization.h b/Engine/Source/Editor/LevelEditor/Private/Elements/Actor/ActorElementLevelEditorSelectionCustomization.h
new file mode 100644
--- /dev/null
+++ b/Engine/Source/Editor/LevelEditor/Private/Elements/Actor/ActorElementLevelEditorSelectionCustomization.h
@@ -0,0 +1,27 @@
+// Copyright Epic Games, Inc. All Rights Reserved.
+
+#pragma once
+
+#include "Elements/Framework/TypedElementSelectionSet.h"
+
+class AGroupActor;
+
+class FActorElementLevelEditorSelectionCustomization : public FTypedElementAssetEditorSelectionCustomization
+{
+public:
+	virtual bool CanSelectElement(const TTypedElement<UTypedElementSelectionInterface>& InElementSelectionHandle, const FTypedElementSelectionOptions& InSelectionOptions) override;
+	virtual bool CanDeselectElement(const TTypedElement<UTypedElementSelectionInterface>& InElementSelectionHandle, const FTypedElementSelectionOptions& InSelectionOptions) override;
+	virtual bool SelectElement(const TTypedElement<UTypedElementSelectionInterface>& InElementSelectionHandle, UTypedElementList* InSelectionSet, const FTypedElementSelectionOptions& InSelectionOptions) override;
+	virtual bool DeselectElement(const TTypedElement<UTypedElementSelectionInterface>& InElementSelectionHandle, UTypedElementList* InSelectionSet, const FTypedElementSelectionOptions& InSelectionOptions) override;
+	virtual bool AllowSelectionModifiers(const TTypedElement<UTypedElementSelectionInterface>& InElementSelectionHandle, const UTypedElementList* InSelectionSet) override;
+	virtual FTypedElementHandle GetSelectionElement(const TTypedElement<UTypedElementSelectionInterface>& InElementSelectionHandle, const UTypedElementList* InCurrentSelection, const ETypedElementSelectionMethod InSelectionMethod) override;
+
+	static bool CanSelectActorElement(const TTypedElement<UTypedElementSelectionInterface>& InActorSelectionHandle, const FTypedElementSelectionOptions& InSelectionOptions);
+	static bool CanDeselectActorElement(const TTypedElement<UTypedElementSelectionInterface>& InActorSelectionHandle, const FTypedElementSelectionOptions& InSelectionOptions);
+
+	static bool SelectActorElement(const TTypedElement<UTypedElementSelectionInterface>& InActorSelectionHandle, UTypedElementList* InSelectionSet, const FTypedElementSelectionOptions& InSelectionOptions);
+	static bool DeselectActorElement(const TTypedElement<UTypedElementSelectionInterface>& InActorSelectionHandle, UTypedElementList* InSelectionSet, const FTypedElementSelectionOptions& InSelectionOptions);
+
+	static bool SelectActorGroup(AGroupActor* InGroupActor, UTypedElementList* InSelectionSet, const FTypedElementSelectionOptions& InSelectionOptions, const bool bForce);
+	static bool DeselectActorGroup(AGroupActor* InGroupActor, UTypedElementList* InSelectionSet, const FTypedElementSelectionOptions& InSelectionOptions, const bool bForce);
+};
diff --git a/Engine/Source/Editor/LevelEditor/Private/Elements/Actor/ActorElementLevelEditorSelectionProxy.cpp b/Engine/Source/Editor/LevelEditor/Private/Elements/Actor/ActorElementLevelEditorSelectionProxy.cpp
deleted file mode 100644
diff --git a/Engine/Source/Editor/LevelEditor/Private/Elements/Actor/ActorElementLevelEditorSelectionProxy.h b/Engine/Source/Editor/LevelEditor/Private/Elements/Actor/ActorElementLevelEditorSelectionProxy.h
deleted file mode 100644
diff --git a/Engine/Source/Editor/LevelEditor/Private/Elements/Component/ComponentElementLevelEditorSelectionCustomization.cpp b/Engine/Source/Editor/LevelEditor/Private/Elements/Component/ComponentElementLevelEditorSelectionCustomization.cpp
new file mode 100644
--- /dev/null
+++ b/Engine/Source/Editor/LevelEditor/Private/Elements/Component/ComponentElementLevelEditorSelectionCustomization.cpp
@@ -0,0 +1,179 @@
+// Copyright Epic Games, Inc. All Rights Reserved.
+
+#include "Elements/Component/ComponentElementLevelEditorSelectionCustomization.h"
+#include "Elements/Component/ComponentElementData.h"
+#include "GameFramework/Actor.h"
+#include "Components/ActorComponent.h"
+
+#include "Elements/Framework/TypedElementList.h"
+#include "Elements/Framework/EngineElementsLibrary.h"
+
+#include "Elements/Actor/ActorElementSelectionInterface.h"
+#include "Elements/Component/ComponentElementSelectionInterface.h"
+
+#include "EditorViewportClient.h"
+#include "Kismet2/ComponentEditorUtils.h"
+#include "Kismet2/KismetEditorUtilities.h"
+
+DEFINE_LOG_CATEGORY_STATIC(LogComponentLevelEditorSelection, Log, All);
+
+bool FComponentElementLevelEditorSelectionCustomization::CanSelectElement(const TTypedElement<UTypedElementSelectionInterface>& InElementSelectionHandle, const FTypedElementSelectionOptions& InSelectionOptions)
+{
+	return InElementSelectionHandle.IsDataOfType<FComponentElementData>() && CanSelectComponentElement(InElementSelectionHandle, InSelectionOptions);
+}
+
+bool FComponentElementLevelEditorSelectionCustomization::CanDeselectElement(const TTypedElement<UTypedElementSelectionInterface>& InElementSelectionHandle, const FTypedElementSelectionOptions& InSelectionOptions)
+{
+	return InElementSelectionHandle.IsDataOfType<FComponentElementData>() && CanDeselectComponentElement(InElementSelectionHandle, InSelectionOptions);
+}
+
+bool FComponentElementLevelEditorSelectionCustomization::SelectElement(const TTypedElement<UTypedElementSelectionInterface>& InElementSelectionHandle, UTypedElementList* InSelectionSet, const FTypedElementSelectionOptions& InSelectionOptions)
+{
+	return InElementSelectionHandle.IsDataOfType<FComponentElementData>() && SelectComponentElement(InElementSelectionHandle, InSelectionSet, InSelectionOptions);
+}
+
+bool FComponentElementLevelEditorSelectionCustomization::DeselectElement(const TTypedElement<UTypedElementSelectionInterface>& InElementSelectionHandle, UTypedElementList* InSelectionSet, const FTypedElementSelectionOptions& InSelectionOptions)
+{
+	return InElementSelectionHandle.IsDataOfType<FComponentElementData>() && DeselectComponentElement(InElementSelectionHandle, InSelectionSet, InSelectionOptions);
+}
+
+FTypedElementHandle FComponentElementLevelEditorSelectionCustomization::GetSelectionElement(const TTypedElement<UTypedElementSelectionInterface>& InElementSelectionHandle, const UTypedElementList* InCurrentSelection, const ETypedElementSelectionMethod InSelectionMethod)
+{
+	if (const FComponentElementData* ComponentData = InElementSelectionHandle.GetData<FComponentElementData>())
+	{
+		const AActor* ConsideredActor = ComponentData->Component->GetOwner();
+		const USceneComponent* ConsideredComponent = Cast<USceneComponent>(ComponentData->Component);
+		if (ConsideredActor)
+		{
+			while (ConsideredActor->IsChildActor())
+			{
+				ConsideredActor = ConsideredActor->GetParentActor();
+				ConsideredComponent = ConsideredActor->GetParentComponent();
+			}
+
+			// If the component selected is a visualization component, we want to select the non-visualization component it's attached to
+			while (ConsideredComponent && ConsideredComponent->IsVisualizationComponent())
+			{
+				ConsideredComponent = ConsideredComponent->GetAttachParent();
+			}
+
+			const FTypedElementHandle ConsideredActorHandle = UEngineElementsLibrary::AcquireEditorActorElementHandle(ConsideredActor);
+
+			// We want to process the click on the component only if:
+			// 1. The actor clicked is already selected and is the only actor selected
+			// 2. The actor selected is blueprintable
+			// 3. No components are already selected and the click was a double click
+			// 4. OR, a component is already selected and the click was NOT a double click
+			const bool bActorAlreadySelectedExclusively = InCurrentSelection->Contains(ConsideredActorHandle) && UActorElementSelectionInterface::GetNumSelectedActors(InCurrentSelection) == 1;
+			const bool bActorIsBlueprintable = FKismetEditorUtilities::CanCreateBlueprintOfClass(ConsideredActor->GetClass());
+			const bool bComponentAlreadySelected = UComponentElementSelectionInterface::HasSelectedComponents(InCurrentSelection);
+			const bool bWasDoubleClick = InSelectionMethod == ETypedElementSelectionMethod::Secondary;
+
+			const bool bSelectComponent = bActorAlreadySelectedExclusively && bActorIsBlueprintable && (bComponentAlreadySelected != bWasDoubleClick);
+
+			if (bSelectComponent && ConsideredComponent)
+			{
+				return UEngineElementsLibrary::AcquireEditorComponentElementHandle(ConsideredComponent);
+			}
+			else
+			{
+				return ConsideredActorHandle;
+			}
+		}
+	}
+	return InElementSelectionHandle;
+}
+
+bool FComponentElementLevelEditorSelectionCustomization::CanSelectComponentElement(const TTypedElement<UTypedElementSelectionInterface>& InComponentSelectionHandle, const FTypedElementSelectionOptions& InSelectionOptions)
+{
+	UActorComponent* Component = InComponentSelectionHandle.GetDataChecked<FComponentElementData>().Component;
+
+	// Bail if global selection is locked
+	return !GEdSelectionLock;
+}
+
+bool FComponentElementLevelEditorSelectionCustomization::CanDeselectComponentElement(const TTypedElement<UTypedElementSelectionInterface>& InComponentSelectionHandle, const FTypedElementSelectionOptions& InSelectionOptions)
+{
+	UActorComponent* Component = InComponentSelectionHandle.GetDataChecked<FComponentElementData>().Component;
+
+	// Bail if global selection is locked
+	return !GEdSelectionLock;
+}
+
+bool FComponentElementLevelEditorSelectionCustomization::SelectComponentElement(const TTypedElement<UTypedElementSelectionInterface>& InComponentSelectionHandle, UTypedElementList* InSelectionSet, const FTypedElementSelectionOptions& InSelectionOptions)
+{
+	UActorComponent* Component = InComponentSelectionHandle.GetDataChecked<FComponentElementData>().Component;
+
+	if (!InComponentSelectionHandle.SelectElement(InSelectionSet, InSelectionOptions))
+	{
+		return false;
+	}
+
+	UE_LOG(LogComponentLevelEditorSelection, Verbose, TEXT("Selected Component: %s"), *Component->GetClass()->GetName());
+	
+	// Update the annotation state
+	GSelectedComponentAnnotation.Set(Component);
+	
+	// Make sure the override delegate is bound properly
+	if (USceneComponent* SceneComponent = Cast<USceneComponent>(Component))
+	{
+		FComponentEditorUtils::BindComponentSelectionOverride(SceneComponent, /*bBind*/true);
+	}
+
+	if (AActor* ComponentOwner = Component->GetOwner())
+	{
+		// Selecting a component requires that its owner actor be selected too
+		TTypedElement<UTypedElementSelectionInterface> ActorSelectionHandle = InSelectionSet->GetElement<UTypedElementSelectionInterface>(UEngineElementsLibrary::AcquireEditorActorElementHandle(ComponentOwner));
+		ActorSelectionHandle.SelectElement(InSelectionSet, InSelectionOptions);
+
+		// Update the selection visualization
+		TInlineComponentArray<UPrimitiveComponent*> PrimitiveComponents;
+		ComponentOwner->GetComponents(PrimitiveComponents);
+	
+		for (UPrimitiveComponent* PrimitiveComponent : PrimitiveComponents)
+		{
+			PrimitiveComponent->PushSelectionToProxy();
+		}
+	}
+
+	return true;
+}
+
+bool FComponentElementLevelEditorSelectionCustomization::DeselectComponentElement(const TTypedElement<UTypedElementSelectionInterface>& InComponentSelectionHandle, UTypedElementList* InSelectionSet, const FTypedElementSelectionOptions& InSelectionOptions)
+{
+	UActorComponent* Component = InComponentSelectionHandle.GetDataChecked<FComponentElementData>().Component;
+
+	if (!InComponentSelectionHandle.DeselectElement(InSelectionSet, InSelectionOptions))
+	{
+		return false;
+	}
+
+	UE_LOG(LogComponentLevelEditorSelection, Verbose, TEXT("Deselected Component: %s"), *Component->GetClass()->GetName());
+
+	// Update the annotation state
+	GSelectedComponentAnnotation.Clear(Component);
+	
+	// Update the selection visualization
+	if (AActor* ComponentOwner = Component->GetOwner())
+	{
+		TTypedElement<UTypedElementSelectionInterface> ActorSelectionHandle = InSelectionSet->GetElement<UTypedElementSelectionInterface>(UEngineElementsLibrary::AcquireEditorActorElementHandle(ComponentOwner));
+		if (!ActorSelectionHandle.IsElementSelected(InSelectionSet, FTypedElementIsSelectedOptions().SetAllowIndirect(true)))
+		{
+			// Make sure the override delegate is unbound properly
+			if (USceneComponent* SceneComponent = Cast<USceneComponent>(Component))
+			{
+				FComponentEditorUtils::BindComponentSelectionOverride(SceneComponent, /*bBind*/false);
+			}
+		}
+
+		TInlineComponentArray<UPrimitiveComponent*> PrimitiveComponents;
+		ComponentOwner->GetComponents(PrimitiveComponents);
+
+		for (UPrimitiveComponent* PrimitiveComponent : PrimitiveComponents)
+		{
+			PrimitiveComponent->PushSelectionToProxy();
+		}
+	}
+
+	return true;
+}
diff --git a/Engine/Source/Editor/LevelEditor/Private/Elements/Component/ComponentElementLevelEditorSelectionCustomization.h b/Engine/Source/Editor/LevelEditor/Private/Elements/Component/ComponentElementLevelEditorSelectionCustomization.h
new file mode 100644
--- /dev/null
+++ b/Engine/Source/Editor/LevelEditor/Private/Elements/Component/ComponentElementLevelEditorSelectionCustomization.h
@@ -0,0 +1,21 @@
+// Copyright Epic Games, Inc. All Rights Reserved.
+
+#pragma once
+
+#include "Elements/Framework/TypedElementSelectionSet.h"
+
+class FComponentElementLevelEditorSelectionCustomization : public FTypedElementAssetEditorSelectionCustomization
+{
+public:
+	virtual bool CanSelectElement(const TTypedElement<UTypedElementSelectionInterface>& InElementSelectionHandle, const FTypedElementSelectionOptions& InSelectionOptions) override;
+	virtual bool CanDeselectElement(const TTypedElement<UTypedElementSelectionInterface>& InElementSelectionHandle, const FTypedElementSelectionOptions& InSelectionOptions) override;
+	virtual bool SelectElement(const TTypedElement<UTypedElementSelectionInterface>& InElementSelectionHandle, UTypedElementList* InSelectionSet, const FTypedElementSelectionOptions& InSelectionOptions) override;
+	virtual bool DeselectElement(const TTypedElement<UTypedElementSelectionInterface>& InElementSelectionHandle, UTypedElementList* InSelectionSet, const FTypedElementSelectionOptions& InSelectionOptions) override;
+	virtual FTypedElementHandle GetSelectionElement(const TTypedElement<UTypedElementSelectionInterface>& InElementSelectionHandle, const UTypedElementList* InCurrentSelection, const ETypedElementSelectionMethod InSelectionMethod) override;
+
+	static bool CanSelectComponentElement(const TTypedElement<UTypedElementSelectionInterface>& InComponentSelectionHandle, const FTypedElementSelectionOptions& InSelectionOptions);
+	static bool CanDeselectComponentElement(const TTypedElement<UTypedElementSelectionInterface>& InComponentSelectionHandle, const FTypedElementSelectionOptions& InSelectionOptions);
+
+	static bool SelectComponentElement(const TTypedElement<UTypedElementSelectionInterface>& InComponentSelectionHandle, UTypedElementList* InSelectionSet, const FTypedElementSelectionOptions& InSelectionOptions);
+	static bool DeselectComponentElement(const TTypedElement<UTypedElementSelectionInterface>& InComponentSelectionHandle, UTypedElementList* InSelectionSet, const FTypedElementSelectionOptions& InSelectionOptions);
+};
diff --git a/Engine/Source/Editor/LevelEditor/Private/Elements/Component/ComponentElementLevelEditorSelectionProxy.cpp b/Engine/Source/Editor/LevelEditor/Private/Elements/Component/ComponentElementLevelEditorSelectionProxy.cpp
deleted file mode 100644
diff --git a/Engine/Source/Editor/LevelEditor/Private/Elements/Component/ComponentElementLevelEditorSelectionProxy.h b/Engine/Source/Editor/LevelEditor/Private/Elements/Component/ComponentElementLevelEditorSelectionProxy.h
deleted file mode 100644
diff --git a/Engine/Source/Editor/LevelEditor/Private/SLevelEditor.cpp b/Engine/Source/Editor/LevelEditor/Private/SLevelEditor.cpp
--- a/Engine/Source/Editor/LevelEditor/Private/SLevelEditor.cpp
+++ b/Engine/Source/Editor/LevelEditor/Private/SLevelEditor.cpp
@@ -61,8 +61,8 @@
 #include "Editor/EnvironmentLightingViewer/Public/EnvironmentLightingModule.h"
 #include "Misc/MessageDialog.h"
 #include "Elements/Framework/TypedElementSelectionSet.h"
-#include "Elements/Actor/ActorElementLevelEditorSelectionProxy.h"
-#include "Elements/Component/ComponentElementLevelEditorSelectionProxy.h"
+#include "Elements/Actor/ActorElementLevelEditorSelectionCustomization.h"
+#include "Elements/Component/ComponentElementLevelEditorSelectionCustomization.h"
 
 #define LOCTEXT_NAMESPACE "SLevelEditor"
 
@@ -222,8 +222,8 @@ void SLevelEditor::Initialize( const TSharedRef<SDockTab>& OwnerTab, const TShar
 	SelectedElements->AddToRoot();
 
 	// Register the level editor specific selection behavior
-	SelectedElements->RegisterAssetEditorSelectionProxy(NAME_Actor, NewObject<UActorElementLevelEditorSelectionProxy>());
-	SelectedElements->RegisterAssetEditorSelectionProxy(NAME_Components, NewObject<UComponentElementLevelEditorSelectionProxy>());
+	SelectedElements->RegisterAssetEditorCustomizationByTypeName(NAME_Actor, MakeUnique<FActorElementLevelEditorSelectionCustomization>());
+	SelectedElements->RegisterAssetEditorCustomizationByTypeName(NAME_Components, MakeUnique<FComponentElementLevelEditorSelectionCustomization>());
 
 	// Allow USelection to bridge to our selected element list
 	GUnrealEd->GetSelectedActors()->SetElementSelectionSet(SelectedElements);
diff --git a/Engine/Source/Editor/UnrealEd/Private/EditorEngine.cpp b/Engine/Source/Editor/UnrealEd/Private/EditorEngine.cpp
--- a/Engine/Source/Editor/UnrealEd/Private/EditorEngine.cpp
+++ b/Engine/Source/Editor/UnrealEd/Private/EditorEngine.cpp
@@ -130,7 +130,7 @@
 #include "Editor/ActorPositioning.h"
 
 #include "Elements/Framework/TypedElementSelectionSet.h"
-#include "Elements/Object/ObjectElementEditorSelectionProxy.h"
+#include "Elements/Object/ObjectElementEditorSelectionCustomization.h"
 
 #include "IDirectoryWatcher.h"
 #include "DirectoryWatcherModule.h"
@@ -269,7 +269,7 @@ void InitSelectionSets()
 	GObjectSelection->AddToRoot();
 
 	UTypedElementSelectionSet* ObjectSelectionSet = NewObject<UTypedElementSelectionSet>(GObjectSelection, NAME_None, RF_Transactional);
-	ObjectSelectionSet->RegisterAssetEditorSelectionProxy(NAME_Object, NewObject<UObjectElementEditorSelectionProxy>());
+	ObjectSelectionSet->RegisterAssetEditorCustomizationByTypeName(NAME_Object, MakeUnique<FObjectElementEditorSelectionCustomization>());
 	GObjectSelection->SetElementSelectionSet(ObjectSelectionSet);
 }
 
diff --git a/Engine/Source/Editor/UnrealEd/Private/Elements/Object/ObjectElementEditorSelectionCustomization.cpp b/Engine/Source/Editor/UnrealEd/Private/Elements/Object/ObjectElementEditorSelectionCustomization.cpp
new file mode 100644
--- /dev/null
+++ b/Engine/Source/Editor/UnrealEd/Private/Elements/Object/ObjectElementEditorSelectionCustomization.cpp
@@ -0,0 +1,33 @@
+// Copyright Epic Games, Inc. All Rights Reserved.
+
+#include "Elements/Object/ObjectElementEditorSelectionCustomization.h"
+#include "Elements/Object/ObjectElementData.h"
+#include "UObject/UObjectAnnotation.h"
+
+bool FObjectElementEditorSelectionCustomization::SelectElement(const TTypedElement<UTypedElementSelectionInterface>& InElementSelectionHandle, UTypedElementList* InSelectionSet, const FTypedElementSelectionOptions& InSelectionOptions)
+{
+	if (const FObjectElementData* ObjectData = InElementSelectionHandle.GetData<FObjectElementData>())
+	{
+		if (InElementSelectionHandle.SelectElement(InSelectionSet, InSelectionOptions))
+		{
+			GSelectedObjectAnnotation.Set(ObjectData->Object);
+			return true;
+		}
+	}
+
+	return false;
+}
+
+bool FObjectElementEditorSelectionCustomization::DeselectElement(const TTypedElement<UTypedElementSelectionInterface>& InElementSelectionHandle, UTypedElementList* InSelectionSet, const FTypedElementSelectionOptions& InSelectionOptions)
+{
+	if (const FObjectElementData* ObjectData = InElementSelectionHandle.GetData<FObjectElementData>())
+	{
+		if (InElementSelectionHandle.DeselectElement(InSelectionSet, InSelectionOptions))
+		{
+			GSelectedObjectAnnotation.Clear(ObjectData->Object);
+			return true;
+		}
+	}
+
+	return false;
+}
diff --git a/Engine/Source/Editor/UnrealEd/Private/Elements/Object/ObjectElementEditorSelectionCustomization.h b/Engine/Source/Editor/UnrealEd/Private/Elements/Object/ObjectElementEditorSelectionCustomization.h
new file mode 100644
--- /dev/null
+++ b/Engine/Source/Editor/UnrealEd/Private/Elements/Object/ObjectElementEditorSelectionCustomization.h
@@ -0,0 +1,12 @@
+// Copyright Epic Games, Inc. All Rights Reserved.
+
+#pragma once
+
+#include "Elements/Framework/TypedElementSelectionSet.h"
+
+class FObjectElementEditorSelectionCustomization : public FTypedElementAssetEditorSelectionCustomization
+{
+public:
+	virtual bool SelectElement(const TTypedElement<UTypedElementSelectionInterface>& InElementSelectionHandle, UTypedElementList* InSelectionSet, const FTypedElementSelectionOptions& InSelectionOptions) override;
+	virtual bool DeselectElement(const TTypedElement<UTypedElementSelectionInterface>& InElementSelectionHandle, UTypedElementList* InSelectionSet, const FTypedElementSelectionOptions& InSelectionOptions) override;
+};
diff --git a/Engine/Source/Editor/UnrealEd/Private/Elements/Object/ObjectElementEditorSelectionProxy.cpp b/Engine/Source/Editor/UnrealEd/Private/Elements/Object/ObjectElementEditorSelectionProxy.cpp
deleted file mode 100644
diff --git a/Engine/Source/Editor/UnrealEd/Private/Elements/Object/ObjectElementEditorSelectionProxy.h b/Engine/Source/Editor/UnrealEd/Private/Elements/Object/ObjectElementEditorSelectionProxy.h
deleted file mode 100644
