Allow pending kill actors to be re-selected during an undo/redo

They may still be pending kill as the transaction is replayed, and restoring the selection state never used to re-test whether an actor could be selected, so this restores the old behavior.

#rnx

diff --git a/Engine/Source/Editor/LevelEditor/Private/Elements/Actor/ActorElementLevelEditorSelectionCustomization.cpp b/Engine/Source/Editor/LevelEditor/Private/Elements/Actor/ActorElementLevelEditorSelectionCustomization.cpp
--- a/Engine/Source/Editor/LevelEditor/Private/Elements/Actor/ActorElementLevelEditorSelectionCustomization.cpp
+++ b/Engine/Source/Editor/LevelEditor/Private/Elements/Actor/ActorElementLevelEditorSelectionCustomization.cpp
@@ -80,12 +80,12 @@ bool FActorElementLevelEditorSelectionCustomization::CanSelectActorElement(const
 
 	// Ensure that neither the level nor the actor is being destroyed or is unreachable
 	const EObjectFlags InvalidSelectableFlags = RF_BeginDestroyed;
-	if (Actor->GetLevel()->HasAnyFlags(InvalidSelectableFlags) || Actor->GetLevel()->IsPendingKillOrUnreachable())
+	if (Actor->GetLevel()->HasAnyFlags(InvalidSelectableFlags) || (!GIsTransacting && Actor->GetLevel()->IsPendingKillOrUnreachable()))
 	{
 		UE_LOG(LogActorLevelEditorSelection, Warning, TEXT("SelectActor: %s (%s)"), TEXT("The requested operation could not be completed because the level has invalid flags."), *Actor->GetActorLabel());
 		return false;
 	}
-	if (Actor->HasAnyFlags(InvalidSelectableFlags) || Actor->IsPendingKillOrUnreachable())
+	if (Actor->HasAnyFlags(InvalidSelectableFlags) || (!GIsTransacting && Actor->IsPendingKillOrUnreachable()))
 	{
 		UE_LOG(LogActorLevelEditorSelection, Warning, TEXT("SelectActor: %s (%s)"), TEXT("The requested operation could not be completed because the actor has invalid flags."), *Actor->GetActorLabel());
 		return false;
