Tab drawer improvements
- activating an editor mode that was already active but in an drawer, opens that drawer
- dock tabs can now register for drawer opening/closing callbacks so that they can do stuff when the drawer opens or closes. For example, now the placement browser focuses its search field when the drawer opens

diff --git a/Engine/Source/Editor/LevelEditor/Private/SLevelEditor.cpp b/Engine/Source/Editor/LevelEditor/Private/SLevelEditor.cpp
--- a/Engine/Source/Editor/LevelEditor/Private/SLevelEditor.cpp
+++ b/Engine/Source/Editor/LevelEditor/Private/SLevelEditor.cpp
@@ -705,14 +705,14 @@ TSharedRef<SDockTab> SLevelEditor::SpawnLevelEditorTab( const FSpawnTabArgs& Arg
 	}
 	else if (TabIdentifier == LevelEditorTabIds::PlacementBrowser)
 	{
-		return
-			SNew(SDockTab)
+		TSharedRef<SDockTab> DockTab = SNew(SDockTab)
 			.Icon(FEditorStyle::GetBrush("LevelEditor.Tabs.PlacementBrowser"))
 			.Label(NSLOCTEXT("LevelEditor", "PlacementBrowserTitle", "Place Actors"))
-			.AddMetaData<FTutorialMetaData>(FTutorialMetaData(TEXT("PlacementBrowser"), TEXT("PlacementBrowser")))
-			[
-				IPlacementModeModule::Get().CreatePlacementModeBrowser()
-			];
+			.AddMetaData<FTutorialMetaData>(FTutorialMetaData(TEXT("PlacementBrowser"), TEXT("PlacementBrowser")));
+
+			DockTab->SetContent(IPlacementModeModule::Get().CreatePlacementModeBrowser(DockTab));
+
+		return DockTab;
 	}
 	else if( TabIdentifier == LevelEditorTabIds::LevelEditorBuildAndSubmit )
 	{
@@ -1509,6 +1509,9 @@ FName SLevelEditor::GetEditorModeTabId( FEditorModeID ModeID )
 
 void SLevelEditor::ToggleEditorMode( FEditorModeID ModeID )
 {
+	FLevelEditorModule& LevelEditorModule = FModuleManager::GetModuleChecked<FLevelEditorModule>("LevelEditor");
+	TSharedPtr<FTabManager> LevelEditorTabManager = LevelEditorModule.GetLevelEditorTabManager();
+
 	// Prompt the user if Matinee must be closed before activating new mode
 	if (ModeID != FBuiltinEditorModes::EM_InterpEdit)
 	{
@@ -1531,6 +1534,10 @@ void SLevelEditor::ToggleEditorMode( FEditorModeID ModeID )
 	// *Important* - activate the mode first since FEditorModeTools::DeactivateMode will
 	// activate the default mode when the stack becomes empty, resulting in multiple active visible modes.
 	GLevelEditorModeTools().ActivateMode( ModeID );
+
+	// Invoke the tab regardless if if the mode was activated or already activated so that it the mode tools come into the foreground
+	LevelEditorTabManager->TryInvokeTab(LevelEditorTabIds::LevelEditorToolBox);
+
 }
 
 bool SLevelEditor::IsModeActive( FEditorModeID ModeID )
diff --git a/Engine/Source/Editor/LevelEditor/Private/SLevelEditorToolBox.h b/Engine/Source/Editor/LevelEditor/Private/SLevelEditorToolBox.h
--- a/Engine/Source/Editor/LevelEditor/Private/SLevelEditorToolBox.h
+++ b/Engine/Source/Editor/LevelEditor/Private/SLevelEditorToolBox.h
@@ -43,8 +43,9 @@ public:
 
 	/** Sets the parent tab of this toolbox */
 	void SetParentTab(TSharedRef<SDockTab>& InDockTab);
-private:
 
+private:
+	void OnContainingTabActivated(TSharedRef<SDockTab> DockTab, ETabActivationCause InActivationCause);
 	/** Gets the visibility for the SBorder showing toolbox editor mode inline content */
 	EVisibility GetInlineContentHolderVisibility() const;
 
diff --git a/Engine/Source/Editor/PlacementMode/Private/PlacementModeModule.cpp b/Engine/Source/Editor/PlacementMode/Private/PlacementModeModule.cpp
--- a/Engine/Source/Editor/PlacementMode/Private/PlacementModeModule.cpp
+++ b/Engine/Source/Editor/PlacementMode/Private/PlacementModeModule.cpp
@@ -322,9 +322,9 @@ void FPlacementModeModule::AddToRecentlyPlaced(UObject* Asset, UActorFactory* Fa
 	AddToRecentlyPlaced(Assets, FactoryUsed);
 }
 
-TSharedRef<SWidget> FPlacementModeModule::CreatePlacementModeBrowser()
+TSharedRef<SWidget> FPlacementModeModule::CreatePlacementModeBrowser(TSharedRef<SDockTab> ParentTab)
 {
-	return SNew(SPlacementModeTools);
+	return SNew(SPlacementModeTools, ParentTab);
 }
 
 bool FPlacementModeModule::RegisterPlacementCategory(const FPlacementCategoryInfo& Info)
diff --git a/Engine/Source/Editor/PlacementMode/Private/SPlacementModeTools.cpp b/Engine/Source/Editor/PlacementMode/Private/SPlacementModeTools.cpp
--- a/Engine/Source/Editor/PlacementMode/Private/SPlacementModeTools.cpp
+++ b/Engine/Source/Editor/PlacementMode/Private/SPlacementModeTools.cpp
@@ -313,12 +313,14 @@ SPlacementModeTools::~SPlacementModeTools()
 	}
 }
 
-void SPlacementModeTools::Construct( const FArguments& InArgs )
+void SPlacementModeTools::Construct( const FArguments& InArgs, TSharedRef<SDockTab> ParentTab )
 {
 	bPlaceablesFullRefreshRequested = false;
 	bRecentlyPlacedRefreshRequested = false;
 	bNeedsUpdate = true;
 
+	ParentTab->SetOnTabDrawerOpened(FSimpleDelegate::CreateSP(this, &SPlacementModeTools::OnTabDrawerOpened));
+
 
 	static const FName PlacementBrowserActiveTabBarBrushName("PlacementBrowser.ActiveTabBar");
 	PlacementBrowserActiveTabBarBrush = FEditorStyle::GetBrush(PlacementBrowserActiveTabBarBrushName);
@@ -593,6 +595,11 @@ void SPlacementModeTools::OnCategoryChanged(const ECheckBoxState NewState, FName
 	}
 }
 
+void SPlacementModeTools::OnTabDrawerOpened()
+{
+	FSlateApplication::Get().SetKeyboardFocus(SearchBoxPtr, EFocusCause::SetDirectly);
+}
+
 void SPlacementModeTools::OnPlacementTabChanged( ECheckBoxState NewState, FName CategoryName )
 {
 	if ( NewState == ECheckBoxState::Checked )
diff --git a/Engine/Source/Editor/PlacementMode/Private/SPlacementModeTools.h b/Engine/Source/Editor/PlacementMode/Private/SPlacementModeTools.h
--- a/Engine/Source/Editor/PlacementMode/Private/SPlacementModeTools.h
+++ b/Engine/Source/Editor/PlacementMode/Private/SPlacementModeTools.h
@@ -57,7 +57,7 @@ public:
 	SLATE_BEGIN_ARGS( SPlacementModeTools ){}
 	SLATE_END_ARGS();
 
-	void Construct( const FArguments& InArgs );
+	void Construct(const FArguments& InArgs, TSharedRef<SDockTab> ParentTab);
 
 	virtual ~SPlacementModeTools();
 
@@ -90,8 +90,6 @@ private:
 	/** Called when the list of placeable assets changes. */
 	void UpdatePlaceableAssets();
 
-private:
-
 	/** Gets the border image for the tab, this is the 'active' orange bar. */
 	const FSlateBrush* PlacementGroupBorderImage( FName CategoryName ) const;
 
@@ -99,11 +97,12 @@ private:
 	void OnPlacementTabChanged( ECheckBoxState NewState, FName CategoryName );
 	void OnCategoryChanged(const ECheckBoxState NewState, FName InCategory);
 
+	/** Called when the placement mode tools are opened from a tab drawer */
+	void OnTabDrawerOpened();
+
 	/** Gets the tab 'active' state, so that we can show the active style */
 	ECheckBoxState GetPlacementTabCheckedState( FName CategoryName ) const;
 
-private:
-
 	/** Gets the visibility for the failed search text */
 	EVisibility GetFailedSearchVisibility() const;
 
diff --git a/Engine/Source/Editor/PlacementMode/Public/IPlacementModeModule.h b/Engine/Source/Editor/PlacementMode/Public/IPlacementModeModule.h
--- a/Engine/Source/Editor/PlacementMode/Public/IPlacementModeModule.h
+++ b/Engine/Source/Editor/PlacementMode/Public/IPlacementModeModule.h
@@ -259,7 +259,7 @@ public:
 	/**
 	 * Creates the placement browser widget
 	 */
-	virtual TSharedRef<SWidget> CreatePlacementModeBrowser() = 0;
+	virtual TSharedRef<SWidget> CreatePlacementModeBrowser(TSharedRef<SDockTab> ParentTab) = 0;
 public:
 
 	/**
diff --git a/Engine/Source/Editor/PlacementMode/Public/PlacementModeModule.h b/Engine/Source/Editor/PlacementMode/Public/PlacementModeModule.h
--- a/Engine/Source/Editor/PlacementMode/Public/PlacementModeModule.h
+++ b/Engine/Source/Editor/PlacementMode/Public/PlacementModeModule.h
@@ -79,7 +79,7 @@ public:
 		return RecentlyPlaced;
 	}
 
-	virtual TSharedRef<SWidget> CreatePlacementModeBrowser() override;
+	virtual TSharedRef<SWidget> CreatePlacementModeBrowser(TSharedRef<SDockTab> ParentTab) override;
 
 	virtual bool RegisterPlacementCategory(const FPlacementCategoryInfo& Info);
 
diff --git a/Engine/Source/Runtime/Slate/Private/Framework/Docking/STabDrawer.cpp b/Engine/Source/Runtime/Slate/Private/Framework/Docking/STabDrawer.cpp
--- a/Engine/Source/Runtime/Slate/Private/Framework/Docking/STabDrawer.cpp
+++ b/Engine/Source/Runtime/Slate/Private/Framework/Docking/STabDrawer.cpp
@@ -79,8 +79,6 @@ void STabDrawer::Close()
 		AnimationThrottle = FSlateThrottleManager::Get().EnterResponsiveMode();
 		OpenCloseTimer = RegisterActiveTimer(0.0f, FWidgetActiveTimerDelegate::CreateSP(this, &STabDrawer::UpdateAnimation));
 	}
-
-	//OnContentBrowserDismissedDelegate.ExecuteIfBound(NewlyFocusedWidget);
 }
 
 bool STabDrawer::IsOpen() const
diff --git a/Engine/Source/Runtime/Slate/Private/Framework/Docking/STabSidebar.cpp b/Engine/Source/Runtime/Slate/Private/Framework/Docking/STabSidebar.cpp
--- a/Engine/Source/Runtime/Slate/Private/Framework/Docking/STabSidebar.cpp
+++ b/Engine/Source/Runtime/Slate/Private/Framework/Docking/STabSidebar.cpp
@@ -600,6 +600,8 @@ void STabSidebar::OpenDrawerInternal(TSharedRef<SDockTab> ForTab)
 		FSlateApplication::Get().SetKeyboardFocus(NewDrawer);
 
 		OpenedDrawers.Add(NewDrawer);
+
+		ForTab->OnTabDrawerOpened();
 	}
 
 	UpdateDrawerAppearance();
diff --git a/Engine/Source/Runtime/Slate/Private/Widgets/Docking/SDockTab.cpp b/Engine/Source/Runtime/Slate/Private/Widgets/Docking/SDockTab.cpp
--- a/Engine/Source/Runtime/Slate/Private/Widgets/Docking/SDockTab.cpp
+++ b/Engine/Source/Runtime/Slate/Private/Widgets/Docking/SDockTab.cpp
@@ -435,6 +435,16 @@ void SDockTab::SetOnTabRenamed(const FOnTabRenamed& InDelegate)
 	OnTabRenamed = InDelegate;
 }
 
+void SDockTab::SetOnTabDrawerOpened(const FSimpleDelegate InDelegate)
+{
+	OnTabDrawerOpenedEvent = InDelegate;
+}
+
+void SDockTab::SetOnTabDrawerClosed(const FSimpleDelegate InDelegate)
+{
+	OnTabDrawerClosedEvent = InDelegate;
+}
+
 TSharedRef<FTabManager> SDockTab::GetTabManager() const
 {
 	return MyTabManager.Pin().ToSharedRef();
@@ -692,6 +702,11 @@ void SDockTab::UpdateTabStyle()
 	LabelWidget->SetTextStyle(&TabTextStyle);
 }
 
+void SDockTab::OnTabDrawerOpened()
+{
+	OnTabDrawerOpenedEvent.ExecuteIfBound();
+}
+
 void SDockTab::OnTabDrawerClosed()
 {
 	OnTabDrawerClosedEvent.ExecuteIfBound();
diff --git a/Engine/Source/Runtime/Slate/Public/Widgets/Docking/SDockTab.h b/Engine/Source/Runtime/Slate/Public/Widgets/Docking/SDockTab.h
--- a/Engine/Source/Runtime/Slate/Public/Widgets/Docking/SDockTab.h
+++ b/Engine/Source/Runtime/Slate/Public/Widgets/Docking/SDockTab.h
@@ -251,6 +251,12 @@ public:
 	/** Set the handler that will be invoked when the tab is renamed */
 	void SetOnTabRenamed(const FOnTabRenamed& InDelegate);
 
+	/** Set the handler that will be invoked when the tab is opened from a drawer */
+	void SetOnTabDrawerOpened(const FSimpleDelegate InDelegate);
+
+	/** Set the handler that will be invoked when the tab is closed from a drawer */
+	void SetOnTabDrawerClosed(const FSimpleDelegate InDelegate);
+
 	/** Get the tab manager currently managing this tab. Note that a user move the tab between Tab Managers, so this return value may change. */
 	TSharedRef<FTabManager> GetTabManager() const;
 
@@ -342,6 +348,7 @@ private:
 
 	void UpdateTabStyle();
 
+	void OnTabDrawerOpened();
 	void OnTabDrawerClosed();
 
 	/** The handle to the active tab activation tick */
@@ -387,6 +394,8 @@ protected:
 
 	FSimpleDelegate OnTabDrawerClosedEvent;
 
+	FSimpleDelegate OnTabDrawerOpenedEvent;
+
 	/**
 	 * Invoked during the Save Visual State pass; gives this tab a chance to save misc info about visual state.
 	 * e.g. Content Browser might save the current filters, current folder, whether some panel is collapsed, etc.
