Add a drop down menu to the Content shortcut in the level editor, to which plugins can add additional shortcuts.

diff --git a/Engine/Source/Editor/ContentBrowser/Private/ContentBrowserSingleton.cpp b/Engine/Source/Editor/ContentBrowser/Private/ContentBrowserSingleton.cpp
--- a/Engine/Source/Editor/ContentBrowser/Private/ContentBrowserSingleton.cpp
+++ b/Engine/Source/Editor/ContentBrowser/Private/ContentBrowserSingleton.cpp
@@ -29,7 +29,7 @@
 #include "ContentBrowserCommands.h"
 #include "CoreGlobals.h"
 #include "StatusBarSubsystem.h"
-#include "Toolkits/GlobalEditorCommonCommands.h"
+#include "ToolMenus.h"
 
 #define LOCTEXT_NAMESPACE "ContentBrowser"
 
@@ -49,6 +49,28 @@ FContentBrowserSingleton::FContentBrowserSingleton()
 		ContentBrowserIcon,
 		true);
 
+	UToolMenu* ContentMenu = UToolMenus::Get()->ExtendMenu("LevelEditor.LevelEditorToolBar.ContentQuickMenu");
+	FToolMenuSection& Section = ContentMenu->FindOrAddSection("ProjectContent");
+	Section.AddMenuEntry("FocusContentBrowser",
+		LOCTEXT("FocusContentBrowser_Label", "Focus Content Browser"),
+		LOCTEXT("FocusContentBrowser_Desc", "Focuses the most recently active content browser tab."),
+		FSlateIcon(),
+		FUIAction(FExecuteAction::CreateRaw(this, &FContentBrowserSingleton::FocusPrimaryContentBrowser, true))
+	);
+	Section.AddSeparator(NAME_None);
+
+	auto AddSpawnerEntryToMenuSection = [](FToolMenuSection& InSection, FTabSpawnerEntry& InSpawnerNode, FName InTabID)
+	{
+		InSection.AddMenuEntry(
+			InTabID,
+			InSpawnerNode.GetDisplayName().IsEmpty() ? FText::FromName(InTabID) : InSpawnerNode.GetDisplayName(),
+			InSpawnerNode.GetTooltipText(),
+			FSlateIcon(),
+			FGlobalTabmanager::Get()->GetUIActionForTabSpawnerMenuEntry(InSpawnerNode.AsSpawnerEntry()),
+			EUserInterfaceActionType::Check
+			);
+	};
+
 	for ( int32 BrowserIdx = 0; BrowserIdx < UE_ARRAY_COUNT(ContentBrowserTabIDs); BrowserIdx++ )
 	{
 		const FName TabID = FName(*FString::Printf(TEXT("ContentBrowserTab%d"), BrowserIdx + 1));
@@ -56,11 +78,13 @@ FContentBrowserSingleton::FContentBrowserSingleton()
 
 		const FText DefaultDisplayName = GetContentBrowserLabelWithIndex( BrowserIdx );
 
-		FGlobalTabmanager::Get()->RegisterNomadTabSpawner( TabID, FOnSpawnTab::CreateRaw(this, &FContentBrowserSingleton::SpawnContentBrowserTab, BrowserIdx) )
+		FTabSpawnerEntry& ContentBrowserTabSpawner = FGlobalTabmanager::Get()->RegisterNomadTabSpawner( TabID, FOnSpawnTab::CreateRaw(this, &FContentBrowserSingleton::SpawnContentBrowserTab, BrowserIdx) )
 			.SetDisplayName(DefaultDisplayName)
 			.SetTooltipText( LOCTEXT( "ContentBrowserMenuTooltipText", "Open a Content Browser tab." ) )
 			.SetGroup( ContentBrowserGroup )
 			.SetIcon(ContentBrowserIcon);
+
+		AddSpawnerEntryToMenuSection(Section, ContentBrowserTabSpawner, TabID);
 	}
 
 	// Register a couple legacy tab ids
diff --git a/Engine/Source/Editor/LevelEditor/Private/LevelEditorToolBar.cpp b/Engine/Source/Editor/LevelEditor/Private/LevelEditorToolBar.cpp
--- a/Engine/Source/Editor/LevelEditor/Private/LevelEditorToolBar.cpp
+++ b/Engine/Source/Editor/LevelEditor/Private/LevelEditorToolBar.cpp
@@ -1179,6 +1179,7 @@ void FLevelEditorToolBar::RegisterLevelEditorToolBar( const TSharedRef<FUIComman
 
 	RegisterQuickSettingsMenu();
 	RegisterOpenBlueprintMenu();
+	RegisterContentMenu();
 
 #define LOCTEXT_NAMESPACE "LevelEditorToolBar"
 
@@ -1200,7 +1201,14 @@ void FLevelEditorToolBar::RegisterLevelEditorToolBar( const TSharedRef<FUIComman
 
 		{
 			FToolMenuSection& Section = AssetsToolBar->AddSection("Content");
-			Section.AddEntry(FToolMenuEntry::InitToolBarButton(FLevelEditorCommands::Get().OpenContentBrowser, LOCTEXT("ContentBrowser_Override", "Content"), TAttribute<FText>(), TAttribute<FSlateIcon>(), "LevelToolbarContent"));
+			Section.AddEntry(FToolMenuEntry::InitComboButton(
+				"Content",
+				FUIAction(),
+				FOnGetContent::CreateStatic(&FLevelEditorToolBar::GenerateContentMenuWidget, InCommandList, TWeakPtr<SLevelEditor>(InLevelEditor)),
+				LOCTEXT("OpenContent_Label", "Content"),
+				LOCTEXT("OpenContent_Tooltip", "Quickly find content available to add to the project."),
+				FSlateIcon(FEditorStyle::GetStyleSetName(), "LevelEditor.OpenContentBrowser")
+			));
 
 			Section.AddEntry(FToolMenuEntry::InitComboButton(
 				"OpenBlueprint",
@@ -1818,6 +1826,22 @@ void FLevelEditorToolBar::RegisterOpenBlueprintMenu()
 #undef LOCTEXT_NAMESPACE
 }
 
+void FLevelEditorToolBar::RegisterContentMenu()
+{
+#define LOCTEXT_NAMESPACE "LevelToolBarViewMenu"
+	UToolMenu* ContentMenu = UToolMenus::Get()->RegisterMenu("LevelEditor.LevelEditorToolBar.ContentQuickMenu");
+	{
+		FToolMenuSection& Section = ContentMenu->FindOrAddSection("ProjectContent");
+		Section.InitSection("ProjectContent", LOCTEXT("ProjectContent_Label", "Project Content"), FToolMenuInsert("ContentBrowser", EToolMenuInsertType::First));
+	}
+	{
+		FToolMenuSection& Section = ContentMenu->FindOrAddSection("ExternalContent");
+		Section.InitSection("ExternalContent", LOCTEXT("ExternalContent_Label", "Get Content"), FToolMenuInsert());
+		Section.AddMenuEntry(FLevelEditorCommands::Get().OpenMarketplace);
+	}
+#undef LOCTEXT_NAMESPACE
+}
+
 void FLevelEditorToolBar::OnOpenSubLevelBlueprint( ULevel* InLevel )
 {
 	ULevelScriptBlueprint* LevelScriptBlueprint = InLevel->GetLevelScriptBlueprint();
@@ -1844,6 +1868,16 @@ TSharedRef< SWidget > FLevelEditorToolBar::GenerateCinematicsMenuContent(TShared
 	return UToolMenus::Get()->GenerateWidget("LevelEditor.LevelEditorToolBar.Cinematics", MenuContext);
 }
 
+TSharedRef< SWidget > FLevelEditorToolBar::GenerateContentMenuWidget(TSharedRef<FUICommandList> InCommandList, TWeakPtr<SLevelEditor> InLevelEditor)
+{
+	FToolMenuContext MenuContext(InCommandList);
+	ULevelEditorMenuContext* LevelEditorMenuContext = NewObject<ULevelEditorMenuContext>();
+	LevelEditorMenuContext->LevelEditor = InLevelEditor;
+	MenuContext.AddObject(LevelEditorMenuContext);
+
+	return UToolMenus::Get()->GenerateWidget("LevelEditor.LevelEditorToolBar.ContentQuickMenu", MenuContext);
+}
+
 void FLevelEditorToolBar::RegisterCinematicsMenu()
 {
 #define LOCTEXT_NAMESPACE "LevelToolBarCinematicsMenu"
diff --git a/Engine/Source/Editor/LevelEditor/Private/LevelEditorToolBar.h b/Engine/Source/Editor/LevelEditor/Private/LevelEditorToolBar.h
--- a/Engine/Source/Editor/LevelEditor/Private/LevelEditorToolBar.h
+++ b/Engine/Source/Editor/LevelEditor/Private/LevelEditorToolBar.h
@@ -64,6 +64,8 @@ protected:
 	 */
 	static TSharedRef< SWidget > GenerateCinematicsMenuContent(TSharedRef<FUICommandList> InCommandList, TWeakPtr<SLevelEditor> InLevelEditor);
 
+	static TSharedRef< SWidget > GenerateContentMenuWidget(TSharedRef<FUICommandList> InCommandList, TWeakPtr<SLevelEditor> InLevelEditor);
+
 	/**
 	 * Delegate for actor selection within the Cinematics popup menu's SceneOutliner.
 	 * Opens the editor for the selected actor and dismisses all popup menus.
@@ -84,4 +86,5 @@ private:
 
 	static void RegisterQuickSettingsMenu();
 	static void RegisterOpenBlueprintMenu();
+	static void RegisterContentMenu();
 };
diff --git a/Engine/Source/Runtime/Slate/Private/Framework/Docking/TabManager.cpp b/Engine/Source/Runtime/Slate/Private/Framework/Docking/TabManager.cpp
--- a/Engine/Source/Runtime/Slate/Private/Framework/Docking/TabManager.cpp
+++ b/Engine/Source/Runtime/Slate/Private/Framework/Docking/TabManager.cpp
@@ -1072,28 +1072,13 @@ void FTabManager::PopulateTabSpawnerMenu_Helper( FMenuBuilder& PopulateMe, FPopu
 
 void FTabManager::MakeSpawnerMenuEntry( FMenuBuilder &PopulateMe, const TSharedPtr<FTabSpawnerEntry> &InSpawnerNode ) 
 {
-	auto CanExecuteMenuEntry = [](TWeakPtr<FTabSpawnerEntry> SpawnerNode) -> bool
-	{
-		TSharedPtr<FTabSpawnerEntry> SpawnerNodePinned = SpawnerNode.Pin();
-		if (SpawnerNodePinned.IsValid() && SpawnerNodePinned->MenuType.Get() == ETabSpawnerMenuType::Enabled)
-		{
-			return SpawnerNodePinned->CanSpawnTab.IsBound() ? SpawnerNodePinned->CanSpawnTab.Execute(FSpawnTabArgs(TSharedPtr<SWindow>(), SpawnerNodePinned->TabType)) : true;
-		}
-
-		return false;
-	};
-
 	if (InSpawnerNode->MenuType.Get() != ETabSpawnerMenuType::Hidden )
 	{
 		PopulateMe.AddMenuEntry(
 			InSpawnerNode->GetDisplayName().IsEmpty() ? FText::FromName(InSpawnerNode->TabType ) : InSpawnerNode->GetDisplayName(),
 			InSpawnerNode->GetTooltipText(),
 			InSpawnerNode->GetIcon(),
-			FUIAction(
-			FExecuteAction::CreateSP(SharedThis(this), &FTabManager::InvokeTabForMenu, InSpawnerNode->TabType),
-			FCanExecuteAction::CreateStatic(CanExecuteMenuEntry, TWeakPtr<FTabSpawnerEntry>(InSpawnerNode)),
-			FIsActionChecked::CreateSP(InSpawnerNode.ToSharedRef(), &FTabSpawnerEntry::IsSoleTabInstanceSpawned)
-			),
+			GetUIActionForTabSpawnerMenuEntry(InSpawnerNode),
 			NAME_None,
 			EUserInterfaceActionType::Check
 			);
@@ -1400,6 +1385,26 @@ TSharedPtr<SDockingTabStack> FTabManager::FindPotentiallyClosedTab( const FTabId
 	return StackWithClosedTab;
 }
 
+FUIAction FTabManager::GetUIActionForTabSpawnerMenuEntry(TSharedPtr<FTabSpawnerEntry> InTabMenuEntry)
+{
+	auto CanExecuteMenuEntry = [](TWeakPtr<FTabSpawnerEntry> SpawnerNode) -> bool
+	{
+		TSharedPtr<FTabSpawnerEntry> SpawnerNodePinned = SpawnerNode.Pin();
+		if (SpawnerNodePinned.IsValid() && SpawnerNodePinned->MenuType.Get() == ETabSpawnerMenuType::Enabled)
+		{
+			return SpawnerNodePinned->CanSpawnTab.IsBound() ? SpawnerNodePinned->CanSpawnTab.Execute(FSpawnTabArgs(TSharedPtr<SWindow>(), SpawnerNodePinned->TabType)) : true;
+		}
+
+		return false;
+	};
+
+	return FUIAction(
+		FExecuteAction::CreateSP(SharedThis(this), &FTabManager::InvokeTabForMenu, InTabMenuEntry->TabType),
+		FCanExecuteAction::CreateStatic(CanExecuteMenuEntry, TWeakPtr<FTabSpawnerEntry>(InTabMenuEntry)),
+		FIsActionChecked::CreateSP(InTabMenuEntry.ToSharedRef(), &FTabSpawnerEntry::IsSoleTabInstanceSpawned)
+		);
+}
+
 void FTabManager::InvokeTabForMenu( FName TabId )
 {
 	TryInvokeTab(TabId);
diff --git a/Engine/Source/Runtime/Slate/Public/Framework/Docking/TabManager.h b/Engine/Source/Runtime/Slate/Public/Framework/Docking/TabManager.h
--- a/Engine/Source/Runtime/Slate/Public/Framework/Docking/TabManager.h
+++ b/Engine/Source/Runtime/Slate/Public/Framework/Docking/TabManager.h
@@ -8,6 +8,7 @@
 #include "Textures/SlateIcon.h"
 #include "Widgets/SWindow.h"
 #include "Framework/Docking/WorkspaceItem.h"
+#include "Framework/Commands/UIAction.h"
 
 class FJsonObject;
 class FMenuBuilder;
@@ -889,6 +890,8 @@ class SLATE_API FTabManager : public TSharedFromThis<FTabManager>
 		/** Returns filter for additional control over available tabs */
 		TSharedRef<FBlacklistNames>& GetTabBlacklist();
 
+		FUIAction GetUIActionForTabSpawnerMenuEntry(TSharedPtr<FTabSpawnerEntry> InTabMenuEntry);
+
 	protected:
 		void InvokeTabForMenu( FName TabId );
 
