Fix assertion when creating a new partitioned level without having first opened the world browser.
#rnx

diff --git a/Engine/Source/Editor/LevelEditor/Private/LevelEditorActions.cpp b/Engine/Source/Editor/LevelEditor/Private/LevelEditorActions.cpp
--- a/Engine/Source/Editor/LevelEditor/Private/LevelEditorActions.cpp
+++ b/Engine/Source/Editor/LevelEditor/Private/LevelEditorActions.cpp
@@ -312,7 +312,7 @@ void FLevelEditorActionCallbacks::NewLevel()
 
 							WorldPartition->Initialize(World, FTransform::Identity);					
 
-							FWorldBrowserModule& WorldBrowserModule = FModuleManager::GetModuleChecked<FWorldBrowserModule>("WorldBrowser");
+							FWorldBrowserModule& WorldBrowserModule = FModuleManager::LoadModuleChecked<FWorldBrowserModule>("WorldBrowser");
 							WorldBrowserModule.OnBrowseWorld.Broadcast(World);
 						}
 					}
