Cleaned-up internal element data access

This avoids interfaces needing to know the direct structure of the element data, in case it were to change in the future. It also provides utils for accessing data from arrays of element handles.

#rnx

diff --git a/Engine/Source/Editor/LevelEditor/Private/Elements/Actor/ActorElementLevelEditorSelectionCustomization.cpp b/Engine/Source/Editor/LevelEditor/Private/Elements/Actor/ActorElementLevelEditorSelectionCustomization.cpp
--- a/Engine/Source/Editor/LevelEditor/Private/Elements/Actor/ActorElementLevelEditorSelectionCustomization.cpp
+++ b/Engine/Source/Editor/LevelEditor/Private/Elements/Actor/ActorElementLevelEditorSelectionCustomization.cpp
@@ -24,35 +24,34 @@ DEFINE_LOG_CATEGORY_STATIC(LogActorLevelEditorSelection, Log, All);
 
 bool FActorElementLevelEditorSelectionCustomization::CanSelectElement(const TTypedElement<UTypedElementSelectionInterface>& InElementSelectionHandle, const FTypedElementSelectionOptions& InSelectionOptions)
 {
-	return InElementSelectionHandle.IsDataOfType<FActorElementData>() && CanSelectActorElement(InElementSelectionHandle, InSelectionOptions);
+	return CanSelectActorElement(InElementSelectionHandle, InSelectionOptions);
 }
 
 bool FActorElementLevelEditorSelectionCustomization::CanDeselectElement(const TTypedElement<UTypedElementSelectionInterface>& InElementSelectionHandle, const FTypedElementSelectionOptions& InSelectionOptions)
 {
-	return InElementSelectionHandle.IsDataOfType<FActorElementData>() && CanDeselectActorElement(InElementSelectionHandle, InSelectionOptions);
+	return CanDeselectActorElement(InElementSelectionHandle, InSelectionOptions);
 }
 
 bool FActorElementLevelEditorSelectionCustomization::SelectElement(const TTypedElement<UTypedElementSelectionInterface>& InElementSelectionHandle, UTypedElementList* InSelectionSet, const FTypedElementSelectionOptions& InSelectionOptions)
 {
-	return InElementSelectionHandle.IsDataOfType<FActorElementData>() && SelectActorElement(InElementSelectionHandle, InSelectionSet, InSelectionOptions);
+	return SelectActorElement(InElementSelectionHandle, InSelectionSet, InSelectionOptions);
 }
 
 bool FActorElementLevelEditorSelectionCustomization::DeselectElement(const TTypedElement<UTypedElementSelectionInterface>& InElementSelectionHandle, UTypedElementList* InSelectionSet, const FTypedElementSelectionOptions& InSelectionOptions)
 {
-	return InElementSelectionHandle.IsDataOfType<FActorElementData>() && DeselectActorElement(InElementSelectionHandle, InSelectionSet, InSelectionOptions);
+	return DeselectActorElement(InElementSelectionHandle, InSelectionSet, InSelectionOptions);
 }
 
 bool FActorElementLevelEditorSelectionCustomization::AllowSelectionModifiers(const TTypedElement<UTypedElementSelectionInterface>& InElementSelectionHandle, const UTypedElementList* InSelectionSet)
 {
 	// Ctrl or Shift clicking an actor is the same as regular clicking when components are selected
-	return InElementSelectionHandle.IsDataOfType<FActorElementData>() && !UComponentElementSelectionInterface::HasSelectedComponents(InSelectionSet);
+	return !UComponentElementSelectionInterface::HasSelectedComponents(InSelectionSet);
 }
 
 FTypedElementHandle FActorElementLevelEditorSelectionCustomization::GetSelectionElement(const TTypedElement<UTypedElementSelectionInterface>& InElementSelectionHandle, const UTypedElementList* InCurrentSelection, const ETypedElementSelectionMethod InSelectionMethod)
 {
-	if (const FActorElementData* ActorData = InElementSelectionHandle.GetData<FActorElementData>())
+	if (AActor* ConsideredActor = ActorElementDataUtil::GetActorFromHandle(InElementSelectionHandle))
 	{
-		AActor* ConsideredActor = ActorData->Actor;
 		while (ConsideredActor->IsChildActor())
 		{
 			ConsideredActor = ConsideredActor->GetParentActor();
@@ -64,7 +63,7 @@ FTypedElementHandle FActorElementLevelEditorSelectionCustomization::GetSelection
 
 bool FActorElementLevelEditorSelectionCustomization::CanSelectActorElement(const TTypedElement<UTypedElementSelectionInterface>& InActorSelectionHandle, const FTypedElementSelectionOptions& InSelectionOptions)
 {
-	AActor* Actor = InActorSelectionHandle.GetDataChecked<FActorElementData>().Actor;
+	AActor* Actor = ActorElementDataUtil::GetActorFromHandleChecked(InActorSelectionHandle);
 
 	// Bail if global selection is locked, or this actor cannot be edited or selected
 	if (GEdSelectionLock || !Actor->IsEditable() || !Actor->IsSelectable())
@@ -110,7 +109,7 @@ bool FActorElementLevelEditorSelectionCustomization::CanSelectActorElement(const
 
 bool FActorElementLevelEditorSelectionCustomization::CanDeselectActorElement(const TTypedElement<UTypedElementSelectionInterface>& InActorSelectionHandle, const FTypedElementSelectionOptions& InSelectionOptions)
 {
-	AActor* Actor = InActorSelectionHandle.GetDataChecked<FActorElementData>().Actor;
+	AActor* Actor = ActorElementDataUtil::GetActorFromHandleChecked(InActorSelectionHandle);
 
 	// Bail if global selection is locked
 	if (GEdSelectionLock)
@@ -124,7 +123,7 @@ bool FActorElementLevelEditorSelectionCustomization::CanDeselectActorElement(con
 
 bool FActorElementLevelEditorSelectionCustomization::SelectActorElement(const TTypedElement<UTypedElementSelectionInterface>& InActorSelectionHandle, UTypedElementList* InSelectionSet, const FTypedElementSelectionOptions& InSelectionOptions)
 {
-	AActor* Actor = InActorSelectionHandle.GetDataChecked<FActorElementData>().Actor;
+	AActor* Actor = ActorElementDataUtil::GetActorFromHandleChecked(InActorSelectionHandle);
 
 	// Allow active modes to potentially handle the selection
 	// TODO: Should this pass through the selection set?
@@ -189,7 +188,7 @@ bool FActorElementLevelEditorSelectionCustomization::SelectActorElement(const TT
 
 bool FActorElementLevelEditorSelectionCustomization::DeselectActorElement(const TTypedElement<UTypedElementSelectionInterface>& InActorSelectionHandle, UTypedElementList* InSelectionSet, const FTypedElementSelectionOptions& InSelectionOptions)
 {
-	AActor* Actor = InActorSelectionHandle.GetDataChecked<FActorElementData>().Actor;
+	AActor* Actor = ActorElementDataUtil::GetActorFromHandleChecked(InActorSelectionHandle);
 
 	// Allow active modes to potentially handle the deselection
 	// TODO: Should this pass through the selection set?
diff --git a/Engine/Source/Editor/LevelEditor/Private/Elements/Actor/ActorElementLevelEditorViewportInteractionCustomization.cpp b/Engine/Source/Editor/LevelEditor/Private/Elements/Actor/ActorElementLevelEditorViewportInteractionCustomization.cpp
--- a/Engine/Source/Editor/LevelEditor/Private/Elements/Actor/ActorElementLevelEditorViewportInteractionCustomization.cpp
+++ b/Engine/Source/Editor/LevelEditor/Private/Elements/Actor/ActorElementLevelEditorViewportInteractionCustomization.cpp
@@ -22,7 +22,7 @@ FActorElementLevelEditorViewportInteractionCustomization::FActorElementLevelEdit
 
 void FActorElementLevelEditorViewportInteractionCustomization::GetElementsToMove(const TTypedElement<UTypedElementWorldInterface>& InElementWorldHandle, const ETypedElementViewportInteractionWorldType InWorldType, const UTypedElementSelectionSet* InSelectionSet, UTypedElementList* OutElementsToMove)
 {
-	AActor* Actor = InElementWorldHandle.GetDataChecked<FActorElementData>().Actor;
+	AActor* Actor = ActorElementDataUtil::GetActorFromHandleChecked(InElementWorldHandle);
 
 	if (UComponentElementSelectionInterface::HasSelectedComponents(InSelectionSet->GetElementList()))
 	{
@@ -39,7 +39,7 @@ void FActorElementLevelEditorViewportInteractionCustomization::GetElementsToMove
 
 void FActorElementLevelEditorViewportInteractionCustomization::GizmoManipulationStarted(const TTypedElement<UTypedElementWorldInterface>& InElementWorldHandle, const UE::Widget::EWidgetMode InWidgetMode)
 {
-	AActor* Actor = InElementWorldHandle.GetDataChecked<FActorElementData>().Actor;
+	AActor* Actor = ActorElementDataUtil::GetActorFromHandleChecked(InElementWorldHandle);
 
 	// Notify that this actor is beginning to move
 	GEditor->BroadcastBeginObjectMovement(*Actor);
@@ -55,7 +55,7 @@ void FActorElementLevelEditorViewportInteractionCustomization::GizmoManipulation
 
 void FActorElementLevelEditorViewportInteractionCustomization::GizmoManipulationDeltaUpdate(const TTypedElement<UTypedElementWorldInterface>& InElementWorldHandle, const UE::Widget::EWidgetMode InWidgetMode, const EAxisList::Type InDragAxis, const FInputDeviceState& InInputState, const FTransform& InDeltaTransform, const FVector& InPivotLocation)
 {
-	AActor* Actor = InElementWorldHandle.GetDataChecked<FActorElementData>().Actor;
+	AActor* Actor = ActorElementDataUtil::GetActorFromHandleChecked(InElementWorldHandle);
 
 	GEditor->NoteActorMovement();
 
@@ -86,7 +86,7 @@ void FActorElementLevelEditorViewportInteractionCustomization::GizmoManipulation
 
 void FActorElementLevelEditorViewportInteractionCustomization::GizmoManipulationStopped(const TTypedElement<UTypedElementWorldInterface>& InElementWorldHandle, const UE::Widget::EWidgetMode InWidgetMode)
 {
-	AActor* Actor = InElementWorldHandle.GetDataChecked<FActorElementData>().Actor;
+	AActor* Actor = ActorElementDataUtil::GetActorFromHandleChecked(InElementWorldHandle);
 
 	// Broadcast Post Edit change notification, we can't call PostEditChangeProperty directly on Actor or ActorComponent from here since it wasn't pair with a proper PreEditChange
 	if (FProperty* TransformProperty = FComponentElementLevelEditorViewportInteractionCustomization::GetEditTransformProperty(InWidgetMode))
@@ -101,12 +101,7 @@ void FActorElementLevelEditorViewportInteractionCustomization::GizmoManipulation
 
 void FActorElementLevelEditorViewportInteractionCustomization::PostGizmoManipulationStopped(TArrayView<const TTypedElement<UTypedElementWorldInterface>> InElementWorldHandles, const UE::Widget::EWidgetMode InWidgetMode)
 {
-	TArray<AActor*> MovedActors;
-	MovedActors.Reserve(InElementWorldHandles.Num());
-	for (const TTypedElement<UTypedElementWorldInterface>& ElementWorldHandle : InElementWorldHandles)
-	{
-		MovedActors.Add(ElementWorldHandle.GetDataChecked<FActorElementData>().Actor);
-	}
+	TArray<AActor*> MovedActors = ActorElementDataUtil::GetActorsFromHandlesChecked(InElementWorldHandles);
 	GEditor->BroadcastActorsMoved(MovedActors);
 }
 
diff --git a/Engine/Source/Editor/LevelEditor/Private/Elements/Component/ComponentElementLevelEditorSelectionCustomization.cpp b/Engine/Source/Editor/LevelEditor/Private/Elements/Component/ComponentElementLevelEditorSelectionCustomization.cpp
--- a/Engine/Source/Editor/LevelEditor/Private/Elements/Component/ComponentElementLevelEditorSelectionCustomization.cpp
+++ b/Engine/Source/Editor/LevelEditor/Private/Elements/Component/ComponentElementLevelEditorSelectionCustomization.cpp
@@ -19,30 +19,30 @@ DEFINE_LOG_CATEGORY_STATIC(LogComponentLevelEditorSelection, Log, All);
 
 bool FComponentElementLevelEditorSelectionCustomization::CanSelectElement(const TTypedElement<UTypedElementSelectionInterface>& InElementSelectionHandle, const FTypedElementSelectionOptions& InSelectionOptions)
 {
-	return InElementSelectionHandle.IsDataOfType<FComponentElementData>() && CanSelectComponentElement(InElementSelectionHandle, InSelectionOptions);
+	return CanSelectComponentElement(InElementSelectionHandle, InSelectionOptions);
 }
 
 bool FComponentElementLevelEditorSelectionCustomization::CanDeselectElement(const TTypedElement<UTypedElementSelectionInterface>& InElementSelectionHandle, const FTypedElementSelectionOptions& InSelectionOptions)
 {
-	return InElementSelectionHandle.IsDataOfType<FComponentElementData>() && CanDeselectComponentElement(InElementSelectionHandle, InSelectionOptions);
+	return CanDeselectComponentElement(InElementSelectionHandle, InSelectionOptions);
 }
 
 bool FComponentElementLevelEditorSelectionCustomization::SelectElement(const TTypedElement<UTypedElementSelectionInterface>& InElementSelectionHandle, UTypedElementList* InSelectionSet, const FTypedElementSelectionOptions& InSelectionOptions)
 {
-	return InElementSelectionHandle.IsDataOfType<FComponentElementData>() && SelectComponentElement(InElementSelectionHandle, InSelectionSet, InSelectionOptions);
+	return SelectComponentElement(InElementSelectionHandle, InSelectionSet, InSelectionOptions);
 }
 
 bool FComponentElementLevelEditorSelectionCustomization::DeselectElement(const TTypedElement<UTypedElementSelectionInterface>& InElementSelectionHandle, UTypedElementList* InSelectionSet, const FTypedElementSelectionOptions& InSelectionOptions)
 {
-	return InElementSelectionHandle.IsDataOfType<FComponentElementData>() && DeselectComponentElement(InElementSelectionHandle, InSelectionSet, InSelectionOptions);
+	return DeselectComponentElement(InElementSelectionHandle, InSelectionSet, InSelectionOptions);
 }
 
 FTypedElementHandle FComponentElementLevelEditorSelectionCustomization::GetSelectionElement(const TTypedElement<UTypedElementSelectionInterface>& InElementSelectionHandle, const UTypedElementList* InCurrentSelection, const ETypedElementSelectionMethod InSelectionMethod)
 {
-	if (const FComponentElementData* ComponentData = InElementSelectionHandle.GetData<FComponentElementData>())
+	if (const UActorComponent* Component = ComponentElementDataUtil::GetComponentFromHandle(InElementSelectionHandle))
 	{
-		const AActor* ConsideredActor = ComponentData->Component->GetOwner();
-		const USceneComponent* ConsideredComponent = Cast<USceneComponent>(ComponentData->Component);
+		const AActor* ConsideredActor = Component->GetOwner();
+		const USceneComponent* ConsideredComponent = Cast<USceneComponent>(Component);
 		if (ConsideredActor)
 		{
 			while (ConsideredActor->IsChildActor())
@@ -86,7 +86,7 @@ FTypedElementHandle FComponentElementLevelEditorSelectionCustomization::GetSelec
 
 bool FComponentElementLevelEditorSelectionCustomization::CanSelectComponentElement(const TTypedElement<UTypedElementSelectionInterface>& InComponentSelectionHandle, const FTypedElementSelectionOptions& InSelectionOptions)
 {
-	UActorComponent* Component = InComponentSelectionHandle.GetDataChecked<FComponentElementData>().Component;
+	const UActorComponent* Component = ComponentElementDataUtil::GetComponentFromHandleChecked(InComponentSelectionHandle);
 
 	// Bail if global selection is locked
 	return !GEdSelectionLock;
@@ -94,7 +94,7 @@ bool FComponentElementLevelEditorSelectionCustomization::CanSelectComponentEleme
 
 bool FComponentElementLevelEditorSelectionCustomization::CanDeselectComponentElement(const TTypedElement<UTypedElementSelectionInterface>& InComponentSelectionHandle, const FTypedElementSelectionOptions& InSelectionOptions)
 {
-	UActorComponent* Component = InComponentSelectionHandle.GetDataChecked<FComponentElementData>().Component;
+	const UActorComponent* Component = ComponentElementDataUtil::GetComponentFromHandleChecked(InComponentSelectionHandle);
 
 	// Bail if global selection is locked
 	return !GEdSelectionLock;
@@ -102,7 +102,7 @@ bool FComponentElementLevelEditorSelectionCustomization::CanDeselectComponentEle
 
 bool FComponentElementLevelEditorSelectionCustomization::SelectComponentElement(const TTypedElement<UTypedElementSelectionInterface>& InComponentSelectionHandle, UTypedElementList* InSelectionSet, const FTypedElementSelectionOptions& InSelectionOptions)
 {
-	UActorComponent* Component = InComponentSelectionHandle.GetDataChecked<FComponentElementData>().Component;
+	UActorComponent* Component = ComponentElementDataUtil::GetComponentFromHandleChecked(InComponentSelectionHandle);
 
 	if (!InComponentSelectionHandle.SelectElement(InSelectionSet, InSelectionOptions))
 	{
@@ -141,7 +141,7 @@ bool FComponentElementLevelEditorSelectionCustomization::SelectComponentElement(
 
 bool FComponentElementLevelEditorSelectionCustomization::DeselectComponentElement(const TTypedElement<UTypedElementSelectionInterface>& InComponentSelectionHandle, UTypedElementList* InSelectionSet, const FTypedElementSelectionOptions& InSelectionOptions)
 {
-	UActorComponent* Component = InComponentSelectionHandle.GetDataChecked<FComponentElementData>().Component;
+	UActorComponent* Component = ComponentElementDataUtil::GetComponentFromHandleChecked(InComponentSelectionHandle);
 
 	if (!InComponentSelectionHandle.DeselectElement(InSelectionSet, InSelectionOptions))
 	{
diff --git a/Engine/Source/Editor/LevelEditor/Private/Elements/Component/ComponentElementLevelEditorViewportInteractionCustomization.cpp b/Engine/Source/Editor/LevelEditor/Private/Elements/Component/ComponentElementLevelEditorViewportInteractionCustomization.cpp
--- a/Engine/Source/Editor/LevelEditor/Private/Elements/Component/ComponentElementLevelEditorViewportInteractionCustomization.cpp
+++ b/Engine/Source/Editor/LevelEditor/Private/Elements/Component/ComponentElementLevelEditorViewportInteractionCustomization.cpp
@@ -19,7 +19,7 @@ FComponentElementLevelEditorViewportInteractionCustomization::FComponentElementL
 
 void FComponentElementLevelEditorViewportInteractionCustomization::GetElementsToMove(const TTypedElement<UTypedElementWorldInterface>& InElementWorldHandle, const ETypedElementViewportInteractionWorldType InWorldType, const UTypedElementSelectionSet* InSelectionSet, UTypedElementList* OutElementsToMove)
 {
-	UActorComponent* Component = InElementWorldHandle.GetDataChecked<FComponentElementData>().Component;
+	UActorComponent* Component = ComponentElementDataUtil::GetComponentFromHandleChecked(InElementWorldHandle);
 
 	if (USceneComponent* SceneComponent = Cast<USceneComponent>(Component))
 	{
@@ -55,7 +55,7 @@ void FComponentElementLevelEditorViewportInteractionCustomization::GetElementsTo
 
 void FComponentElementLevelEditorViewportInteractionCustomization::GizmoManipulationStarted(const TTypedElement<UTypedElementWorldInterface>& InElementWorldHandle, const UE::Widget::EWidgetMode InWidgetMode)
 {
-	UActorComponent* Component = InElementWorldHandle.GetDataChecked<FComponentElementData>().Component;
+	UActorComponent* Component = ComponentElementDataUtil::GetComponentFromHandleChecked(InElementWorldHandle);
 
 	if (USceneComponent* SceneComponent = Cast<USceneComponent>(Component))
 	{
@@ -74,7 +74,7 @@ void FComponentElementLevelEditorViewportInteractionCustomization::GizmoManipula
 
 void FComponentElementLevelEditorViewportInteractionCustomization::GizmoManipulationDeltaUpdate(const TTypedElement<UTypedElementWorldInterface>& InElementWorldHandle, const UE::Widget::EWidgetMode InWidgetMode, const EAxisList::Type InDragAxis, const FInputDeviceState& InInputState, const FTransform& InDeltaTransform, const FVector& InPivotLocation)
 {
-	UActorComponent* Component = InElementWorldHandle.GetDataChecked<FComponentElementData>().Component;
+	UActorComponent* Component = ComponentElementDataUtil::GetComponentFromHandleChecked(InElementWorldHandle);
 
 	if (USceneComponent* SceneComponent = Cast<USceneComponent>(Component))
 	{
@@ -109,7 +109,7 @@ void FComponentElementLevelEditorViewportInteractionCustomization::GizmoManipula
 
 void FComponentElementLevelEditorViewportInteractionCustomization::GizmoManipulationStopped(const TTypedElement<UTypedElementWorldInterface>& InElementWorldHandle, const UE::Widget::EWidgetMode InWidgetMode)
 {
-	UActorComponent* Component = InElementWorldHandle.GetDataChecked<FComponentElementData>().Component;
+	UActorComponent* Component = ComponentElementDataUtil::GetComponentFromHandleChecked(InElementWorldHandle);
 
 	if (USceneComponent* SceneComponent = Cast<USceneComponent>(Component))
 	{
diff --git a/Engine/Source/Editor/UnrealEd/Private/Elements/Actor/ActorElementEditorSelectionInterface.cpp b/Engine/Source/Editor/UnrealEd/Private/Elements/Actor/ActorElementEditorSelectionInterface.cpp
--- a/Engine/Source/Editor/UnrealEd/Private/Elements/Actor/ActorElementEditorSelectionInterface.cpp
+++ b/Engine/Source/Editor/UnrealEd/Private/Elements/Actor/ActorElementEditorSelectionInterface.cpp
@@ -10,22 +10,21 @@
 
 bool UActorElementEditorSelectionInterface::IsElementSelected(const FTypedElementHandle& InElementHandle, const UTypedElementList* InSelectionSet, const FTypedElementIsSelectedOptions& InSelectionOptions)
 {
-	const FActorElementData* ActorData = InElementHandle.GetData<FActorElementData>();
-	return ActorData && IsActorSelected(ActorData->Actor, InSelectionSet, InSelectionOptions);
+	const AActor* Actor = ActorElementDataUtil::GetActorFromHandle(InElementHandle);
+	return Actor && IsActorSelected(Actor, InSelectionSet, InSelectionOptions);
 }
 
 bool UActorElementEditorSelectionInterface::ShouldPreventTransactions(const FTypedElementHandle& InElementHandle)
 {
-	const FActorElementData* ActorData = InElementHandle.GetData<FActorElementData>();
-	return ActorData && UObjectElementEditorSelectionInterface::ShouldObjectPreventTransactions(ActorData->Actor);
+	const AActor* Actor = ActorElementDataUtil::GetActorFromHandle(InElementHandle);
+	return Actor && UObjectElementEditorSelectionInterface::ShouldObjectPreventTransactions(Actor);
 }
 
 void UActorElementEditorSelectionInterface::WriteTransactedElement(const FTypedElementHandle& InElementHandle, FArchive& InArchive)
 {
-	const FActorElementData* ActorData = InElementHandle.GetData<FActorElementData>();
-	if (ActorData)
+	if (const AActor* Actor = ActorElementDataUtil::GetActorFromHandle(InElementHandle))
 	{
-		UObjectElementEditorSelectionInterface::WriteTransactedObject(ActorData->Actor, InArchive);
+		UObjectElementEditorSelectionInterface::WriteTransactedObject(Actor, InArchive);
 	}
 }
 
diff --git a/Engine/Source/Editor/UnrealEd/Private/Elements/Actor/ActorElementEditorViewportInteractionCustomization.cpp b/Engine/Source/Editor/UnrealEd/Private/Elements/Actor/ActorElementEditorViewportInteractionCustomization.cpp
--- a/Engine/Source/Editor/UnrealEd/Private/Elements/Actor/ActorElementEditorViewportInteractionCustomization.cpp
+++ b/Engine/Source/Editor/UnrealEd/Private/Elements/Actor/ActorElementEditorViewportInteractionCustomization.cpp
@@ -16,7 +16,7 @@ bool FActorElementEditorViewportInteractionCustomization::GetGizmoPivotLocation(
 
 void FActorElementEditorViewportInteractionCustomization::GizmoManipulationDeltaUpdate(const TTypedElement<UTypedElementWorldInterface>& InElementWorldHandle, const UE::Widget::EWidgetMode InWidgetMode, const EAxisList::Type InDragAxis, const FInputDeviceState& InInputState, const FTransform& InDeltaTransform, const FVector& InPivotLocation)
 {
-	AActor* Actor = InElementWorldHandle.GetDataChecked<FActorElementData>().Actor;
+	AActor* Actor = ActorElementDataUtil::GetActorFromHandleChecked(InElementWorldHandle);
 
 	const FVector DeltaTranslation = InDeltaTransform.GetTranslation();
 	const FRotator DeltaRotation = InDeltaTransform.Rotator();
diff --git a/Engine/Source/Editor/UnrealEd/Private/Elements/Component/ComponentElementEditorSelectionInterface.cpp b/Engine/Source/Editor/UnrealEd/Private/Elements/Component/ComponentElementEditorSelectionInterface.cpp
--- a/Engine/Source/Editor/UnrealEd/Private/Elements/Component/ComponentElementEditorSelectionInterface.cpp
+++ b/Engine/Source/Editor/UnrealEd/Private/Elements/Component/ComponentElementEditorSelectionInterface.cpp
@@ -11,22 +11,21 @@
 
 bool UComponentElementEditorSelectionInterface::IsElementSelected(const FTypedElementHandle& InElementHandle, const UTypedElementList* InSelectionSet, const FTypedElementIsSelectedOptions& InSelectionOptions)
 {
-	const FComponentElementData* ComponentData = InElementHandle.GetData<FComponentElementData>();
-	return ComponentData && IsComponentSelected(ComponentData->Component, InSelectionSet, InSelectionOptions);
+	const UActorComponent* Component = ComponentElementDataUtil::GetComponentFromHandleChecked(InElementHandle);
+	return Component && IsComponentSelected(Component, InSelectionSet, InSelectionOptions);
 }
 
 bool UComponentElementEditorSelectionInterface::ShouldPreventTransactions(const FTypedElementHandle& InElementHandle)
 {
-	const FComponentElementData* ComponentData = InElementHandle.GetData<FComponentElementData>();
-	return ComponentData && UObjectElementEditorSelectionInterface::ShouldObjectPreventTransactions(ComponentData->Component);
+	const UActorComponent* Component = ComponentElementDataUtil::GetComponentFromHandleChecked(InElementHandle);
+	return Component && UObjectElementEditorSelectionInterface::ShouldObjectPreventTransactions(Component);
 }
 
 void UComponentElementEditorSelectionInterface::WriteTransactedElement(const FTypedElementHandle& InElementHandle, FArchive& InArchive)
 {
-	const FComponentElementData* ComponentData = InElementHandle.GetData<FComponentElementData>();
-	if (ComponentData)
+	if (const UActorComponent* Component = ComponentElementDataUtil::GetComponentFromHandleChecked(InElementHandle))
 	{
-		UObjectElementEditorSelectionInterface::WriteTransactedObject(ComponentData->Component, InArchive);
+		UObjectElementEditorSelectionInterface::WriteTransactedObject(Component, InArchive);
 	}
 }
 
diff --git a/Engine/Source/Editor/UnrealEd/Private/Elements/Component/ComponentElementEditorViewportInteractionCustomization.cpp b/Engine/Source/Editor/UnrealEd/Private/Elements/Component/ComponentElementEditorViewportInteractionCustomization.cpp
--- a/Engine/Source/Editor/UnrealEd/Private/Elements/Component/ComponentElementEditorViewportInteractionCustomization.cpp
+++ b/Engine/Source/Editor/UnrealEd/Private/Elements/Component/ComponentElementEditorViewportInteractionCustomization.cpp
@@ -10,9 +10,9 @@
 
 bool FComponentElementEditorViewportInteractionCustomization::GetGizmoPivotLocation(const TTypedElement<UTypedElementWorldInterface>& InElementWorldHandle, const UE::Widget::EWidgetMode InWidgetMode, FVector& OutPivotLocation)
 {
-	UActorComponent* Component = InElementWorldHandle.GetDataChecked<FComponentElementData>().Component;
+	const UActorComponent* Component = ComponentElementDataUtil::GetComponentFromHandleChecked(InElementWorldHandle);
 
-	if (USceneComponent* SceneComponent = Cast<USceneComponent>(Component))
+	if (const USceneComponent* SceneComponent = Cast<USceneComponent>(Component))
 	{
 		// If necessary, transform the editor pivot location to be relative to the component's parent
 		const bool bIsRootComponent = SceneComponent->GetOwner()->GetRootComponent() == SceneComponent;
@@ -25,7 +25,7 @@ bool FComponentElementEditorViewportInteractionCustomization::GetGizmoPivotLocat
 
 void FComponentElementEditorViewportInteractionCustomization::GizmoManipulationDeltaUpdate(const TTypedElement<UTypedElementWorldInterface>& InElementWorldHandle, const UE::Widget::EWidgetMode InWidgetMode, const EAxisList::Type InDragAxis, const FInputDeviceState& InInputState, const FTransform& InDeltaTransform, const FVector& InPivotLocation)
 {
-	UActorComponent* Component = InElementWorldHandle.GetDataChecked<FComponentElementData>().Component;
+	UActorComponent* Component = ComponentElementDataUtil::GetComponentFromHandleChecked(InElementWorldHandle);
 
 	if (USceneComponent* SceneComponent = Cast<USceneComponent>(Component))
 	{
diff --git a/Engine/Source/Editor/UnrealEd/Private/Elements/Object/ObjectElementEditorSelectionCustomization.cpp b/Engine/Source/Editor/UnrealEd/Private/Elements/Object/ObjectElementEditorSelectionCustomization.cpp
--- a/Engine/Source/Editor/UnrealEd/Private/Elements/Object/ObjectElementEditorSelectionCustomization.cpp
+++ b/Engine/Source/Editor/UnrealEd/Private/Elements/Object/ObjectElementEditorSelectionCustomization.cpp
@@ -6,11 +6,11 @@
 
 bool FObjectElementEditorSelectionCustomization::SelectElement(const TTypedElement<UTypedElementSelectionInterface>& InElementSelectionHandle, UTypedElementList* InSelectionSet, const FTypedElementSelectionOptions& InSelectionOptions)
 {
-	if (const FObjectElementData* ObjectData = InElementSelectionHandle.GetData<FObjectElementData>())
+	if (const UObject* Object = ObjectElementDataUtil::GetObjectFromHandle(InElementSelectionHandle))
 	{
 		if (InElementSelectionHandle.SelectElement(InSelectionSet, InSelectionOptions))
 		{
-			GSelectedObjectAnnotation.Set(ObjectData->Object);
+			GSelectedObjectAnnotation.Set(Object);
 			return true;
 		}
 	}
@@ -20,11 +20,11 @@ bool FObjectElementEditorSelectionCustomization::SelectElement(const TTypedEleme
 
 bool FObjectElementEditorSelectionCustomization::DeselectElement(const TTypedElement<UTypedElementSelectionInterface>& InElementSelectionHandle, UTypedElementList* InSelectionSet, const FTypedElementSelectionOptions& InSelectionOptions)
 {
-	if (const FObjectElementData* ObjectData = InElementSelectionHandle.GetData<FObjectElementData>())
+	if (const UObject* Object = ObjectElementDataUtil::GetObjectFromHandle(InElementSelectionHandle))
 	{
 		if (InElementSelectionHandle.DeselectElement(InSelectionSet, InSelectionOptions))
 		{
-			GSelectedObjectAnnotation.Clear(ObjectData->Object);
+			GSelectedObjectAnnotation.Clear(Object);
 			return true;
 		}
 	}
diff --git a/Engine/Source/Editor/UnrealEd/Private/Elements/Object/ObjectElementEditorSelectionInterface.cpp b/Engine/Source/Editor/UnrealEd/Private/Elements/Object/ObjectElementEditorSelectionInterface.cpp
--- a/Engine/Source/Editor/UnrealEd/Private/Elements/Object/ObjectElementEditorSelectionInterface.cpp
+++ b/Engine/Source/Editor/UnrealEd/Private/Elements/Object/ObjectElementEditorSelectionInterface.cpp
@@ -8,16 +8,15 @@
 
 bool UObjectElementEditorSelectionInterface::ShouldPreventTransactions(const FTypedElementHandle& InElementHandle)
 {
-	const FObjectElementData* ObjectData = InElementHandle.GetData<FObjectElementData>();
-	return ObjectData && ShouldObjectPreventTransactions(ObjectData->Object);
+	const UObject* Object = ObjectElementDataUtil::GetObjectFromHandle(InElementHandle);
+	return Object && ShouldObjectPreventTransactions(Object);
 }
 
 void UObjectElementEditorSelectionInterface::WriteTransactedElement(const FTypedElementHandle& InElementHandle, FArchive& InArchive)
 {
-	const FObjectElementData* ObjectData = InElementHandle.GetData<FObjectElementData>();
-	if (ObjectData)
+	if (const UObject* Object = ObjectElementDataUtil::GetObjectFromHandle(InElementHandle))
 	{
-		WriteTransactedObject(ObjectData->Object, InArchive);
+		WriteTransactedObject(Object, InArchive);
 	}
 }
 
diff --git a/Engine/Source/Runtime/Engine/Private/Elements/Actor/ActorElementData.cpp b/Engine/Source/Runtime/Engine/Private/Elements/Actor/ActorElementData.cpp
--- a/Engine/Source/Runtime/Engine/Private/Elements/Actor/ActorElementData.cpp
+++ b/Engine/Source/Runtime/Engine/Private/Elements/Actor/ActorElementData.cpp
@@ -1,5 +1,23 @@
 // Copyright Epic Games, Inc. All Rights Reserved.
 
 #include "Elements/Actor/ActorElementData.h"
+#include "Elements/Framework/TypedElementHandle.h"
 
 UE_DEFINE_TYPED_ELEMENT_DATA_RTTI(FActorElementData);
+
+namespace ActorElementDataUtil
+{
+
+AActor* GetActorFromHandle(const FTypedElementHandle& InHandle, const bool bSilent)
+{
+	const FActorElementData* ActorElement = InHandle.GetData<FActorElementData>(bSilent);
+	return ActorElement ? ActorElement->Actor : nullptr;
+}
+
+AActor* GetActorFromHandleChecked(const FTypedElementHandle& InHandle)
+{
+	const FActorElementData& ActorElement = InHandle.GetDataChecked<FActorElementData>();
+	return ActorElement.Actor;
+}
+
+} // namespace ActorElementDataUtil
diff --git a/Engine/Source/Runtime/Engine/Private/Elements/Actor/ActorElementObjectInterface.cpp b/Engine/Source/Runtime/Engine/Private/Elements/Actor/ActorElementObjectInterface.cpp
--- a/Engine/Source/Runtime/Engine/Private/Elements/Actor/ActorElementObjectInterface.cpp
+++ b/Engine/Source/Runtime/Engine/Private/Elements/Actor/ActorElementObjectInterface.cpp
@@ -6,6 +6,5 @@
 
 UObject* UActorElementObjectInterface::GetObject(const FTypedElementHandle& InElementHandle)
 {
-	const FActorElementData* ActorData = InElementHandle.GetData<FActorElementData>();
-	return ActorData ? ActorData->Actor : nullptr;
+	return ActorElementDataUtil::GetActorFromHandle(InElementHandle);
 }
diff --git a/Engine/Source/Runtime/Engine/Private/Elements/Actor/ActorElementSelectionInterface.cpp b/Engine/Source/Runtime/Engine/Private/Elements/Actor/ActorElementSelectionInterface.cpp
--- a/Engine/Source/Runtime/Engine/Private/Elements/Actor/ActorElementSelectionInterface.cpp
+++ b/Engine/Source/Runtime/Engine/Private/Elements/Actor/ActorElementSelectionInterface.cpp
@@ -10,7 +10,7 @@ int32 UActorElementSelectionInterface::GetNumSelectedActors(const UTypedElementL
 	int32 NumSelected = 0;
 	InCurrentSelection->ForEachElementHandle([&NumSelected](const FTypedElementHandle& InSelectedElement)
 	{
-		if (InSelectedElement.GetData<FActorElementData>(/*bSilent*/true))
+		if (ActorElementDataUtil::GetActorFromHandle(InSelectedElement, /*bSilent*/true))
 		{
 			++NumSelected;
 		}
@@ -24,7 +24,7 @@ bool UActorElementSelectionInterface::HasSelectedActors(const UTypedElementList*
 	bool bHasSelectedActors = false;
 	InCurrentSelection->ForEachElementHandle([&bHasSelectedActors](const FTypedElementHandle& InSelectedElement)
 	{
-		bHasSelectedActors = InSelectedElement.GetData<FActorElementData>(/*bSilent*/true) != nullptr;
+		bHasSelectedActors = ActorElementDataUtil::GetActorFromHandle(InSelectedElement, /*bSilent*/true) != nullptr;
 		return !bHasSelectedActors;
 	});
 	return bHasSelectedActors;
diff --git a/Engine/Source/Runtime/Engine/Private/Elements/Actor/ActorElementWorldInterface.cpp b/Engine/Source/Runtime/Engine/Private/Elements/Actor/ActorElementWorldInterface.cpp
--- a/Engine/Source/Runtime/Engine/Private/Elements/Actor/ActorElementWorldInterface.cpp
+++ b/Engine/Source/Runtime/Engine/Private/Elements/Actor/ActorElementWorldInterface.cpp
@@ -6,16 +6,16 @@
 
 UWorld* UActorElementWorldInterface::GetOwnerWorld(const FTypedElementHandle& InElementHandle)
 {
-	const FActorElementData* ActorData = InElementHandle.GetData<FActorElementData>();
-	return ActorData ? ActorData->Actor->GetWorld() : nullptr;
+	const AActor* Actor = ActorElementDataUtil::GetActorFromHandle(InElementHandle);
+	return Actor ? Actor->GetWorld() : nullptr;
 }
 
 bool UActorElementWorldInterface::GetWorldBounds(const FTypedElementHandle& InElementHandle, FBoxSphereBounds& OutBounds)
 {
-	if (const FActorElementData* ActorData = InElementHandle.GetData<FActorElementData>())
+	if (const AActor* Actor = ActorElementDataUtil::GetActorFromHandle(InElementHandle))
 	{
 		// TODO: This was taken from FActorOrComponent, but AActor has a function to calculate bounds too...
-		OutBounds = ActorData->Actor->GetRootComponent()->Bounds;
+		OutBounds = Actor->GetRootComponent()->Bounds;
 		return true;
 	}
 
@@ -24,9 +24,9 @@ bool UActorElementWorldInterface::GetWorldBounds(const FTypedElementHandle& InEl
 
 bool UActorElementWorldInterface::GetWorldTransform(const FTypedElementHandle& InElementHandle, FTransform& OutTransform)
 {
-	if (const FActorElementData* ActorData = InElementHandle.GetData<FActorElementData>())
+	if (const AActor* Actor = ActorElementDataUtil::GetActorFromHandle(InElementHandle))
 	{
-		OutTransform = ActorData->Actor->GetActorTransform();
+		OutTransform = Actor->GetActorTransform();
 		return true;
 	}
 
@@ -35,10 +35,10 @@ bool UActorElementWorldInterface::GetWorldTransform(const FTypedElementHandle& I
 
 bool UActorElementWorldInterface::SetWorldTransform(const FTypedElementHandle& InElementHandle, const FTransform& InTransform)
 {
-	if (const FActorElementData* ActorData = InElementHandle.GetData<FActorElementData>())
+	if (AActor* Actor = ActorElementDataUtil::GetActorFromHandle(InElementHandle))
 	{
-		ActorData->Actor->Modify();
-		return ActorData->Actor->SetActorTransform(InTransform);
+		Actor->Modify();
+		return Actor->SetActorTransform(InTransform);
 	}
 
 	return false;
diff --git a/Engine/Source/Runtime/Engine/Private/Elements/Component/ComponentElementData.cpp b/Engine/Source/Runtime/Engine/Private/Elements/Component/ComponentElementData.cpp
--- a/Engine/Source/Runtime/Engine/Private/Elements/Component/ComponentElementData.cpp
+++ b/Engine/Source/Runtime/Engine/Private/Elements/Component/ComponentElementData.cpp
@@ -1,5 +1,23 @@
 // Copyright Epic Games, Inc. All Rights Reserved.
 
 #include "Elements/Component/ComponentElementData.h"
+#include "Elements/Framework/TypedElementHandle.h"
 
 UE_DEFINE_TYPED_ELEMENT_DATA_RTTI(FComponentElementData);
+
+namespace ComponentElementDataUtil
+{
+
+UActorComponent* GetComponentFromHandle(const FTypedElementHandle& InHandle, const bool bSilent)
+{
+	const FComponentElementData* ComponentElement = InHandle.GetData<FComponentElementData>(bSilent);
+	return ComponentElement ? ComponentElement->Component : nullptr;
+}
+
+UActorComponent* GetComponentFromHandleChecked(const FTypedElementHandle& InHandle)
+{
+	const FComponentElementData& ComponentElement = InHandle.GetDataChecked<FComponentElementData>();
+	return ComponentElement.Component;
+}
+
+} // namespace ComponentElementDataUtil
diff --git a/Engine/Source/Runtime/Engine/Private/Elements/Component/ComponentElementObjectInterface.cpp b/Engine/Source/Runtime/Engine/Private/Elements/Component/ComponentElementObjectInterface.cpp
--- a/Engine/Source/Runtime/Engine/Private/Elements/Component/ComponentElementObjectInterface.cpp
+++ b/Engine/Source/Runtime/Engine/Private/Elements/Component/ComponentElementObjectInterface.cpp
@@ -6,6 +6,5 @@
 
 UObject* UComponentElementObjectInterface::GetObject(const FTypedElementHandle& InElementHandle)
 {
-	const FComponentElementData* ComponentData = InElementHandle.GetData<FComponentElementData>();
-	return ComponentData ? ComponentData->Component : nullptr;
+	return ComponentElementDataUtil::GetComponentFromHandle(InElementHandle);
 }
diff --git a/Engine/Source/Runtime/Engine/Private/Elements/Component/ComponentElementSelectionInterface.cpp b/Engine/Source/Runtime/Engine/Private/Elements/Component/ComponentElementSelectionInterface.cpp
--- a/Engine/Source/Runtime/Engine/Private/Elements/Component/ComponentElementSelectionInterface.cpp
+++ b/Engine/Source/Runtime/Engine/Private/Elements/Component/ComponentElementSelectionInterface.cpp
@@ -10,7 +10,7 @@ int32 UComponentElementSelectionInterface::GetNumSelectedComponents(const UTyped
 	int32 NumSelected = 0;
 	InCurrentSelection->ForEachElementHandle([&NumSelected](const FTypedElementHandle& InSelectedElement)
 	{
-		if (InSelectedElement.GetData<FComponentElementData>(/*bSilent*/true))
+		if (ComponentElementDataUtil::GetComponentFromHandle(InSelectedElement, /*bSilent*/true))
 		{
 			++NumSelected;
 		}
@@ -24,7 +24,7 @@ bool UComponentElementSelectionInterface::HasSelectedComponents(const UTypedElem
 	bool bHasSelectedComponents = false;
 	InCurrentSelection->ForEachElementHandle([&bHasSelectedComponents](const FTypedElementHandle& InSelectedElement)
 	{
-		bHasSelectedComponents = InSelectedElement.GetData<FComponentElementData>(/*bSilent*/true) != nullptr;
+		bHasSelectedComponents = ComponentElementDataUtil::GetComponentFromHandle(InSelectedElement, /*bSilent*/true) != nullptr;
 		return !bHasSelectedComponents;
 	});
 	return bHasSelectedComponents;
diff --git a/Engine/Source/Runtime/Engine/Private/Elements/Component/ComponentElementWorldInterface.cpp b/Engine/Source/Runtime/Engine/Private/Elements/Component/ComponentElementWorldInterface.cpp
--- a/Engine/Source/Runtime/Engine/Private/Elements/Component/ComponentElementWorldInterface.cpp
+++ b/Engine/Source/Runtime/Engine/Private/Elements/Component/ComponentElementWorldInterface.cpp
@@ -7,21 +7,21 @@
 
 bool UComponentElementWorldInterface::CanEditElement(const FTypedElementHandle& InElementHandle)
 {
-	const FComponentElementData* ComponentData = InElementHandle.GetData<FComponentElementData>();
-	return ComponentData && ComponentData->Component->IsEditableWhenInherited();
+	const UActorComponent* Component = ComponentElementDataUtil::GetComponentFromHandleChecked(InElementHandle);
+	return Component && Component->IsEditableWhenInherited();
 }
 
 UWorld* UComponentElementWorldInterface::GetOwnerWorld(const FTypedElementHandle& InElementHandle)
 {
-	const FComponentElementData* ComponentData = InElementHandle.GetData<FComponentElementData>();
-	return ComponentData ? ComponentData->Component->GetWorld() : nullptr;
+	const UActorComponent* Component = ComponentElementDataUtil::GetComponentFromHandleChecked(InElementHandle);
+	return Component ? Component->GetWorld() : nullptr;
 }
 
 bool UComponentElementWorldInterface::GetWorldBounds(const FTypedElementHandle& InElementHandle, FBoxSphereBounds& OutBounds)
 {
-	if (const FComponentElementData* ComponentData = InElementHandle.GetData<FComponentElementData>())
+	if (const UActorComponent* Component = ComponentElementDataUtil::GetComponentFromHandleChecked(InElementHandle))
 	{
-		if (const USceneComponent* SceneComponent = Cast<USceneComponent>(ComponentData->Component))
+		if (const USceneComponent* SceneComponent = Cast<USceneComponent>(Component))
 		{
 			OutBounds = SceneComponent->Bounds;
 			return true;
@@ -33,9 +33,9 @@ bool UComponentElementWorldInterface::GetWorldBounds(const FTypedElementHandle&
 
 bool UComponentElementWorldInterface::GetWorldTransform(const FTypedElementHandle& InElementHandle, FTransform& OutTransform)
 {
-	if (const FComponentElementData* ComponentData = InElementHandle.GetData<FComponentElementData>())
+	if (const UActorComponent* Component = ComponentElementDataUtil::GetComponentFromHandleChecked(InElementHandle))
 	{
-		if (const USceneComponent* SceneComponent = Cast<USceneComponent>(ComponentData->Component))
+		if (const USceneComponent* SceneComponent = Cast<USceneComponent>(Component))
 		{
 			OutTransform = SceneComponent->GetComponentTransform();
 			return true;
@@ -47,9 +47,9 @@ bool UComponentElementWorldInterface::GetWorldTransform(const FTypedElementHandl
 
 bool UComponentElementWorldInterface::SetWorldTransform(const FTypedElementHandle& InElementHandle, const FTransform& InTransform)
 {
-	if (const FComponentElementData* ComponentData = InElementHandle.GetData<FComponentElementData>())
+	if (UActorComponent* Component = ComponentElementDataUtil::GetComponentFromHandleChecked(InElementHandle))
 	{
-		if (USceneComponent* SceneComponent = Cast<USceneComponent>(ComponentData->Component))
+		if (USceneComponent* SceneComponent = Cast<USceneComponent>(Component))
 		{
 			SceneComponent->Modify();
 			SceneComponent->SetWorldTransform(InTransform);
diff --git a/Engine/Source/Runtime/Engine/Private/Elements/Object/ObjectElementData.cpp b/Engine/Source/Runtime/Engine/Private/Elements/Object/ObjectElementData.cpp
--- a/Engine/Source/Runtime/Engine/Private/Elements/Object/ObjectElementData.cpp
+++ b/Engine/Source/Runtime/Engine/Private/Elements/Object/ObjectElementData.cpp
@@ -1,5 +1,23 @@
 // Copyright Epic Games, Inc. All Rights Reserved.
 
 #include "Elements/Object/ObjectElementData.h"
+#include "Elements/Framework/TypedElementHandle.h"
 
 UE_DEFINE_TYPED_ELEMENT_DATA_RTTI(FObjectElementData);
+
+namespace ObjectElementDataUtil
+{
+
+UObject* GetObjectFromHandle(const FTypedElementHandle& InHandle, const bool bSilent)
+{
+	const FObjectElementData* ObjectElement = InHandle.GetData<FObjectElementData>(bSilent);
+	return ObjectElement ? ObjectElement->Object : nullptr;
+}
+
+UObject* GetObjectFromHandleChecked(const FTypedElementHandle& InHandle)
+{
+	const FObjectElementData& ObjectElement = InHandle.GetDataChecked<FObjectElementData>();
+	return ObjectElement.Object;
+}
+
+} // namespace ObjectElementDataUtil
diff --git a/Engine/Source/Runtime/Engine/Private/Elements/Object/ObjectElementObjectInterface.cpp b/Engine/Source/Runtime/Engine/Private/Elements/Object/ObjectElementObjectInterface.cpp
--- a/Engine/Source/Runtime/Engine/Private/Elements/Object/ObjectElementObjectInterface.cpp
+++ b/Engine/Source/Runtime/Engine/Private/Elements/Object/ObjectElementObjectInterface.cpp
@@ -5,6 +5,5 @@
 
 UObject* UObjectElementObjectInterface::GetObject(const FTypedElementHandle& InElementHandle)
 {
-	const FObjectElementData* ObjectData = InElementHandle.GetData<FObjectElementData>();
-	return ObjectData ? ObjectData->Object : nullptr;
+	return ObjectElementDataUtil::GetObjectFromHandle(InElementHandle);
 }
diff --git a/Engine/Source/Runtime/Engine/Public/Elements/Actor/ActorElementData.h b/Engine/Source/Runtime/Engine/Public/Elements/Actor/ActorElementData.h
--- a/Engine/Source/Runtime/Engine/Public/Elements/Actor/ActorElementData.h
+++ b/Engine/Source/Runtime/Engine/Public/Elements/Actor/ActorElementData.h
@@ -5,6 +5,7 @@
 #include "Elements/Framework/TypedElementData.h"
 
 class AActor;
+struct FTypedElementHandle;
 
 /**
  * Element data that represents an Actor.
@@ -15,3 +16,63 @@ struct ENGINE_API FActorElementData
 
 	AActor* Actor = nullptr;
 };
+
+namespace ActorElementDataUtil
+{
+
+/**
+ * Attempt to get the actor from the given element handle.
+ * @note This is not typically something you'd want to use outside of data access within an interface implementation.
+ * @return The actor if the element handle contains FActorElementData, otherwise null.
+ */
+ENGINE_API AActor* GetActorFromHandle(const FTypedElementHandle& InHandle, const bool bSilent = false);
+
+/**
+ * Attempt to get the actor from the given element handle, asserting if the element handle doesn't contain FActorElementData.
+ * @note This is not typically something you'd want to use outside of data access within an interface implementation.
+ * @return The actor.
+ */
+ENGINE_API AActor* GetActorFromHandleChecked(const FTypedElementHandle& InHandle);
+
+/**
+ * Attempt to get the actors from the given element handles.
+ * @note This is not typically something you'd want to use outside of data access within an interface implementation.
+ * @return The actors of any element handles that contain FActorElementData, skipping any that don't.
+ */
+template <typename ElementHandleType>
+TArray<AActor*> GetActorsFromHandles(TArrayView<const ElementHandleType> InHandles, const bool bSilent = false)
+{
+	TArray<AActor*> Actors;
+	Actors.Reserve(InHandles.Num());
+
+	for (const FTypedElementHandle& Handle : InHandles)
+	{
+		if (AActor* Actor = GetActorFromHandle(Handle, bSilent))
+		{
+			Actors.Add(Actor);
+		}
+	}
+
+	return Actors;
+}
+
+/**
+ * Attempt to get the actors from the given element handles, asserting if any element handle doesn't contain FActorElementData.
+ * @note This is not typically something you'd want to use outside of data access within an interface implementation.
+ * @return The actors.
+ */
+template <typename ElementHandleType>
+TArray<AActor*> GetActorsFromHandlesChecked(TArrayView<const ElementHandleType> InHandles)
+{
+	TArray<AActor*> Actors;
+	Actors.Reserve(InHandles.Num());
+
+	for (const FTypedElementHandle& Handle : InHandles)
+	{
+		Actors.Add(GetActorFromHandleChecked(Handle));
+	}
+
+	return Actors;
+}
+
+} // namespace ActorElementDataUtil
diff --git a/Engine/Source/Runtime/Engine/Public/Elements/Component/ComponentElementData.h b/Engine/Source/Runtime/Engine/Public/Elements/Component/ComponentElementData.h
--- a/Engine/Source/Runtime/Engine/Public/Elements/Component/ComponentElementData.h
+++ b/Engine/Source/Runtime/Engine/Public/Elements/Component/ComponentElementData.h
@@ -5,6 +5,7 @@
 #include "Elements/Framework/TypedElementData.h"
 
 class UActorComponent;
+struct FTypedElementHandle;
 
 /**
  * Element data that represents an Actor Component.
@@ -15,3 +16,63 @@ struct ENGINE_API FComponentElementData
 
 	UActorComponent* Component = nullptr;
 };
+
+namespace ComponentElementDataUtil
+{
+
+/**
+ * Attempt to get the actor component from the given element handle.
+ * @note This is not typically something you'd want to use outside of data access within an interface implementation.
+ * @return The actor component if the element handle contains FComponentElementData, otherwise null.
+ */
+ENGINE_API UActorComponent* GetComponentFromHandle(const FTypedElementHandle& InHandle, const bool bSilent = false);
+
+/**
+ * Attempt to get the actor component from the given element handle, asserting if the element handle doesn't contain FComponentElementData.
+ * @note This is not typically something you'd want to use outside of data access within an interface implementation.
+ * @return The actor.
+ */
+ENGINE_API UActorComponent* GetComponentFromHandleChecked(const FTypedElementHandle& InHandle);
+
+/**
+ * Attempt to get the actor components from the given element handles.
+ * @note This is not typically something you'd want to use outside of data access within an interface implementation.
+ * @return The actor components of any element handles that contain FComponentElementData, skipping any that don't.
+ */
+template <typename ElementHandleType>
+TArray<UActorComponent*> GetComponentsFromHandles(TArrayView<const ElementHandleType> InHandles, const bool bSilent = false)
+{
+	TArray<UActorComponent*> Components;
+	Components.Reserve(InHandles.Num());
+
+	for (const FTypedElementHandle& Handle : InHandles)
+	{
+		if (UActorComponent* Component = GetComponentFromHandle(Handle, bSilent))
+		{
+			Components.Add(Component);
+		}
+	}
+
+	return Components;
+}
+
+/**
+ * Attempt to get the actor components from the given element handles, asserting if any element handle doesn't contain FComponentElementData.
+ * @note This is not typically something you'd want to use outside of data access within an interface implementation.
+ * @return The actor components.
+ */
+template <typename ElementHandleType>
+TArray<UActorComponent*> GetComponentsFromHandlesChecked(TArrayView<const ElementHandleType> InHandles)
+{
+	TArray<UActorComponent*> Components;
+	Components.Reserve(InHandles.Num());
+
+	for (const FTypedElementHandle& Handle : InHandles)
+	{
+		Components.Add(GetComponentFromHandleChecked(Handle));
+	}
+
+	return Components;
+}
+
+} // namespace ComponentElementDataUtil
diff --git a/Engine/Source/Runtime/Engine/Public/Elements/Object/ObjectElementData.h b/Engine/Source/Runtime/Engine/Public/Elements/Object/ObjectElementData.h
--- a/Engine/Source/Runtime/Engine/Public/Elements/Object/ObjectElementData.h
+++ b/Engine/Source/Runtime/Engine/Public/Elements/Object/ObjectElementData.h
@@ -5,6 +5,7 @@
 #include "Elements/Framework/TypedElementData.h"
 
 class UObject;
+struct FTypedElementHandle;
 
 /**
  * Element data that represents an Object.
@@ -15,3 +16,63 @@ struct ENGINE_API FObjectElementData
 
 	UObject* Object = nullptr;
 };
+
+namespace ObjectElementDataUtil
+{
+
+/**
+ * Attempt to get the object from the given element handle.
+ * @note This is not typically something you'd want to use outside of data access within an interface implementation.
+ * @return The object if the element handle contains FObjectElementData, otherwise null.
+ */
+ENGINE_API UObject* GetObjectFromHandle(const FTypedElementHandle& InHandle, const bool bSilent = false);
+
+/**
+ * Attempt to get the object from the given element handle, asserting if the element handle doesn't contain FObjectElementData.
+ * @note This is not typically something you'd want to use outside of data access within an interface implementation.
+ * @return The object.
+ */
+ENGINE_API UObject* GetObjectFromHandleChecked(const FTypedElementHandle& InHandle);
+
+/**
+ * Attempt to get the objects from the given element handles.
+ * @note This is not typically something you'd want to use outside of data access within an interface implementation.
+ * @return The objects of any element handles that contain FObjectElementData, skipping any that don't.
+ */
+template <typename ElementHandleType>
+TArray<UObject*> GetObjectsFromHandles(TArrayView<const ElementHandleType> InHandles, const bool bSilent = false)
+{
+	TArray<UObject*> Objects;
+	Objects.Reserve(InHandles.Num());
+
+	for (const FTypedElementHandle& Handle : InHandles)
+	{
+		if (UObject* Object = GetObjectFromHandle(Handle, bSilent))
+		{
+			Objects.Add(Object);
+		}
+	}
+
+	return Objects;
+}
+
+/**
+ * Attempt to get the objects from the given element handles, asserting if any element handle doesn't contain FObjectElementData.
+ * @note This is not typically something you'd want to use outside of data access within an interface implementation.
+ * @return The objects.
+ */
+template <typename ElementHandleType>
+TArray<UObject*> GetObjectsFromHandlesChecked(TArrayView<const ElementHandleType> InHandles)
+{
+	TArray<UObject*> Objects;
+	Objects.Reserve(InHandles.Num());
+
+	for (const FTypedElementHandle& Handle : InHandles)
+	{
+		Objects.Add(GetObjectFromHandleChecked(Handle));
+	}
+
+	return Objects;
+}
+
+} // namespace ObjectElementDataUtil
