Fix style of current level picker in the viewport bottom right corner

diff --git a/Engine/Source/Editor/LevelEditor/Private/SLevelViewport.cpp b/Engine/Source/Editor/LevelEditor/Private/SLevelViewport.cpp
--- a/Engine/Source/Editor/LevelEditor/Private/SLevelViewport.cpp
+++ b/Engine/Source/Editor/LevelEditor/Private/SLevelViewport.cpp
@@ -280,86 +280,81 @@ void SLevelViewport::ConstructViewportOverlayContent()
 	.HAlign( HAlign_Right )
 	.Padding( 5.0f )
 	[
-		SNew(SVerticalBox)
-		+ SVerticalBox::Slot()
-		.AutoHeight()
-		.Padding(2.0f, 1.0f, 2.0f, 1.0f)
+		SNew(SBorder)
+		.BorderImage(FAppStyle::Get().GetBrush("FloatingBorder"))
 		[
-			SNew(SHorizontalBox)
-			.Visibility(this, &SLevelViewport::GetCurrentScreenPercentageVisibility)
-			// Current screen percentage label
-			+ SHorizontalBox::Slot()
-			.AutoWidth()
+			SNew(SVerticalBox)
+			+ SVerticalBox::Slot()
+			.AutoHeight()
 			.Padding(2.0f, 1.0f, 2.0f, 1.0f)
 			[
-				SNew(STextBlock)
-				.Text(this, &SLevelViewport::GetCurrentScreenPercentageText, true)
-				.Font(FEditorStyle::GetFontStyle(TEXT("MenuItem.Font")))
-				.ShadowOffset(FVector2D(1, 1))
-			]
+				SNew(SHorizontalBox)
+				.Visibility(this, &SLevelViewport::GetCurrentScreenPercentageVisibility)
+				// Current screen percentage label
+				+ SHorizontalBox::Slot()
+				.AutoWidth()
+				.Padding(2.0f, 1.0f, 2.0f, 1.0f)
+				[
+					SNew(STextBlock)
+					.Text(LOCTEXT("ScreenPercentageLabel", "Screen Percentage"))
+					.ShadowOffset(FVector2D(1, 1))
+				]
 
-			// Current screen percentage
-			+ SHorizontalBox::Slot()
+				// Current screen percentage
+				+ SHorizontalBox::Slot()
 				.AutoWidth()
 				.Padding(4.0f, 1.0f, 2.0f, 1.0f)
 				[
 					SNew(STextBlock)
-					.Text(this, &SLevelViewport::GetCurrentScreenPercentageText, false)
-					.Font(FEditorStyle::GetFontStyle(TEXT("MenuItem.Font")))
-					.ColorAndOpacity(FLinearColor(0.4f, 1.0f, 1.0f))
+					.Text(this, &SLevelViewport::GetCurrentScreenPercentageText)
 					.ShadowOffset(FVector2D(1, 1))
 				]
-		]
-		// add feature level widget
-		+ SVerticalBox::Slot()
-		.AutoHeight()
-		.Padding(2.0f, 1.0f, 2.0f, 1.0f)
-		[
-			BuildFeatureLevelWidget()
-		]
-		+ SVerticalBox::Slot()
-		.AutoHeight()
-		.Padding(2.0f, 1.0f, 2.0f, 1.0f)
-		[
-			SNew(SVerticalBox)
-			.Visibility(this, &SLevelViewport::GetSelectedActorsCurrentLevelTextVisibility)
-			// Current level label
+			]
+			// add feature level widget
 			+ SVerticalBox::Slot()
 			.AutoHeight()
-			.Padding(6.0f, 1.0f, 2.0f, 1.0f)
+			.Padding(2.0f, 1.0f, 2.0f, 1.0f)
 			[
-				SNew(STextBlock)
-				.Text(this, &SLevelViewport::GetSelectedActorsCurrentLevelText, true)
-				.Font(FEditorStyle::GetFontStyle(TEXT("MenuItem.Font")))
-				.ShadowOffset(FVector2D(1, 1))
+				BuildFeatureLevelWidget()
 			]
-			// Current level
-			+ SVerticalBox::Slot()
+		/*	+ SVerticalBox::Slot()
 			.AutoHeight()
-			.Padding(6.0f, 1.0f, 2.0f, 1.0f)
+			.Padding(2.0f, 1.0f, 2.0f, 1.0f)
 			[
-				SNew(STextBlock)
-				.Text(this, &SLevelViewport::GetSelectedActorsCurrentLevelText, false)
-				.Font(FEditorStyle::GetFontStyle(TEXT("MenuItem.Font")))
-				.ShadowOffset(FVector2D(1, 1))
-			]
-		]
-		+ SVerticalBox::Slot()
-		.AutoHeight()
-		.Padding(2.0f, 1.0f, 2.0f, 1.0f)
-		[
-			SAssignNew(LevelMenuAnchor, SMenuAnchor)
-			.Placement(MenuPlacement_AboveAnchor)
-			.Visibility(this, &SLevelViewport::GetCurrentLevelTextVisibility)
+				SNew(SVerticalBox)
+				.Visibility(this, &SLevelViewport::GetSelectedActorsCurrentLevelTextVisibility)
+				// Current level label
+				+ SVerticalBox::Slot()
+				.AutoHeight()
+				.Padding(2.0f, 1.0f, 2.0f, 1.0f)
+				[
+					SNew(STextBlock)
+					.Text(this, &SLevelViewport::GetSelectedActorsCurrentLevelText, true)
+					.ShadowOffset(FVector2D(1, 1))
+				]
+				// Current level
+				+ SVerticalBox::Slot()
+				.AutoHeight()
+				.Padding(4.0f, 1.0f, 2.0f, 1.0f)
+				[
+					SNew(STextBlock)
+					.Text(this, &SLevelViewport::GetSelectedActorsCurrentLevelText, false)
+					.ShadowOffset(FVector2D(1, 1))
+				]
+			]*/
+			+ SVerticalBox::Slot()
+			.AutoHeight()
+			.Padding(2.0f, 1.0f, 2.0f, 1.0f)
 			[
-				SNew(SButton)
+				SNew(SComboButton)
+				.Cursor(EMouseCursor::Default)
 				// Allows users to drag with the mouse to select options after opening the menu */
-				.ClickMethod(EButtonClickMethod::MouseDown)
-				.ContentPadding(FMargin(5.0f, 2.0f))
 				.VAlign(VAlign_Center)
-				.ButtonStyle(FEditorStyle::Get(), "EditorViewportToolBar.MenuButton")
-				.OnClicked(this, &SLevelViewport::OnMenuClicked)
+				.ComboButtonStyle(FAppStyle::Get(), "SimpleComboButton")
 				.Visibility(this, &SLevelViewport::GetCurrentLevelButtonVisibility)
+				.OnGetMenuContent(this, &SLevelViewport::GenerateLevelMenu)
+				.OnMenuOpenChanged_Lambda([this](bool){	OnFloatingButtonClicked(); })
+				.ButtonContent()
 				[
 					SNew(SHorizontalBox)
 					.Visibility(this, &SLevelViewport::GetCurrentLevelTextVisibility)
@@ -370,7 +365,7 @@ void SLevelViewport::ConstructViewportOverlayContent()
 					[
 						SNew(STextBlock)
 						.Text(this, &SLevelViewport::GetCurrentLevelText, true)
-						.Font(FEditorStyle::GetFontStyle("EditorViewportToolBar.Font"))
+						.ShadowOffset(FVector2D(1,1))
 					]
 					// Current level
 					+ SHorizontalBox::Slot()
@@ -379,11 +374,10 @@ void SLevelViewport::ConstructViewportOverlayContent()
 					[
 						SNew(STextBlock)
 						.Text(this, &SLevelViewport::GetCurrentLevelText, false)
-						.Font(FEditorStyle::GetFontStyle("EditorViewportToolBar.Font"))
+						.ShadowOffset(FVector2D(1, 1))
 					]
 				]
 			]
-			.OnGetMenuContent(this, &SLevelViewport::GenerateLevelMenu)
 		]
 	];
 
@@ -416,21 +410,7 @@ TSharedRef<SWidget> SLevelViewport::GenerateLevelMenu() const
 	return LevelMenuBuilder.MakeWidget();
 }
 
-FReply SLevelViewport::OnMenuClicked()
-{
-	OnFloatingButtonClicked();
-	// If the menu button is clicked toggle the state of the menu anchor which will open or close the menu
-	if (LevelMenuAnchor->ShouldOpenDueToClick())
-	{
-		LevelMenuAnchor->SetIsOpen(true);
-	}
-	else
-	{
-		LevelMenuAnchor->SetIsOpen(false);
-	}
 
-	return FReply::Handled();
-}
 void SLevelViewport::ConstructLevelEditorViewportClient(FLevelEditorViewportInstanceSettings& ViewportInstanceSettings)
 {
 	if (!LevelViewportClient.IsValid())
@@ -3846,13 +3826,8 @@ FString SLevelViewport::GetDeviceProfileString( ) const
 	return DeviceProfile;
 }
 
-FText SLevelViewport::GetCurrentScreenPercentageText(bool bDrawOnlyLabel) const
+FText SLevelViewport::GetCurrentScreenPercentageText() const
 {
-	if (bDrawOnlyLabel)
-	{
-		return LOCTEXT("ScreenPercentageLabel", "Screen Percentage:");
-	}
-
 	return FText::FromString(FString::Printf(TEXT("%3d%%"), int32(GetLevelViewportClient().GetPreviewScreenPercentage())));
 }
 
@@ -3869,7 +3844,7 @@ FText SLevelViewport::GetCurrentLevelText( bool bDrawOnlyLabel ) const
 		{
 			if(bDrawOnlyLabel)
 			{
-				LabelName = LOCTEXT("CurrentLevelLabel", "Level:");
+				LabelName = LOCTEXT("CurrentLevelLabel", "Level");
 			}
 			else
 			{
@@ -3948,7 +3923,7 @@ FText SLevelViewport::GetSelectedActorsCurrentLevelText(bool bDrawOnlyLabel) con
 		{
 			if (bDrawOnlyLabel)
 			{
-				LabelName = LOCTEXT("SelectedActorsCurrentLevelLabel", "Selected Actor(s) in:");
+				LabelName = LOCTEXT("SelectedActorsCurrentLevelLabel", "Selected Actor(s) in");
 			}
 			else
 			{	
diff --git a/Engine/Source/Editor/LevelEditor/Public/SLevelViewport.h b/Engine/Source/Editor/LevelEditor/Public/SLevelViewport.h
--- a/Engine/Source/Editor/LevelEditor/Public/SLevelViewport.h
+++ b/Engine/Source/Editor/LevelEditor/Public/SLevelViewport.h
@@ -59,7 +59,6 @@ public:
 	void ConstructViewportOverlayContent();
 
 	TSharedRef<SWidget> GenerateLevelMenu() const;
-	FReply OnMenuClicked();
 
 	/**
 	 * Constructs the level editor viewport client
@@ -304,7 +303,7 @@ public:
 	FText GetCurrentLevelText( bool bDrawOnlyLabel ) const;
 
 	/** Called to get the screen percentage preview text */
-	FText GetCurrentScreenPercentageText(bool bDrawOnlyLabel) const;
+	FText GetCurrentScreenPercentageText() const;
 
 	/** @return The visibility of the current level text display */
 	virtual EVisibility GetCurrentLevelTextVisibility() const;
@@ -910,9 +909,6 @@ private:
 
 	/** Whether to show a full toolbar, or a compact one */
 	bool bShowFullToolbar;
-
-	TSharedPtr<class SMenuAnchor> LevelMenuAnchor;
-
 protected:
 	void LockActorInternal(AActor* NewActorToLock);
 
diff --git a/Engine/Source/Editor/UnrealEd/Private/SEditorViewport.cpp b/Engine/Source/Editor/UnrealEd/Private/SEditorViewport.cpp
--- a/Engine/Source/Editor/UnrealEd/Private/SEditorViewport.cpp
+++ b/Engine/Source/Editor/UnrealEd/Private/SEditorViewport.cpp
@@ -772,8 +772,7 @@ TSharedRef<SWidget> SEditorViewport::BuildFeatureLevelWidget() const
 		[
 			SNew(STextBlock)
 			.Text(this, &SEditorViewport::GetCurrentFeatureLevelPreviewText, true)
-		.Font(FEditorStyle::GetFontStyle(TEXT("MenuItem.Font")))
-		.ShadowOffset(FVector2D(1, 1))
+			.ShadowOffset(FVector2D(1, 1))
 		]
 		+ SHorizontalBox::Slot()
 		.AutoWidth()
@@ -781,8 +780,6 @@ TSharedRef<SWidget> SEditorViewport::BuildFeatureLevelWidget() const
 		[
 			SNew(STextBlock)
 			.Text(this, &SEditorViewport::GetCurrentFeatureLevelPreviewText, false)
-			.Font(FEditorStyle::GetFontStyle(TEXT("MenuItem.Font")))
-			.ColorAndOpacity(FLinearColor(0.4f, 1.0f, 1.0f))
 			.ShadowOffset(FVector2D(1, 1))
 		];
 
