Removed GLevelEditorModeTools use from typed element interface implementations

#rnx

diff --git a/Engine/Source/Editor/EditorFramework/Public/Elements/Framework/TypedElementAssetEditorToolkitHostMixin.h b/Engine/Source/Editor/EditorFramework/Public/Elements/Framework/TypedElementAssetEditorToolkitHostMixin.h
new file mode 100644
--- /dev/null
+++ b/Engine/Source/Editor/EditorFramework/Public/Elements/Framework/TypedElementAssetEditorToolkitHostMixin.h
@@ -0,0 +1,42 @@
+// Copyright Epic Games, Inc. All Rights Reserved.
+
+#pragma once
+
+#include "CoreMinimal.h"
+
+class IToolkitHost;
+
+/**
+ * Mixin for asset editor customizations that provides access to the toolkit host of the asset editor.
+ */
+class FTypedElementAssetEditorToolkitHostMixin
+{
+public:
+	/**
+	 * Get the toolkit host associated with our asset editor.
+	 */
+	const IToolkitHost* GetToolkitHost() const
+	{
+		return ToolkitHost;
+	}
+
+	/**
+	 * Get the toolkit host associated with our asset editor.
+	 */
+	IToolkitHost* GetMutableToolkitHost()
+	{
+		return ToolkitHost;
+	}
+
+	/**
+	 * Set the toolkit host associated with our asset editor.
+	 */
+	void SetToolkitHost(IToolkitHost* InToolkitHost)
+	{
+		ToolkitHost = InToolkitHost;
+	}
+
+private:
+	/** The toolkit host associated with our asset editor. */
+	IToolkitHost* ToolkitHost = nullptr;
+};
diff --git a/Engine/Source/Editor/LevelEditor/Private/Elements/Actor/ActorElementLevelEditorSelectionCustomization.cpp b/Engine/Source/Editor/LevelEditor/Private/Elements/Actor/ActorElementLevelEditorSelectionCustomization.cpp
--- a/Engine/Source/Editor/LevelEditor/Private/Elements/Actor/ActorElementLevelEditorSelectionCustomization.cpp
+++ b/Engine/Source/Editor/LevelEditor/Private/Elements/Actor/ActorElementLevelEditorSelectionCustomization.cpp
@@ -18,6 +18,7 @@
 #include "Editor/GroupActor.h"
 #include "ActorGroupingUtils.h"
 #include "Editor/UnrealEdEngine.h"
+#include "Toolkits/IToolkitHost.h"
 #include "Kismet2/ComponentEditorUtils.h"
 
 DEFINE_LOG_CATEGORY_STATIC(LogActorLevelEditorSelection, Log, All);
@@ -61,7 +62,7 @@ FTypedElementHandle FActorElementLevelEditorSelectionCustomization::GetSelection
 	return InElementSelectionHandle;
 }
 
-bool FActorElementLevelEditorSelectionCustomization::CanSelectActorElement(const TTypedElement<UTypedElementSelectionInterface>& InActorSelectionHandle, const FTypedElementSelectionOptions& InSelectionOptions)
+bool FActorElementLevelEditorSelectionCustomization::CanSelectActorElement(const TTypedElement<UTypedElementSelectionInterface>& InActorSelectionHandle, const FTypedElementSelectionOptions& InSelectionOptions) const
 {
 	AActor* Actor = ActorElementDataUtil::GetActorFromHandleChecked(InActorSelectionHandle);
 
@@ -103,11 +104,16 @@ bool FActorElementLevelEditorSelectionCustomization::CanSelectActorElement(const
 		return false;
 	}
 
-	// Allow active modes to determine whether the selection is allowed
-	return GLevelEditorModeTools().IsSelectionAllowed(Actor, /*bInSelected*/true);
+	if (const IToolkitHost* ToolkitHostPtr = GetToolkitHost())
+	{
+		// Allow active modes to determine whether the selection is allowed
+		return ToolkitHostPtr->GetEditorModeManager().IsSelectionAllowed(Actor, /*bInSelected*/true);
+	}
+
+	return true;
 }
 
-bool FActorElementLevelEditorSelectionCustomization::CanDeselectActorElement(const TTypedElement<UTypedElementSelectionInterface>& InActorSelectionHandle, const FTypedElementSelectionOptions& InSelectionOptions)
+bool FActorElementLevelEditorSelectionCustomization::CanDeselectActorElement(const TTypedElement<UTypedElementSelectionInterface>& InActorSelectionHandle, const FTypedElementSelectionOptions& InSelectionOptions) const
 {
 	AActor* Actor = ActorElementDataUtil::GetActorFromHandleChecked(InActorSelectionHandle);
 
@@ -117,19 +123,27 @@ bool FActorElementLevelEditorSelectionCustomization::CanDeselectActorElement(con
 		return false;
 	}
 
-	// Allow active modes to determine whether the deselection is allowed
-	return GLevelEditorModeTools().IsSelectionAllowed(Actor, /*bInSelected*/false);
+	if (const IToolkitHost* ToolkitHostPtr = GetToolkitHost())
+	{
+		// Allow active modes to determine whether the deselection is allowed
+		return ToolkitHostPtr->GetEditorModeManager().IsSelectionAllowed(Actor, /*bInSelected*/false);
+	}
+
+	return true;
 }
 
 bool FActorElementLevelEditorSelectionCustomization::SelectActorElement(const TTypedElement<UTypedElementSelectionInterface>& InActorSelectionHandle, UTypedElementList* InSelectionSet, const FTypedElementSelectionOptions& InSelectionOptions)
 {
 	AActor* Actor = ActorElementDataUtil::GetActorFromHandleChecked(InActorSelectionHandle);
 
-	// Allow active modes to potentially handle the selection
-	// TODO: Should this pass through the selection set?
-	if (GLevelEditorModeTools().IsSelectionHandled(Actor, /*bInSelected*/true))
+	if (const IToolkitHost* ToolkitHostPtr = GetToolkitHost())
 	{
-		return true;
+		// Allow active modes to potentially handle the selection
+		// TODO: Should this pass through the selection set?
+		if (ToolkitHostPtr->GetEditorModeManager().IsSelectionHandled(Actor, /*bInSelected*/true))
+		{
+			return true;
+		}
 	}
 
 	// If trying to select an actor, use this actors root selection actor instead (if it has one)
@@ -190,11 +204,14 @@ bool FActorElementLevelEditorSelectionCustomization::DeselectActorElement(const
 {
 	AActor* Actor = ActorElementDataUtil::GetActorFromHandleChecked(InActorSelectionHandle);
 
-	// Allow active modes to potentially handle the deselection
-	// TODO: Should this pass through the selection set?
-	if (GLevelEditorModeTools().IsSelectionHandled(Actor, /*bInSelected*/false))
+	if (const IToolkitHost* ToolkitHostPtr = GetToolkitHost())
 	{
-		return true;
+		// Allow active modes to potentially handle the deselection
+		// TODO: Should this pass through the selection set?
+		if (ToolkitHostPtr->GetEditorModeManager().IsSelectionHandled(Actor, /*bInSelected*/false))
+		{
+			return true;
+		}
 	}
 
 	bool bSelectionChanged = false;
diff --git a/Engine/Source/Editor/LevelEditor/Private/Elements/Actor/ActorElementLevelEditorSelectionCustomization.h b/Engine/Source/Editor/LevelEditor/Private/Elements/Actor/ActorElementLevelEditorSelectionCustomization.h
--- a/Engine/Source/Editor/LevelEditor/Private/Elements/Actor/ActorElementLevelEditorSelectionCustomization.h
+++ b/Engine/Source/Editor/LevelEditor/Private/Elements/Actor/ActorElementLevelEditorSelectionCustomization.h
@@ -3,10 +3,11 @@
 #pragma once
 
 #include "Elements/Framework/TypedElementSelectionSet.h"
+#include "Elements/Framework/TypedElementAssetEditorToolkitHostMixin.h"
 
 class AGroupActor;
 
-class FActorElementLevelEditorSelectionCustomization : public FTypedElementAssetEditorSelectionCustomization
+class FActorElementLevelEditorSelectionCustomization : public FTypedElementAssetEditorSelectionCustomization, public FTypedElementAssetEditorToolkitHostMixin
 {
 public:
 	virtual bool CanSelectElement(const TTypedElement<UTypedElementSelectionInterface>& InElementSelectionHandle, const FTypedElementSelectionOptions& InSelectionOptions) override;
@@ -16,12 +17,12 @@ public:
 	virtual bool AllowSelectionModifiers(const TTypedElement<UTypedElementSelectionInterface>& InElementSelectionHandle, const UTypedElementList* InSelectionSet) override;
 	virtual FTypedElementHandle GetSelectionElement(const TTypedElement<UTypedElementSelectionInterface>& InElementSelectionHandle, const UTypedElementList* InCurrentSelection, const ETypedElementSelectionMethod InSelectionMethod) override;
 
-	static bool CanSelectActorElement(const TTypedElement<UTypedElementSelectionInterface>& InActorSelectionHandle, const FTypedElementSelectionOptions& InSelectionOptions);
-	static bool CanDeselectActorElement(const TTypedElement<UTypedElementSelectionInterface>& InActorSelectionHandle, const FTypedElementSelectionOptions& InSelectionOptions);
+	bool CanSelectActorElement(const TTypedElement<UTypedElementSelectionInterface>& InActorSelectionHandle, const FTypedElementSelectionOptions& InSelectionOptions) const;
+	bool CanDeselectActorElement(const TTypedElement<UTypedElementSelectionInterface>& InActorSelectionHandle, const FTypedElementSelectionOptions& InSelectionOptions) const;
 
-	static bool SelectActorElement(const TTypedElement<UTypedElementSelectionInterface>& InActorSelectionHandle, UTypedElementList* InSelectionSet, const FTypedElementSelectionOptions& InSelectionOptions);
-	static bool DeselectActorElement(const TTypedElement<UTypedElementSelectionInterface>& InActorSelectionHandle, UTypedElementList* InSelectionSet, const FTypedElementSelectionOptions& InSelectionOptions);
+	bool SelectActorElement(const TTypedElement<UTypedElementSelectionInterface>& InActorSelectionHandle, UTypedElementList* InSelectionSet, const FTypedElementSelectionOptions& InSelectionOptions);
+	bool DeselectActorElement(const TTypedElement<UTypedElementSelectionInterface>& InActorSelectionHandle, UTypedElementList* InSelectionSet, const FTypedElementSelectionOptions& InSelectionOptions);
 
-	static bool SelectActorGroup(AGroupActor* InGroupActor, UTypedElementList* InSelectionSet, const FTypedElementSelectionOptions& InSelectionOptions, const bool bForce);
-	static bool DeselectActorGroup(AGroupActor* InGroupActor, UTypedElementList* InSelectionSet, const FTypedElementSelectionOptions& InSelectionOptions, const bool bForce);
+	bool SelectActorGroup(AGroupActor* InGroupActor, UTypedElementList* InSelectionSet, const FTypedElementSelectionOptions& InSelectionOptions, const bool bForce);
+	bool DeselectActorGroup(AGroupActor* InGroupActor, UTypedElementList* InSelectionSet, const FTypedElementSelectionOptions& InSelectionOptions, const bool bForce);
 };
diff --git a/Engine/Source/Editor/LevelEditor/Private/Elements/Actor/ActorElementLevelEditorViewportInteractionCustomization.cpp b/Engine/Source/Editor/LevelEditor/Private/Elements/Actor/ActorElementLevelEditorViewportInteractionCustomization.cpp
--- a/Engine/Source/Editor/LevelEditor/Private/Elements/Actor/ActorElementLevelEditorViewportInteractionCustomization.cpp
+++ b/Engine/Source/Editor/LevelEditor/Private/Elements/Actor/ActorElementLevelEditorViewportInteractionCustomization.cpp
@@ -15,11 +15,6 @@
 #include "ActorGroupingUtils.h"
 #include "LevelEditorViewport.h"
 
-FActorElementLevelEditorViewportInteractionCustomization::FActorElementLevelEditorViewportInteractionCustomization(FLevelEditorViewportClient* InLevelEditorViewportClient)
-	: LevelEditorViewportClient(InLevelEditorViewportClient)
-{
-}
-
 void FActorElementLevelEditorViewportInteractionCustomization::GetElementsToMove(const TTypedElement<UTypedElementWorldInterface>& InElementWorldHandle, const ETypedElementViewportInteractionWorldType InWorldType, const UTypedElementSelectionSet* InSelectionSet, UTypedElementList* OutElementsToMove)
 {
 	AActor* Actor = ActorElementDataUtil::GetActorFromHandleChecked(InElementWorldHandle);
@@ -81,7 +76,7 @@ void FActorElementLevelEditorViewportInteractionCustomization::GizmoManipulation
 	FActorElementEditorViewportInteractionCustomization::GizmoManipulationDeltaUpdate(InElementWorldHandle, InWidgetMode, InDragAxis, InInputState, ModifiedDeltaTransform, InPivotLocation);
 
 	// Update the cameras from their locked actor (if any) only if the viewport is real-time enabled
-	LevelEditorViewportClient->UpdateLockedActorViewports(Actor, true);
+	GetMutableLevelEditorViewportClient()->UpdateLockedActorViewports(Actor, true);
 }
 
 void FActorElementLevelEditorViewportInteractionCustomization::GizmoManipulationStopped(const TTypedElement<UTypedElementWorldInterface>& InElementWorldHandle, const UE::Widget::EWidgetMode InWidgetMode)
@@ -105,7 +100,7 @@ void FActorElementLevelEditorViewportInteractionCustomization::PostGizmoManipula
 	GEditor->BroadcastActorsMoved(MovedActors);
 }
 
-void FActorElementLevelEditorViewportInteractionCustomization::ModifyScale(AActor* InActor, const EAxisList::Type InDragAxis, FVector& ScaleDelta, bool bCheckSmallExtent) const
+void FActorElementLevelEditorViewportInteractionCustomization::ModifyScale(AActor* InActor, const EAxisList::Type InDragAxis, FVector& ScaleDelta, bool bCheckSmallExtent)
 {
 	if (InActor->GetRootComponent())
 	{
@@ -113,7 +108,7 @@ void FActorElementLevelEditorViewportInteractionCustomization::ModifyScale(AActo
 
 		const FBox LocalBox = InActor->GetComponentsBoundingBox(true);
 		const FVector ScaledExtents = LocalBox.GetExtent() * CurrentScale;
-		const FTransform PreDragTransform = LevelEditorViewportClient->CachePreDragActorTransform(InActor);
+		const FTransform PreDragTransform = GetMutableLevelEditorViewportClient()->CachePreDragActorTransform(InActor);
 
 		FComponentElementLevelEditorViewportInteractionCustomization::ValidateScale(PreDragTransform.GetScale3D(), InDragAxis, CurrentScale, ScaledExtents, ScaleDelta, bCheckSmallExtent);
 
diff --git a/Engine/Source/Editor/LevelEditor/Private/Elements/Component/ComponentElementLevelEditorSelectionCustomization.cpp b/Engine/Source/Editor/LevelEditor/Private/Elements/Component/ComponentElementLevelEditorSelectionCustomization.cpp
--- a/Engine/Source/Editor/LevelEditor/Private/Elements/Component/ComponentElementLevelEditorSelectionCustomization.cpp
+++ b/Engine/Source/Editor/LevelEditor/Private/Elements/Component/ComponentElementLevelEditorSelectionCustomization.cpp
@@ -84,7 +84,7 @@ FTypedElementHandle FComponentElementLevelEditorSelectionCustomization::GetSelec
 	return InElementSelectionHandle;
 }
 
-bool FComponentElementLevelEditorSelectionCustomization::CanSelectComponentElement(const TTypedElement<UTypedElementSelectionInterface>& InComponentSelectionHandle, const FTypedElementSelectionOptions& InSelectionOptions)
+bool FComponentElementLevelEditorSelectionCustomization::CanSelectComponentElement(const TTypedElement<UTypedElementSelectionInterface>& InComponentSelectionHandle, const FTypedElementSelectionOptions& InSelectionOptions) const
 {
 	const UActorComponent* Component = ComponentElementDataUtil::GetComponentFromHandleChecked(InComponentSelectionHandle);
 
@@ -92,7 +92,7 @@ bool FComponentElementLevelEditorSelectionCustomization::CanSelectComponentEleme
 	return !GEdSelectionLock;
 }
 
-bool FComponentElementLevelEditorSelectionCustomization::CanDeselectComponentElement(const TTypedElement<UTypedElementSelectionInterface>& InComponentSelectionHandle, const FTypedElementSelectionOptions& InSelectionOptions)
+bool FComponentElementLevelEditorSelectionCustomization::CanDeselectComponentElement(const TTypedElement<UTypedElementSelectionInterface>& InComponentSelectionHandle, const FTypedElementSelectionOptions& InSelectionOptions) const
 {
 	const UActorComponent* Component = ComponentElementDataUtil::GetComponentFromHandleChecked(InComponentSelectionHandle);
 
diff --git a/Engine/Source/Editor/LevelEditor/Private/Elements/Component/ComponentElementLevelEditorSelectionCustomization.h b/Engine/Source/Editor/LevelEditor/Private/Elements/Component/ComponentElementLevelEditorSelectionCustomization.h
--- a/Engine/Source/Editor/LevelEditor/Private/Elements/Component/ComponentElementLevelEditorSelectionCustomization.h
+++ b/Engine/Source/Editor/LevelEditor/Private/Elements/Component/ComponentElementLevelEditorSelectionCustomization.h
@@ -3,8 +3,9 @@
 #pragma once
 
 #include "Elements/Framework/TypedElementSelectionSet.h"
+#include "Elements/Framework/TypedElementAssetEditorToolkitHostMixin.h"
 
-class FComponentElementLevelEditorSelectionCustomization : public FTypedElementAssetEditorSelectionCustomization
+class FComponentElementLevelEditorSelectionCustomization : public FTypedElementAssetEditorSelectionCustomization, public FTypedElementAssetEditorToolkitHostMixin
 {
 public:
 	virtual bool CanSelectElement(const TTypedElement<UTypedElementSelectionInterface>& InElementSelectionHandle, const FTypedElementSelectionOptions& InSelectionOptions) override;
@@ -13,9 +14,9 @@ public:
 	virtual bool DeselectElement(const TTypedElement<UTypedElementSelectionInterface>& InElementSelectionHandle, UTypedElementList* InSelectionSet, const FTypedElementSelectionOptions& InSelectionOptions) override;
 	virtual FTypedElementHandle GetSelectionElement(const TTypedElement<UTypedElementSelectionInterface>& InElementSelectionHandle, const UTypedElementList* InCurrentSelection, const ETypedElementSelectionMethod InSelectionMethod) override;
 
-	static bool CanSelectComponentElement(const TTypedElement<UTypedElementSelectionInterface>& InComponentSelectionHandle, const FTypedElementSelectionOptions& InSelectionOptions);
-	static bool CanDeselectComponentElement(const TTypedElement<UTypedElementSelectionInterface>& InComponentSelectionHandle, const FTypedElementSelectionOptions& InSelectionOptions);
+	bool CanSelectComponentElement(const TTypedElement<UTypedElementSelectionInterface>& InComponentSelectionHandle, const FTypedElementSelectionOptions& InSelectionOptions) const;
+	bool CanDeselectComponentElement(const TTypedElement<UTypedElementSelectionInterface>& InComponentSelectionHandle, const FTypedElementSelectionOptions& InSelectionOptions) const;
 
-	static bool SelectComponentElement(const TTypedElement<UTypedElementSelectionInterface>& InComponentSelectionHandle, UTypedElementList* InSelectionSet, const FTypedElementSelectionOptions& InSelectionOptions);
-	static bool DeselectComponentElement(const TTypedElement<UTypedElementSelectionInterface>& InComponentSelectionHandle, UTypedElementList* InSelectionSet, const FTypedElementSelectionOptions& InSelectionOptions);
+	bool SelectComponentElement(const TTypedElement<UTypedElementSelectionInterface>& InComponentSelectionHandle, UTypedElementList* InSelectionSet, const FTypedElementSelectionOptions& InSelectionOptions);
+	bool DeselectComponentElement(const TTypedElement<UTypedElementSelectionInterface>& InComponentSelectionHandle, UTypedElementList* InSelectionSet, const FTypedElementSelectionOptions& InSelectionOptions);
 };
diff --git a/Engine/Source/Editor/LevelEditor/Private/Elements/Component/ComponentElementLevelEditorViewportInteractionCustomization.cpp b/Engine/Source/Editor/LevelEditor/Private/Elements/Component/ComponentElementLevelEditorViewportInteractionCustomization.cpp
--- a/Engine/Source/Editor/LevelEditor/Private/Elements/Component/ComponentElementLevelEditorViewportInteractionCustomization.cpp
+++ b/Engine/Source/Editor/LevelEditor/Private/Elements/Component/ComponentElementLevelEditorViewportInteractionCustomization.cpp
@@ -12,11 +12,6 @@
 #include "LevelEditorViewport.h"
 #include "Kismet2/ComponentEditorUtils.h"
 
-FComponentElementLevelEditorViewportInteractionCustomization::FComponentElementLevelEditorViewportInteractionCustomization(FLevelEditorViewportClient* InLevelEditorViewportClient)
-	: LevelEditorViewportClient(InLevelEditorViewportClient)
-{
-}
-
 void FComponentElementLevelEditorViewportInteractionCustomization::GetElementsToMove(const TTypedElement<UTypedElementWorldInterface>& InElementWorldHandle, const ETypedElementViewportInteractionWorldType InWorldType, const UTypedElementSelectionSet* InSelectionSet, UTypedElementList* OutElementsToMove)
 {
 	UActorComponent* Component = ComponentElementDataUtil::GetComponentFromHandleChecked(InElementWorldHandle);
@@ -125,10 +120,10 @@ void FComponentElementLevelEditorViewportInteractionCustomization::GizmoManipula
 	}
 }
 
-void FComponentElementLevelEditorViewportInteractionCustomization::ModifyScale(USceneComponent* InComponent, const EAxisList::Type InDragAxis, FVector& ScaleDelta) const
+void FComponentElementLevelEditorViewportInteractionCustomization::ModifyScale(USceneComponent* InComponent, const EAxisList::Type InDragAxis, FVector& ScaleDelta)
 {
 	AActor* Actor = InComponent->GetOwner();
-	const FTransform PreDragTransform = LevelEditorViewportClient->CachePreDragActorTransform(Actor);
+	const FTransform PreDragTransform = GetMutableLevelEditorViewportClient()->CachePreDragActorTransform(Actor);
 	const FBox LocalBox = Actor->GetComponentsBoundingBox(true);
 	const FVector ScaledExtents = LocalBox.GetExtent() * InComponent->GetRelativeScale3D();
 
diff --git a/Engine/Source/Editor/LevelEditor/Private/SLevelEditor.cpp b/Engine/Source/Editor/LevelEditor/Private/SLevelEditor.cpp
--- a/Engine/Source/Editor/LevelEditor/Private/SLevelEditor.cpp
+++ b/Engine/Source/Editor/LevelEditor/Private/SLevelEditor.cpp
@@ -222,8 +222,16 @@ void SLevelEditor::Initialize( const TSharedRef<SDockTab>& OwnerTab, const TShar
 	SelectedElements->AddToRoot();
 
 	// Register the level editor specific selection behavior
-	SelectedElements->RegisterAssetEditorCustomizationByTypeName(NAME_Actor, MakeUnique<FActorElementLevelEditorSelectionCustomization>());
-	SelectedElements->RegisterAssetEditorCustomizationByTypeName(NAME_Components, MakeUnique<FComponentElementLevelEditorSelectionCustomization>());
+	{
+		TUniquePtr<FActorElementLevelEditorSelectionCustomization> ActorCustomization = MakeUnique<FActorElementLevelEditorSelectionCustomization>();
+		ActorCustomization->SetToolkitHost(this);
+		SelectedElements->RegisterAssetEditorCustomizationByTypeName(NAME_Actor, MoveTemp(ActorCustomization));
+	}
+	{
+		TUniquePtr<FComponentElementLevelEditorSelectionCustomization> ComponentCustomization = MakeUnique<FComponentElementLevelEditorSelectionCustomization>();
+		ComponentCustomization->SetToolkitHost(this);
+		SelectedElements->RegisterAssetEditorCustomizationByTypeName(NAME_Components, MoveTemp(ComponentCustomization));
+	}
 
 	// Allow USelection to bridge to our selected element list
 	GUnrealEd->GetSelectedActors()->SetElementSelectionSet(SelectedElements);
diff --git a/Engine/Source/Editor/LevelEditor/Private/SLevelViewport.cpp b/Engine/Source/Editor/LevelEditor/Private/SLevelViewport.cpp
--- a/Engine/Source/Editor/LevelEditor/Private/SLevelViewport.cpp
+++ b/Engine/Source/Editor/LevelEditor/Private/SLevelViewport.cpp
@@ -511,6 +511,8 @@ void SLevelViewport::ConstructLevelEditorViewportClient(FLevelEditorViewportInst
 
 	LevelViewportClient->SetViewModes(ViewportInstanceSettings.PerspViewModeIndex, ViewportInstanceSettings.OrthoViewModeIndex );
 
+	LevelViewportClient->InitializeViewportInteraction();
+
 	bShowFullToolbar = ViewportInstanceSettings.bShowFullToolbar;
 
 	// Always set to true initially
diff --git a/Engine/Source/Editor/LevelEditor/Public/Elements/Actor/ActorElementLevelEditorViewportInteractionCustomization.h b/Engine/Source/Editor/LevelEditor/Public/Elements/Actor/ActorElementLevelEditorViewportInteractionCustomization.h
--- a/Engine/Source/Editor/LevelEditor/Public/Elements/Actor/ActorElementLevelEditorViewportInteractionCustomization.h
+++ b/Engine/Source/Editor/LevelEditor/Public/Elements/Actor/ActorElementLevelEditorViewportInteractionCustomization.h
@@ -3,14 +3,11 @@
 #pragma once
 
 #include "Elements/Actor/ActorElementEditorViewportInteractionCustomization.h"
+#include "Elements/Framework/TypedElementAssetEditorLevelEditorViewportClientMixin.h"
 
-class FLevelEditorViewportClient;
-
-class LEVELEDITOR_API FActorElementLevelEditorViewportInteractionCustomization : public FActorElementEditorViewportInteractionCustomization
+class LEVELEDITOR_API FActorElementLevelEditorViewportInteractionCustomization : public FActorElementEditorViewportInteractionCustomization, public FTypedElementAssetEditorLevelEditorViewportClientMixin
 {
 public:
-	explicit FActorElementLevelEditorViewportInteractionCustomization(FLevelEditorViewportClient* InLevelEditorViewportClient);
-
 	virtual void GetElementsToMove(const TTypedElement<UTypedElementWorldInterface>& InElementWorldHandle, const ETypedElementViewportInteractionWorldType InWorldType, const UTypedElementSelectionSet* InSelectionSet, UTypedElementList* OutElementsToMove) override;
 	virtual void GizmoManipulationStarted(const TTypedElement<UTypedElementWorldInterface>& InElementWorldHandle, const UE::Widget::EWidgetMode InWidgetMode) override;
 	virtual void GizmoManipulationDeltaUpdate(const TTypedElement<UTypedElementWorldInterface>& InElementWorldHandle, const UE::Widget::EWidgetMode InWidgetMode, const EAxisList::Type InDragAxis, const FInputDeviceState& InInputState, const FTransform& InDeltaTransform, const FVector& InPivotLocation) override;
@@ -21,7 +18,5 @@ public:
 	static void AppendActorsToMove(AActor* InActor, const UTypedElementSelectionSet* InSelectionSet, UTypedElementList* OutElementsToMove);
 
 private:
-	void ModifyScale(AActor* InActor, const EAxisList::Type InDragAxis, FVector& ScaleDelta, bool bCheckSmallExtent) const;
-
-	FLevelEditorViewportClient* LevelEditorViewportClient = nullptr;
+	void ModifyScale(AActor* InActor, const EAxisList::Type InDragAxis, FVector& ScaleDelta, bool bCheckSmallExtent);
 };
diff --git a/Engine/Source/Editor/LevelEditor/Public/Elements/Component/ComponentElementLevelEditorViewportInteractionCustomization.h b/Engine/Source/Editor/LevelEditor/Public/Elements/Component/ComponentElementLevelEditorViewportInteractionCustomization.h
--- a/Engine/Source/Editor/LevelEditor/Public/Elements/Component/ComponentElementLevelEditorViewportInteractionCustomization.h
+++ b/Engine/Source/Editor/LevelEditor/Public/Elements/Component/ComponentElementLevelEditorViewportInteractionCustomization.h
@@ -3,14 +3,13 @@
 #pragma once
 
 #include "Elements/Component/ComponentElementEditorViewportInteractionCustomization.h"
+#include "Elements/Framework/TypedElementAssetEditorLevelEditorViewportClientMixin.h"
 
 class FLevelEditorViewportClient;
 
-class LEVELEDITOR_API FComponentElementLevelEditorViewportInteractionCustomization : public FComponentElementEditorViewportInteractionCustomization
+class LEVELEDITOR_API FComponentElementLevelEditorViewportInteractionCustomization : public FComponentElementEditorViewportInteractionCustomization, public FTypedElementAssetEditorLevelEditorViewportClientMixin
 {
 public:
-	explicit FComponentElementLevelEditorViewportInteractionCustomization(FLevelEditorViewportClient* InLevelEditorViewportClient);
-
 	virtual void GetElementsToMove(const TTypedElement<UTypedElementWorldInterface>& InElementWorldHandle, const ETypedElementViewportInteractionWorldType InWorldType, const UTypedElementSelectionSet* InSelectionSet, UTypedElementList* OutElementsToMove) override;
 	virtual void GizmoManipulationStarted(const TTypedElement<UTypedElementWorldInterface>& InElementWorldHandle, const UE::Widget::EWidgetMode InWidgetMode) override;
 	virtual void GizmoManipulationDeltaUpdate(const TTypedElement<UTypedElementWorldInterface>& InElementWorldHandle, const UE::Widget::EWidgetMode InWidgetMode, const EAxisList::Type InDragAxis, const FInputDeviceState& InInputState, const FTransform& InDeltaTransform, const FVector& InPivotLocation) override;
@@ -20,7 +19,5 @@ public:
 	static FProperty* GetEditTransformProperty(const UE::Widget::EWidgetMode InWidgetMode);
 
 private:
-	void ModifyScale(USceneComponent* InComponent, const EAxisList::Type InDragAxis, FVector& ScaleDelta) const;
-
-	FLevelEditorViewportClient* LevelEditorViewportClient = nullptr;
+	void ModifyScale(USceneComponent* InComponent, const EAxisList::Type InDragAxis, FVector& ScaleDelta);
 };
diff --git a/Engine/Source/Editor/LevelEditor/Public/Elements/Framework/TypedElementAssetEditorLevelEditorViewportClientMixin.h b/Engine/Source/Editor/LevelEditor/Public/Elements/Framework/TypedElementAssetEditorLevelEditorViewportClientMixin.h
new file mode 100644
--- /dev/null
+++ b/Engine/Source/Editor/LevelEditor/Public/Elements/Framework/TypedElementAssetEditorLevelEditorViewportClientMixin.h
@@ -0,0 +1,42 @@
+// Copyright Epic Games, Inc. All Rights Reserved.
+
+#pragma once
+
+#include "CoreMinimal.h"
+
+class FLevelEditorViewportClient;
+
+/**
+ * Mixin for asset editor customizations that provides access to the owner level editor viewport client.
+ */
+class FTypedElementAssetEditorLevelEditorViewportClientMixin
+{
+public:
+	/**
+	 * Get the owner level editor viewport client.
+	 */
+	const FLevelEditorViewportClient* GetLevelEditorViewportClient() const
+	{
+		return LevelEditorViewportClient;
+	}
+
+	/**
+	 * Get the owner level editor viewport client.
+	 */
+	FLevelEditorViewportClient* GetMutableLevelEditorViewportClient()
+	{
+		return LevelEditorViewportClient;
+	}
+
+	/**
+	 * Set the owner level editor viewport client.
+	 */
+	void SetLevelEditorViewportClient(FLevelEditorViewportClient* InLevelEditorViewportClient)
+	{
+		LevelEditorViewportClient = InLevelEditorViewportClient;
+	}
+
+private:
+	/** The owner level editor viewport client. */
+	FLevelEditorViewportClient* LevelEditorViewportClient = nullptr;
+};
diff --git a/Engine/Source/Editor/UnrealEd/Private/Elements/Actor/ActorElementEditorViewportInteractionCustomization.cpp b/Engine/Source/Editor/UnrealEd/Private/Elements/Actor/ActorElementEditorViewportInteractionCustomization.cpp
--- a/Engine/Source/Editor/UnrealEd/Private/Elements/Actor/ActorElementEditorViewportInteractionCustomization.cpp
+++ b/Engine/Source/Editor/UnrealEd/Private/Elements/Actor/ActorElementEditorViewportInteractionCustomization.cpp
@@ -6,12 +6,18 @@
 
 #include "Editor.h"
 #include "EditorModeManager.h"
+#include "Toolkits/IToolkitHost.h"
 #include "AI/NavigationSystemBase.h"
 
 bool FActorElementEditorViewportInteractionCustomization::GetGizmoPivotLocation(const TTypedElement<UTypedElementWorldInterface>& InElementWorldHandle, const UE::Widget::EWidgetMode InWidgetMode, FVector& OutPivotLocation)
 {
-	OutPivotLocation = GLevelEditorModeTools().PivotLocation;
-	return true;
+	if (const IToolkitHost* ToolkitHostPtr = GetToolkitHost())
+	{
+		OutPivotLocation = ToolkitHostPtr->GetEditorModeManager().PivotLocation;
+		return true;
+	}
+	
+	return FTypedElementAssetEditorViewportInteractionCustomization::GetGizmoPivotLocation(InElementWorldHandle, InWidgetMode, OutPivotLocation);
 }
 
 void FActorElementEditorViewportInteractionCustomization::GizmoManipulationDeltaUpdate(const TTypedElement<UTypedElementWorldInterface>& InElementWorldHandle, const UE::Widget::EWidgetMode InWidgetMode, const EAxisList::Type InDragAxis, const FInputDeviceState& InInputState, const FTransform& InDeltaTransform, const FVector& InPivotLocation)
diff --git a/Engine/Source/Editor/UnrealEd/Private/Elements/Component/ComponentElementEditorViewportInteractionCustomization.cpp b/Engine/Source/Editor/UnrealEd/Private/Elements/Component/ComponentElementEditorViewportInteractionCustomization.cpp
--- a/Engine/Source/Editor/UnrealEd/Private/Elements/Component/ComponentElementEditorViewportInteractionCustomization.cpp
+++ b/Engine/Source/Editor/UnrealEd/Private/Elements/Component/ComponentElementEditorViewportInteractionCustomization.cpp
@@ -6,6 +6,7 @@
 
 #include "Editor.h"
 #include "EditorModeManager.h"
+#include "Toolkits/IToolkitHost.h"
 #include "EditorSupportDelegates.h"
 
 bool FComponentElementEditorViewportInteractionCustomization::GetGizmoPivotLocation(const TTypedElement<UTypedElementWorldInterface>& InElementWorldHandle, const UE::Widget::EWidgetMode InWidgetMode, FVector& OutPivotLocation)
@@ -14,9 +15,19 @@ bool FComponentElementEditorViewportInteractionCustomization::GetGizmoPivotLocat
 
 	if (const USceneComponent* SceneComponent = Cast<USceneComponent>(Component))
 	{
+		FVector BasePivotLocation = FVector::ZeroVector;
+		if (const IToolkitHost* ToolkitHostPtr = GetToolkitHost())
+		{
+			BasePivotLocation = ToolkitHostPtr->GetEditorModeManager().PivotLocation;
+		}
+		else
+		{
+			FTypedElementAssetEditorViewportInteractionCustomization::GetGizmoPivotLocation(InElementWorldHandle, InWidgetMode, BasePivotLocation);
+		}
+
 		// If necessary, transform the editor pivot location to be relative to the component's parent
 		const bool bIsRootComponent = SceneComponent->GetOwner()->GetRootComponent() == SceneComponent;
-		OutPivotLocation = bIsRootComponent || !SceneComponent->GetAttachParent() ? GEditor->GetPivotLocation() : SceneComponent->GetAttachParent()->GetComponentToWorld().Inverse().TransformPosition(GEditor->GetPivotLocation());
+		OutPivotLocation = bIsRootComponent || !SceneComponent->GetAttachParent() ? BasePivotLocation : SceneComponent->GetAttachParent()->GetComponentToWorld().Inverse().TransformPosition(BasePivotLocation);
 		return true;
 	}
 
diff --git a/Engine/Source/Editor/UnrealEd/Private/LevelEditorViewport.cpp b/Engine/Source/Editor/UnrealEd/Private/LevelEditorViewport.cpp
--- a/Engine/Source/Editor/UnrealEd/Private/LevelEditorViewport.cpp
+++ b/Engine/Source/Editor/UnrealEd/Private/LevelEditorViewport.cpp
@@ -1718,9 +1718,6 @@ FLevelEditorViewportClient::FLevelEditorViewportClient(const TSharedPtr<SLevelVi
 	, bWasEditorCameraCut(false)
 	, bApplyCameraSpeedScaleByDistance(true)
 {
-	ViewportInteraction->RegisterAssetEditorCustomizationByTypeName(NAME_Actor, MakeUnique<FActorElementLevelEditorViewportInteractionCustomization>(this));
-	ViewportInteraction->RegisterAssetEditorCustomizationByTypeName(NAME_Components, MakeUnique<FComponentElementLevelEditorViewportInteractionCustomization>(this));
-
 	// By default a level editor viewport is pointed to the editor world
 	SetReferenceToWorldContext(GEditor->GetEditorWorldContext());
 
@@ -1813,6 +1810,22 @@ void FLevelEditorViewportClient::InitializeVisibilityFlags()
 	SpriteCategoryVisibility.Init(true, GUnrealEd->SpriteIDToIndexMap.Num());
 }
 
+void FLevelEditorViewportClient::InitializeViewportInteraction()
+{
+	{
+		TUniquePtr<FActorElementLevelEditorViewportInteractionCustomization> ActorCustomization = MakeUnique<FActorElementLevelEditorViewportInteractionCustomization>();
+		ActorCustomization->SetLevelEditorViewportClient(this);
+		ActorCustomization->SetToolkitHost(ParentLevelEditor.Pin().Get());
+		ViewportInteraction->RegisterAssetEditorCustomizationByTypeName(NAME_Actor, MoveTemp(ActorCustomization));
+	}
+	{
+		TUniquePtr<FComponentElementLevelEditorViewportInteractionCustomization> ComponentCustomization = MakeUnique<FComponentElementLevelEditorViewportInteractionCustomization>();
+		ComponentCustomization->SetLevelEditorViewportClient(this);
+		ComponentCustomization->SetToolkitHost(ParentLevelEditor.Pin().Get());
+		ViewportInteraction->RegisterAssetEditorCustomizationByTypeName(NAME_Components, MoveTemp(ComponentCustomization));
+	}
+}
+
 FSceneView* FLevelEditorViewportClient::CalcSceneView(FSceneViewFamily* ViewFamily, const EStereoscopicPass StereoPass)
 {
 	bWasControlledByOtherViewport = false;
diff --git a/Engine/Source/Editor/UnrealEd/Public/Elements/Actor/ActorElementEditorViewportInteractionCustomization.h b/Engine/Source/Editor/UnrealEd/Public/Elements/Actor/ActorElementEditorViewportInteractionCustomization.h
--- a/Engine/Source/Editor/UnrealEd/Public/Elements/Actor/ActorElementEditorViewportInteractionCustomization.h
+++ b/Engine/Source/Editor/UnrealEd/Public/Elements/Actor/ActorElementEditorViewportInteractionCustomization.h
@@ -3,10 +3,11 @@
 #pragma once
 
 #include "Elements/Framework/TypedElementViewportInteraction.h"
+#include "Elements/Framework/TypedElementAssetEditorToolkitHostMixin.h"
 
 class AActor;
 
-class UNREALED_API FActorElementEditorViewportInteractionCustomization : public FTypedElementAssetEditorViewportInteractionCustomization
+class UNREALED_API FActorElementEditorViewportInteractionCustomization : public FTypedElementAssetEditorViewportInteractionCustomization, public FTypedElementAssetEditorToolkitHostMixin
 {
 public:
 	virtual bool GetGizmoPivotLocation(const TTypedElement<UTypedElementWorldInterface>& InElementWorldHandle, const UE::Widget::EWidgetMode InWidgetMode, FVector& OutPivotLocation) override;
diff --git a/Engine/Source/Editor/UnrealEd/Public/Elements/Component/ComponentElementEditorViewportInteractionCustomization.h b/Engine/Source/Editor/UnrealEd/Public/Elements/Component/ComponentElementEditorViewportInteractionCustomization.h
--- a/Engine/Source/Editor/UnrealEd/Public/Elements/Component/ComponentElementEditorViewportInteractionCustomization.h
+++ b/Engine/Source/Editor/UnrealEd/Public/Elements/Component/ComponentElementEditorViewportInteractionCustomization.h
@@ -3,11 +3,12 @@
 #pragma once
 
 #include "Elements/Framework/TypedElementViewportInteraction.h"
+#include "Elements/Framework/TypedElementAssetEditorToolkitHostMixin.h"
 
 class FProperty;
 class USceneComponent;
 
-class UNREALED_API FComponentElementEditorViewportInteractionCustomization : public FTypedElementAssetEditorViewportInteractionCustomization
+class UNREALED_API FComponentElementEditorViewportInteractionCustomization : public FTypedElementAssetEditorViewportInteractionCustomization, public FTypedElementAssetEditorToolkitHostMixin
 {
 public:
 	virtual bool GetGizmoPivotLocation(const TTypedElement<UTypedElementWorldInterface>& InElementWorldHandle, const UE::Widget::EWidgetMode InWidgetMode, FVector& OutPivotLocation) override;
diff --git a/Engine/Source/Editor/UnrealEd/Public/LevelEditorViewport.h b/Engine/Source/Editor/UnrealEd/Public/LevelEditorViewport.h
--- a/Engine/Source/Editor/UnrealEd/Public/LevelEditorViewport.h
+++ b/Engine/Source/Editor/UnrealEd/Public/LevelEditorViewport.h
@@ -241,6 +241,11 @@ public:
 	 */
 	void InitializeVisibilityFlags();
 
+	/**
+	 * Initialize viewport interaction
+	 */
+	void InitializeViewportInteraction();
+
 	/**
 	 * Reset the camera position and rotation.  Used when creating a new level.
 	 */
