Updated viewport menu icons

diff --git a/Engine/Plugins/Cameras/CameraShakePreviewer/Source/CameraShakePreviewer/Private/CameraShakePreviewerModule.cpp b/Engine/Plugins/Cameras/CameraShakePreviewer/Source/CameraShakePreviewer/Private/CameraShakePreviewerModule.cpp
--- a/Engine/Plugins/Cameras/CameraShakePreviewer/Source/CameraShakePreviewer/Private/CameraShakePreviewerModule.cpp
+++ b/Engine/Plugins/Cameras/CameraShakePreviewer/Source/CameraShakePreviewer/Private/CameraShakePreviewerModule.cpp
@@ -160,7 +160,7 @@ void FCameraShakePreviewerModule::RegisterViewportOptionMenuExtender()
 						"ToggleCameraShakesPreview",
 						ToggleCameraShakesPreview->GetLabel(),
 						ToggleCameraShakesPreview->GetDescription(),
-						ToggleCameraShakesPreview->GetIcon(),
+						FSlateIcon(FAppStyle::Get().GetStyleSetName(), "LevelViewport.ToggleCameraShakePreview"),
 						Action,
 						ToggleCameraShakesPreview->GetUserInterfaceType()
 					);
diff --git a/Engine/Source/Editor/CommonMenuExtensions/Private/ShowFlagMenuCommands.cpp b/Engine/Source/Editor/CommonMenuExtensions/Private/ShowFlagMenuCommands.cpp
--- a/Engine/Source/Editor/CommonMenuExtensions/Private/ShowFlagMenuCommands.cpp
+++ b/Engine/Source/Editor/CommonMenuExtensions/Private/ShowFlagMenuCommands.cpp
@@ -66,15 +66,15 @@ void FShowFlagMenuCommands::BuildShowFlagsMenu(UToolMenu* Menu, const FShowFlagF
 	CreateCommonShowFlagMenuItems(Menu, Filter);
 
 	{
-		FToolMenuSection& Section = Menu->AddSection("LevelViewportShowFlags", LOCTEXT("AllShowFlagHeader", "All Show Flags"));;
-		CreateSubMenuIfRequired(Section, Filter, SFG_PostProcess, "SFG_PostProcess", LOCTEXT("PostProcessShowFlagsMenu", "Post Processing"), LOCTEXT("PostProcessShowFlagsMenu_ToolTip", "Post process show flags"));
-		CreateSubMenuIfRequired(Section, Filter, SFG_LightTypes, "SFG_LightTypes", LOCTEXT("LightTypesShowFlagsMenu", "Light Types"), LOCTEXT("LightTypesShowFlagsMenu_ToolTip", "Light Types show flags"));
-		CreateSubMenuIfRequired(Section, Filter, SFG_LightingComponents, "SFG_LightingComponents", LOCTEXT("LightingComponentsShowFlagsMenu", "Lighting Components"), LOCTEXT("LightingComponentsShowFlagsMenu_ToolTip", "Lighting Components show flags"));
-		CreateSubMenuIfRequired(Section, Filter, SFG_LightingFeatures, "SFG_LightingFeatures", LOCTEXT("LightingFeaturesShowFlagsMenu", "Lighting Features"), LOCTEXT("LightingFeaturesShowFlagsMenu_ToolTip", "Lighting Features show flags"));
-		CreateSubMenuIfRequired(Section, Filter, SFG_Developer, "SFG_Developer", LOCTEXT("DeveloperShowFlagsMenu", "Developer"), LOCTEXT("DeveloperShowFlagsMenu_ToolTip", "Developer show flags"));
-		CreateSubMenuIfRequired(Section, Filter, SFG_Visualize, "SFG_Visualize", LOCTEXT("VisualizeShowFlagsMenu", "Visualize"), LOCTEXT("VisualizeShowFlagsMenu_ToolTip", "Visualize show flags"));
-		CreateSubMenuIfRequired(Section, Filter, SFG_Advanced, "SFG_Advanced", LOCTEXT("AdvancedShowFlagsMenu", "Advanced"), LOCTEXT("AdvancedShowFlagsMenu_ToolTip", "Advanced show flags"));
-		CreateSubMenuIfRequired(Section, Filter, SFG_Custom, "SFG_Custom", LOCTEXT("CustomShowFlagsMenu", "Custom"), LOCTEXT("CustomShowFlagsMenu_ToolTip", "Custom show flags"));
+		FToolMenuSection& Section = Menu->AddSection("LevelViewportShowFlags", LOCTEXT("AllShowFlagHeader", "All Show Flags"));
+		CreateSubMenuIfRequired(Section, Filter, SFG_PostProcess, "SFG_PostProcess", LOCTEXT("PostProcessShowFlagsMenu", "Post Processing"), LOCTEXT("PostProcessShowFlagsMenu_ToolTip", "Post process show flags"), "ShowFlagsMenu.SubMenu.PostProcessing");
+		CreateSubMenuIfRequired(Section, Filter, SFG_LightTypes, "SFG_LightTypes", LOCTEXT("LightTypesShowFlagsMenu", "Light Types"), LOCTEXT("LightTypesShowFlagsMenu_ToolTip", "Light Types show flags"), "ShowFlagsMenu.SubMenu.LightTypes");
+		CreateSubMenuIfRequired(Section, Filter, SFG_LightingComponents, "SFG_LightingComponents", LOCTEXT("LightingComponentsShowFlagsMenu", "Lighting Components"), LOCTEXT("LightingComponentsShowFlagsMenu_ToolTip", "Lighting Components show flags"), "ShowFlagsMenu.SubMenu.LightingComponents");
+		CreateSubMenuIfRequired(Section, Filter, SFG_LightingFeatures, "SFG_LightingFeatures", LOCTEXT("LightingFeaturesShowFlagsMenu", "Lighting Features"), LOCTEXT("LightingFeaturesShowFlagsMenu_ToolTip", "Lighting Features show flags"), "ShowFlagsMenu.SubMenu.LightingFeatures");
+		CreateSubMenuIfRequired(Section, Filter, SFG_Developer, "SFG_Developer", LOCTEXT("DeveloperShowFlagsMenu", "Developer"), LOCTEXT("DeveloperShowFlagsMenu_ToolTip", "Developer show flags"), "ShowFlagsMenu.SubMenu.Developer");
+		CreateSubMenuIfRequired(Section, Filter, SFG_Visualize, "SFG_Visualize", LOCTEXT("VisualizeShowFlagsMenu", "Visualize"), LOCTEXT("VisualizeShowFlagsMenu_ToolTip", "Visualize show flags"), "ShowFlagsMenu.SubMenu.Visualize");
+		CreateSubMenuIfRequired(Section, Filter, SFG_Advanced, "SFG_Advanced", LOCTEXT("AdvancedShowFlagsMenu", "Advanced"), LOCTEXT("AdvancedShowFlagsMenu_ToolTip", "Advanced show flags"), "ShowFlagsMenu.SubMenu.Advanced");
+		CreateSubMenuIfRequired(Section, Filter, SFG_Custom, "SFG_Custom", LOCTEXT("CustomShowFlagsMenu", "Custom"), LOCTEXT("CustomShowFlagsMenu_ToolTip", "Custom show flags"), NAME_None);
 	}
 }
 
@@ -101,7 +101,7 @@ void FShowFlagMenuCommands::CreateCommonShowFlagMenuItems(UToolMenu* Menu, const
 	}
 }
 
-void FShowFlagMenuCommands::CreateSubMenuIfRequired(FToolMenuSection& Section, const FShowFlagFilter& Filter, EShowFlagGroup Group, const FName SubMenuName, const FText& MenuLabel, const FText& ToolTip) const
+void FShowFlagMenuCommands::CreateSubMenuIfRequired(FToolMenuSection& Section, const FShowFlagFilter& Filter, EShowFlagGroup Group, const FName SubMenuName, const FText& MenuLabel, const FText& ToolTip, const FName IconName) const
 {
 	const FShowFlagFilter::FGroupedShowFlagIndices& GroupedFlagIndices = Filter.GetFilteredIndices();
 	const TArray<uint32>& FlagIndices = GroupedFlagIndices[Group];
@@ -111,7 +111,7 @@ void FShowFlagMenuCommands::CreateSubMenuIfRequired(FToolMenuSection& Section, c
 		return;
 	}
 
-	Section.AddSubMenu(SubMenuName, MenuLabel, ToolTip, FNewToolMenuDelegate::CreateStatic(&FShowFlagMenuCommands::StaticCreateShowFlagsSubMenu, FlagIndices, 0));
+	Section.AddSubMenu(SubMenuName, MenuLabel, ToolTip, FNewToolMenuDelegate::CreateStatic(&FShowFlagMenuCommands::StaticCreateShowFlagsSubMenu, FlagIndices, 0), false, FSlateIcon(FAppStyle::Get().GetStyleSetName(), IconName));
 }
 
 void FShowFlagMenuCommands::CreateShowFlagsSubMenu(UToolMenu* Menu, TArray<uint32> FlagIndices, int32 EntryOffset) const
diff --git a/Engine/Source/Editor/CommonMenuExtensions/Public/ShowFlagMenuCommands.h b/Engine/Source/Editor/CommonMenuExtensions/Public/ShowFlagMenuCommands.h
--- a/Engine/Source/Editor/CommonMenuExtensions/Public/ShowFlagMenuCommands.h
+++ b/Engine/Source/Editor/CommonMenuExtensions/Public/ShowFlagMenuCommands.h
@@ -54,7 +54,7 @@ private:
 	void CreateShowFlagCommands();
 	void UpdateCustomShowFlags() const;
 	void CreateCommonShowFlagMenuItems(UToolMenu* Menu, const FShowFlagFilter& Filter) const;
-	void CreateSubMenuIfRequired(FToolMenuSection& Section, const FShowFlagFilter& Filter, EShowFlagGroup Group, const FName SubMenuName, const FText& MenuLabel, const FText& ToolTip) const;
+	void CreateSubMenuIfRequired(FToolMenuSection& Section, const FShowFlagFilter& Filter, EShowFlagGroup Group, const FName SubMenuName, const FText& MenuLabel, const FText& ToolTip, const FName IconName) const;
 	void CreateShowFlagsSubMenu(UToolMenu* Menu, TArray<uint32> MenuCommands, int32 EntryOffset) const;
 
 	TArray<FShowFlagCommand> ShowFlagCommands;
diff --git a/Engine/Source/Editor/EditorStyle/Private/StarshipStyle.cpp b/Engine/Source/Editor/EditorStyle/Private/StarshipStyle.cpp
--- a/Engine/Source/Editor/EditorStyle/Private/StarshipStyle.cpp
+++ b/Engine/Source/Editor/EditorStyle/Private/StarshipStyle.cpp
@@ -4291,6 +4291,22 @@ void FStarshipEditorStyle::FStyle::SetupLevelEditorStyle()
 		Set( "EditorViewport.Right",       new IMAGE_BRUSH_SVG("Starship/Common/ViewRight", Icon16x16 ) );
 		Set( "EditorViewport.Back",        new IMAGE_BRUSH_SVG("Starship/Common/ViewBack", Icon16x16 ) );
 
+		Set("EditorViewport.ToggleStats", new IMAGE_BRUSH_SVG("Starship/Common/Statistics", Icon16x16));
+		Set("EditorViewport.ToggleFPS", new IMAGE_BRUSH_SVG("Starship/Common/FPS", Icon16x16));
+		Set("EditorViewport.ToggleViewportToolbar", new IMAGE_BRUSH_SVG("Starship/Common/Toolbar", Icon16x16));
+
+		Set("EditorViewport.SubMenu.Stats", new IMAGE_BRUSH_SVG("Starship/Common/Statistics", Icon16x16));
+		Set("EditorViewport.SubMenu.Bookmarks", new IMAGE_BRUSH_SVG("Starship/Common/Bookmarks", Icon16x16));
+		Set("EditorViewport.SubMenu.CreateCamera", new IMAGE_BRUSH_SVG("Starship/Common/CreateCamera", Icon16x16));
+		Set("EditorViewport.SubMenu.Layouts", new IMAGE_BRUSH_SVG("Starship/Common/Layout", Icon16x16));
+
+		Set("LevelViewport.ToggleGameView", new IMAGE_BRUSH_SVG("Starship/Common/GameView", Icon16x16));
+		Set("LevelViewport.ToggleImmersive", new IMAGE_BRUSH_SVG("Starship/Common/EnableFullscreen", Icon16x16));
+		Set("LevelViewport.HighResScreenshot", new IMAGE_BRUSH_SVG("Starship/Common/HighResolutionScreenshot", Icon16x16));
+		Set("LevelViewport.ToggleCinematicPreview", new IMAGE_BRUSH_SVG("Starship/Common/Cinematics", Icon16x16));
+		Set("LevelViewport.ToggleCameraShakePreview", new IMAGE_BRUSH_SVG("Starship/Common/CameraShake", Icon16x16));
+		Set("LevelViewport.AdvancedSettings", new CORE_IMAGE_BRUSH_SVG("Starship/Common/settings", Icon16x16));
+
 		Set("EditorViewport.ActiveBorderColor", FStyleColors::Primary);
 
 #endif
@@ -4709,21 +4725,34 @@ void FStarshipEditorStyle::FStyle::SetupLevelEditorStyle()
 
 	// Show flags menus
 	{
-		Set( "ShowFlagsMenu.AntiAliasing", new IMAGE_BRUSH( "Icons/icon_ShowAnti-aliasing_16x", Icon16x16 ) );
-		Set( "ShowFlagsMenu.Atmosphere", new IMAGE_BRUSH( "Icons/icon_ShowAtmosphere_16x", Icon16x16 ) );
-		Set( "ShowFlagsMenu.BSP", new IMAGE_BRUSH( "Icons/icon_ShowBSP_16x", Icon16x16 ) );
-		Set( "ShowFlagsMenu.Collision", new IMAGE_BRUSH( "Icons/icon_ShowCollision_16x", Icon16x16 ) );
-		Set( "ShowFlagsMenu.Decals", new IMAGE_BRUSH( "Icons/icon_ShowDecals_16x", Icon16x16 ) );
-		Set( "ShowFlagsMenu.Fog", new IMAGE_BRUSH( "Icons/icon_ShowFog_16x", Icon16x16 ) );
-		Set( "ShowFlagsMenu.Grid", new IMAGE_BRUSH( "Icons/icon_ShowGrid_16x", Icon16x16 ) );
-		Set( "ShowFlagsMenu.Landscape", new IMAGE_BRUSH( "Icons/icon_ShowLandscape_16x", Icon16x16 ) );
-		Set( "ShowFlagsMenu.MediaPlanes", new IMAGE_BRUSH( "Icons/icon_ShowMediaPlanes_16x", Icon16x16 ) );
-		Set( "ShowFlagsMenu.Navigation", new IMAGE_BRUSH( "Icons/icon_ShowNavigation_16x", Icon16x16 ) );
-		Set( "ShowFlagsMenu.Particles", new IMAGE_BRUSH( "Icons/icon_ShowParticlesSprite_16x", Icon16x16 ) );
-		Set( "ShowFlagsMenu.SkeletalMeshes", new IMAGE_BRUSH( "Icons/icon_ShowSkeletalMeshes_16x", Icon16x16 ) );
-		Set( "ShowFlagsMenu.StaticMeshes", new IMAGE_BRUSH( "Icons/icon_ShowStaticMeshes_16x", Icon16x16 ) );
-		Set( "ShowFlagsMenu.Translucency", new IMAGE_BRUSH( "Icons/icon_ShowTranslucency_16x", Icon16x16 ) );
-		Set( "ShowFlagsMenu.WidgetComponents", new IMAGE_BRUSH( "UMG/Designer_16x", Icon16x16 ) );
+		Set( "ShowFlagsMenu.AntiAliasing", new IMAGE_BRUSH_SVG( "Starship/Common/AntiAliasing", Icon16x16 ) );
+		Set( "ShowFlagsMenu.Atmosphere", new IMAGE_BRUSH_SVG( "Starship/Common/Atmosphere", Icon16x16 ) );
+		Set( "ShowFlagsMenu.BSP", new IMAGE_BRUSH_SVG( "Starship/Common/BSP", Icon16x16 ) );
+		Set( "ShowFlagsMenu.Collision", new IMAGE_BRUSH_SVG( "Starship/Common/Collision", Icon16x16 ) );
+		Set( "ShowFlagsMenu.Decals", new IMAGE_BRUSH_SVG( "Starship/Common/Decals", Icon16x16 ) );
+		Set( "ShowFlagsMenu.Fog", new IMAGE_BRUSH_SVG( "Starship/Common/Fog", Icon16x16 ) );
+		Set( "ShowFlagsMenu.Grid", new IMAGE_BRUSH_SVG( "Starship/Common/Grid", Icon16x16 ) );
+		Set( "ShowFlagsMenu.Landscape", new IMAGE_BRUSH_SVG( "Starship/Common/Landscape", Icon16x16 ) );
+		Set( "ShowFlagsMenu.MediaPlanes", new IMAGE_BRUSH_SVG( "Starship/Common/MediaPlanes", Icon16x16 ) );
+		Set( "ShowFlagsMenu.Navigation", new IMAGE_BRUSH_SVG( "Starship/Common/Navigation", Icon16x16 ) );
+		Set( "ShowFlagsMenu.Particles", new IMAGE_BRUSH_SVG( "Starship/Common/ParticleSprites", Icon16x16 ) );
+		Set( "ShowFlagsMenu.SkeletalMeshes", new IMAGE_BRUSH_SVG( "Starship/Common/SkeletalMesh", Icon16x16 ) );
+		Set( "ShowFlagsMenu.StaticMeshes", new IMAGE_BRUSH_SVG( "Starship/Common/StaticMesh", Icon16x16 ) );
+		Set( "ShowFlagsMenu.Translucency", new IMAGE_BRUSH_SVG( "Starship/Common/Transparency", Icon16x16 ) );
+		Set( "ShowFlagsMenu.WidgetComponents", new IMAGE_BRUSH_SVG( "Starship/Common/WidgetComponents", Icon16x16 ) );
+
+		Set("ShowFlagsMenu.SubMenu.PostProcessing", new IMAGE_BRUSH_SVG("Starship/Common/PostProcessing", Icon16x16));
+		Set("ShowFlagsMenu.SubMenu.LightTypes", new IMAGE_BRUSH_SVG("Starship/Common/LightTypes", Icon16x16));
+		Set("ShowFlagsMenu.SubMenu.LightingComponents", new IMAGE_BRUSH_SVG("Starship/Common/LightingComponents", Icon16x16));
+		Set("ShowFlagsMenu.SubMenu.LightingFeatures", new IMAGE_BRUSH_SVG("Starship/Common/LightingFeatures", Icon16x16));
+		Set("ShowFlagsMenu.SubMenu.Developer", new IMAGE_BRUSH_SVG("Starship/Common/Developer", Icon16x16));
+		Set("ShowFlagsMenu.SubMenu.Visualize", new IMAGE_BRUSH_SVG("Starship/Common/Visualize", Icon16x16));
+		Set("ShowFlagsMenu.SubMenu.Advanced", new IMAGE_BRUSH_SVG("Starship/Common/Advanced", Icon16x16));
+
+		Set("ShowFlagsMenu.SubMenu.Volumes", new IMAGE_BRUSH_SVG("Starship/Common/Volume", Icon16x16));
+		Set("ShowFlagsMenu.SubMenu.Layers", new IMAGE_BRUSH_SVG("Starship/Common/Layers", Icon16x16));
+		Set("ShowFlagsMenu.SubMenu.FoliageTypes", new IMAGE_BRUSH_SVG("Starship/Common/FoliageTypes", Icon16x16));
+		Set("ShowFlagsMenu.SubMenu.Sprites", new IMAGE_BRUSH_SVG("Starship/Common/Sprite", Icon16x16));
 	}
 #endif // WITH_EDITOR || (IS_PROGRAM && WITH_UNREAL_DEVELOPER_TOOLS)
 
diff --git a/Engine/Source/Editor/LevelEditor/Private/SLevelViewportToolBar.cpp b/Engine/Source/Editor/LevelEditor/Private/SLevelViewportToolBar.cpp
--- a/Engine/Source/Editor/LevelEditor/Private/SLevelViewportToolBar.cpp
+++ b/Engine/Source/Editor/LevelEditor/Private/SLevelViewportToolBar.cpp
@@ -580,7 +580,9 @@ static void OnGenerateBookmarkMenu(UToolMenu* Menu, TWeakPtr<class SLevelViewpor
 				"ClearBookmark",
 				LOCTEXT("ClearBookmarkSubMenu", "Clear Bookmark"),
 				LOCTEXT("ClearBookmarkSubMenu_ToolTip", "Clear viewport bookmarks"),
-				FNewToolMenuDelegate::CreateStatic( &OnGenerateClearBookmarkMenu, Viewport )
+				FNewToolMenuDelegate::CreateStatic( &OnGenerateClearBookmarkMenu, Viewport ),
+				false,
+				FSlateIcon(FAppStyle::Get().GetStyleSetName(), "EditorViewport.SubMenu.Bookmarks")
 				);
 
 			Section.AddMenuEntry( Actions.ClearAllBookmarks );
@@ -646,21 +648,28 @@ void SLevelViewportToolBar::FillOptionsMenu(UToolMenu* Menu)
 					
 			}
 
-			Section.AddMenuEntry(FEditorViewportCommands::Get().ToggleStats);
 			Section.AddMenuEntry(FEditorViewportCommands::Get().ToggleFPS);
-			Section.AddMenuEntry(LevelViewportActions.ToggleViewportToolbar);
 
-			FText HideAllLabel = LOCTEXT("HideAllLabel", "Hide All");
 			TArray< FLevelViewportCommands::FShowMenuCommand > HideStatsMenu;
 
-			// 'Hide All' button
-			HideStatsMenu.Add(FLevelViewportCommands::FShowMenuCommand(LevelViewportActions.HideAllStats, HideAllLabel));
-
+			Section.AddMenuEntry(FEditorViewportCommands::Get().ToggleStats);
 			Section.AddSubMenu(
 				"ShowStatsMenu",
-				LOCTEXT("ShowStatsMenu", "Stat"), 
+				LOCTEXT("ShowStatsMenu", "Stat"),
 				LOCTEXT("ShowStatsMenu_ToolTip", "Show Stat commands"),
-				FNewToolMenuDelegate::CreateStatic(&FillShowStatsSubMenus, HideStatsMenu, LevelViewportActions.ShowStatCatCommands));
+				FNewToolMenuDelegate::CreateStatic(&FillShowStatsSubMenus, HideStatsMenu, LevelViewportActions.ShowStatCatCommands),
+				false,
+				FSlateIcon(FAppStyle::Get().GetStyleSetName(), "EditorViewport.SubMenu.Stats"));
+
+
+
+
+			FText HideAllLabel = LOCTEXT("HideAllLabel", "Hide All");
+		
+			// 'Hide All' button
+			HideStatsMenu.Add(FLevelViewportCommands::FShowMenuCommand(LevelViewportActions.HideAllStats, HideAllLabel));
+
+			Section.AddMenuEntry(LevelViewportActions.ToggleViewportToolbar);
 
 			if( bIsPerspective )
 			{
@@ -685,6 +694,7 @@ void SLevelViewportToolBar::FillOptionsMenu(UToolMenu* Menu)
 				bAddDefaultScreenPercentageSlider = !GCustomEditorStaticScreenPercentage->GenerateEditorViewportOptionsMenuEntry(Arguments);
 			}
 
+
 			if (bAddDefaultScreenPercentageSlider)
 			{
 				Section.AddEntry(FToolMenuEntry::InitWidget("ScreenPercentage", GenerateScreenPercentageMenu(), LOCTEXT("ScreenPercentage", "Screen Percentage")));
@@ -715,14 +725,18 @@ void SLevelViewportToolBar::FillOptionsMenu(UToolMenu* Menu)
 					"Bookmark",
 					LOCTEXT("BookmarkSubMenu", "Bookmarks"),
 					LOCTEXT("BookmarkSubMenu_ToolTip", "Viewport location bookmarking"),
-					FNewToolMenuDelegate::CreateStatic(&OnGenerateBookmarkMenu, Viewport)
+					FNewToolMenuDelegate::CreateStatic(&OnGenerateBookmarkMenu, Viewport),
+					false,
+					FSlateIcon(FAppStyle::Get().GetStyleSetName(), "EditorViewport.SubMenu.Bookmarks")
 				);
 
 				Section.AddSubMenu(
 					"Camera",
 					LOCTEXT("CameraSubMeun", "Create Camera Here"),
 					LOCTEXT("CameraSubMenu_ToolTip", "Select a camera type to create at current viewport's location"),
-					FNewToolMenuDelegate::CreateSP(this, &SLevelViewportToolBar::GenerateCameraSpawnMenu)
+					FNewToolMenuDelegate::CreateSP(this, &SLevelViewportToolBar::GenerateCameraSpawnMenu),
+					false,
+					FSlateIcon(FAppStyle::Get().GetStyleSetName(), "EditorViewport.SubMenu.CreateCamera")
 				);
 			}
 
@@ -736,7 +750,10 @@ void SLevelViewportToolBar::FillOptionsMenu(UToolMenu* Menu)
 				"Configs",
 				LOCTEXT("ConfigsSubMenu", "Layouts"),
 				FText::GetEmpty(),
-				FNewToolMenuDelegate::CreateSP(this, &SLevelViewportToolBar::GenerateViewportConfigsMenu));
+				FNewToolMenuDelegate::CreateSP(this, &SLevelViewportToolBar::GenerateViewportConfigsMenu),
+				false,
+				FSlateIcon(FAppStyle::Get().GetStyleSetName(), "EditorViewport.SubMenu.Layouts")
+			);
 		}
 
 		{
@@ -1170,14 +1187,14 @@ void SLevelViewportToolBar::FillShowMenu(UToolMenu* Menu) const
 				ShowVolumesMenu += Actions.ShowVolumeCommands;
 
 				Section.AddSubMenu("ShowVolumes", LOCTEXT("ShowVolumesMenu", "Volumes"), LOCTEXT("ShowVolumesMenu_ToolTip", "Show volumes flags"),
-					FNewToolMenuDelegate::CreateStatic(&FillShowMenuStatic, ShowVolumesMenu, 2));
+					FNewToolMenuDelegate::CreateStatic(&FillShowMenuStatic, ShowVolumesMenu, 2), false, FSlateIcon(FAppStyle::Get().GetStyleSetName(), "ShowFlagsMenu.SubMenu.Volumes"));
 			}
 
 			// Show Layers sub-menu is dynamically generated when the user enters 'show' menu
 			if (!UWorld::HasSubsystem<UWorldPartitionSubsystem>(World))
 			{
 				Section.AddSubMenu("ShowLayers", LOCTEXT("ShowLayersMenu", "Layers"), LOCTEXT("ShowLayersMenu_ToolTip", "Show layers flags"),
-					FNewToolMenuDelegate::CreateStatic(&SLevelViewportToolBar::FillShowLayersMenu, Viewport));
+					FNewToolMenuDelegate::CreateStatic(&SLevelViewportToolBar::FillShowLayersMenu, Viewport), false, FSlateIcon(FAppStyle::Get().GetStyleSetName(), "ShowFlagsMenu.SubMenu.Layers"));
 			}
 			else
 			{
@@ -1199,13 +1216,13 @@ void SLevelViewportToolBar::FillShowMenu(UToolMenu* Menu) const
 				ShowSpritesMenu += Actions.ShowSpriteCommands;
 
 				Section.AddSubMenu("ShowSprites", LOCTEXT("ShowSpritesMenu", "Sprites"), LOCTEXT("ShowSpritesMenu_ToolTip", "Show sprites flags"),
-					FNewToolMenuDelegate::CreateStatic(&FillShowMenuStatic, ShowSpritesMenu, 2));
+					FNewToolMenuDelegate::CreateStatic(&FillShowMenuStatic, ShowSpritesMenu, 2), false, FSlateIcon(FAppStyle::Get().GetStyleSetName(), "ShowFlagsMenu.SubMenu.Sprites"));
 			}
 
 			// Show 'Foliage types' sub-menu is dynamically generated when the user enters 'show' menu
 			{
 				Section.AddSubMenu("ShowFoliageTypes", LOCTEXT("ShowFoliageTypesMenu", "Foliage Types"), LOCTEXT("ShowFoliageTypesMenu_ToolTip", "Show/hide specific foliage types"),
-					FNewToolMenuDelegate::CreateStatic(&SLevelViewportToolBar::FillShowFoliageTypesMenu, Viewport));
+					FNewToolMenuDelegate::CreateStatic(&SLevelViewportToolBar::FillShowFoliageTypesMenu, Viewport), false, FSlateIcon(FAppStyle::Get().GetStyleSetName(), "ShowFlagsMenu.SubMenu.FoliageTypes"));
 			}
 
 		}
diff --git a/Engine/Source/Runtime/Slate/Private/Framework/MultiBox/SWidgetBlock.cpp b/Engine/Source/Runtime/Slate/Private/Framework/MultiBox/SWidgetBlock.cpp
--- a/Engine/Source/Runtime/Slate/Private/Framework/MultiBox/SWidgetBlock.cpp
+++ b/Engine/Source/Runtime/Slate/Private/Framework/MultiBox/SWidgetBlock.cpp
@@ -90,7 +90,7 @@ void SWidgetBlock::BuildMultiBlockWidget(const ISlateStyle* StyleSet, const FNam
 	if(OwnerMultiBoxWidgetPinned->GetMultiBox()->GetType() == EMultiBoxType::Menu)
 	{
 		// Account for checkmark used in other menu blocks but not used in for widget rows
-		Padding = Padding + FMargin(12, 0, 8, 0);
+		Padding = Padding + FMargin(14, 0, 8, 0);
 	}
 
 	// Add this widget to the search list of the multibox
