Ported the actor snap/align functions to use typed elements internally

The following options from the actor viewport menu are now elemenent based:
 - Snap/Align (excluding layers)
 - Delta Transform
 - Mirror X/Y/Z

diff --git a/Engine/Source/Editor/EditorFramework/Private/Elements/Framework/TypedElementViewportInteraction.cpp b/Engine/Source/Editor/EditorFramework/Private/Elements/Framework/TypedElementViewportInteraction.cpp
--- a/Engine/Source/Editor/EditorFramework/Private/Elements/Framework/TypedElementViewportInteraction.cpp
+++ b/Engine/Source/Editor/EditorFramework/Private/Elements/Framework/TypedElementViewportInteraction.cpp
@@ -27,6 +27,7 @@ void FTypedElementAssetEditorViewportInteractionCustomization::PreGizmoManipulat
 
 void FTypedElementAssetEditorViewportInteractionCustomization::GizmoManipulationStarted(const TTypedElement<UTypedElementWorldInterface>& InElementWorldHandle, const UE::Widget::EWidgetMode InWidgetMode)
 {
+	InElementWorldHandle.NotifyMovementStarted();
 }
 
 void FTypedElementAssetEditorViewportInteractionCustomization::GizmoManipulationDeltaUpdate(const TTypedElement<UTypedElementWorldInterface>& InElementWorldHandle, const UE::Widget::EWidgetMode InWidgetMode, const EAxisList::Type InDragAxis, const FInputDeviceState& InInputState, const FTransform& InDeltaTransform, const FVector& InPivotLocation)
@@ -71,17 +72,68 @@ void FTypedElementAssetEditorViewportInteractionCustomization::GizmoManipulation
 		}
 
 		InElementWorldHandle.SetWorldTransform(ElementWorldTransform);
+		InElementWorldHandle.NotifyMovementOngoing();
 	}
 }
 
 void FTypedElementAssetEditorViewportInteractionCustomization::GizmoManipulationStopped(const TTypedElement<UTypedElementWorldInterface>& InElementWorldHandle, const UE::Widget::EWidgetMode InWidgetMode)
 {
+	InElementWorldHandle.NotifyMovementEnded();
 }
 
 void FTypedElementAssetEditorViewportInteractionCustomization::PostGizmoManipulationStopped(TArrayView<const FTypedElementHandle> InElementHandles, const UE::Widget::EWidgetMode InWidgetMode)
 {
 }
 
+void FTypedElementAssetEditorViewportInteractionCustomization::MirrorElement(const TTypedElement<UTypedElementWorldInterface>& InElementWorldHandle, const FVector& InMirrorScale, const FVector& InPivotLocation)
+{
+	FTransform ElementWorldTransform;
+	if (InElementWorldHandle.GetWorldTransform(ElementWorldTransform))
+	{
+		InElementWorldHandle.NotifyMovementStarted();
+
+		// Apply mirrored rotation
+		{
+			// Revert the handedness of the rotation, but make up for it in the scaling
+			// Arbitrarily choose the X axis to remain fixed
+			const FMatrix TempRot = FRotationMatrix::Make(ElementWorldTransform.GetRotation());
+			const FMatrix NewRot(
+				-TempRot.GetScaledAxis(EAxis::X) * InMirrorScale, 
+				TempRot.GetScaledAxis(EAxis::Y) * InMirrorScale, 
+				TempRot.GetScaledAxis(EAxis::Z) * InMirrorScale, 
+				FVector::ZeroVector
+				);
+			ElementWorldTransform.SetRotation(NewRot.ToQuat());
+		}
+
+		// Apply mirrored location around the pivot location
+		{
+			FVector Loc = ElementWorldTransform.GetTranslation();
+			Loc -= InPivotLocation;
+			Loc *= InMirrorScale;
+			Loc += InPivotLocation;
+			ElementWorldTransform.SetTranslation(Loc);
+		}
+
+		InElementWorldHandle.SetWorldTransform(ElementWorldTransform);
+
+		// Apply mirrored relative scale
+		{
+			FTransform ElementRelativeTransform;
+			if (InElementWorldHandle.GetRelativeTransform(ElementRelativeTransform))
+			{
+				FVector Scale3D = ElementRelativeTransform.GetScale3D();
+				Scale3D.X = -Scale3D.X;
+				ElementRelativeTransform.SetScale3D(Scale3D);
+
+				InElementWorldHandle.SetRelativeTransform(ElementRelativeTransform);
+			}
+		}
+
+		InElementWorldHandle.NotifyMovementEnded();
+	}
+}
+
 
 void UTypedElementViewportInteraction::GetSelectedElementsToMove(const UTypedElementSelectionSet* InSelectionSet, const ETypedElementViewportInteractionWorldType InWorldType, UTypedElementList* OutElementsToMove) const
 {
@@ -174,6 +226,17 @@ void UTypedElementViewportInteraction::ApplyDeltaToElement(const FTypedElementHa
 	}
 }
 
+void UTypedElementViewportInteraction::MirrorElement(const FTypedElementHandle& InElementHandle, const FVector& InMirrorScale)
+{
+	FTypedElementViewportInteractionElement ViewportInteractionElement = ResolveViewportInteractionElement(InElementHandle);
+	if (ViewportInteractionElement)
+	{
+		FVector PivotLocation = FVector::ZeroVector;
+		ViewportInteractionElement.GetGizmoPivotLocation(UE::Widget::WM_None, PivotLocation);
+		ViewportInteractionElement.MirrorElement(InMirrorScale, PivotLocation);
+	}
+}
+
 FTypedElementViewportInteractionElement UTypedElementViewportInteraction::ResolveViewportInteractionElement(const FTypedElementHandle& InElementHandle) const
 {
 	return InElementHandle
diff --git a/Engine/Source/Editor/EditorFramework/Public/Elements/Framework/TypedElementViewportInteraction.h b/Engine/Source/Editor/EditorFramework/Public/Elements/Framework/TypedElementViewportInteraction.h
--- a/Engine/Source/Editor/EditorFramework/Public/Elements/Framework/TypedElementViewportInteraction.h
+++ b/Engine/Source/Editor/EditorFramework/Public/Elements/Framework/TypedElementViewportInteraction.h
@@ -33,6 +33,7 @@ public:
 	virtual void GizmoManipulationDeltaUpdate(const TTypedElement<UTypedElementWorldInterface>& InElementWorldHandle, const UE::Widget::EWidgetMode InWidgetMode, const EAxisList::Type InDragAxis, const FInputDeviceState& InInputState, const FTransform& InDeltaTransform, const FVector& InPivotLocation);
 	virtual void GizmoManipulationStopped(const TTypedElement<UTypedElementWorldInterface>& InElementWorldHandle, const UE::Widget::EWidgetMode InWidgetMode);
 	virtual void PostGizmoManipulationStopped(TArrayView<const FTypedElementHandle> InElementHandles, const UE::Widget::EWidgetMode InWidgetMode);
+	virtual void MirrorElement(const TTypedElement<UTypedElementWorldInterface>& InElementWorldHandle, const FVector& InMirrorScale, const FVector& InPivotLocation);
 };
 
 /**
@@ -67,11 +68,12 @@ public:
 	}
 
 	//~ See UTypedElementViewportInteraction for API docs
-	void GetElementsToMove(const ETypedElementViewportInteractionWorldType InWorldType, const UTypedElementSelectionSet* InSelectionSet, UTypedElementList* OutElementsToMove) { return AssetEditorViewportInteractionCustomization->GetElementsToMove(ElementWorldHandle, InWorldType, InSelectionSet, OutElementsToMove); }
+	void GetElementsToMove(const ETypedElementViewportInteractionWorldType InWorldType, const UTypedElementSelectionSet* InSelectionSet, UTypedElementList* OutElementsToMove) { AssetEditorViewportInteractionCustomization->GetElementsToMove(ElementWorldHandle, InWorldType, InSelectionSet, OutElementsToMove); }
 	bool GetGizmoPivotLocation(const UE::Widget::EWidgetMode InWidgetMode, FVector& OutPivotLocation) const { return AssetEditorViewportInteractionCustomization->GetGizmoPivotLocation(ElementWorldHandle, InWidgetMode, OutPivotLocation); }
-	void GizmoManipulationStarted(const UE::Widget::EWidgetMode InWidgetMode) const { return AssetEditorViewportInteractionCustomization->GizmoManipulationStarted(ElementWorldHandle, InWidgetMode); }
-	void GizmoManipulationDeltaUpdate(const UE::Widget::EWidgetMode InWidgetMode, const EAxisList::Type InDragAxis, const FInputDeviceState& InInputState, const FTransform& InDeltaTransform, const FVector& InPivotLocation) const { return AssetEditorViewportInteractionCustomization->GizmoManipulationDeltaUpdate(ElementWorldHandle, InWidgetMode, InDragAxis, InInputState, InDeltaTransform, InPivotLocation); }
-	void GizmoManipulationStopped(const UE::Widget::EWidgetMode InWidgetMode) const { return AssetEditorViewportInteractionCustomization->GizmoManipulationStopped(ElementWorldHandle, InWidgetMode); }
+	void GizmoManipulationStarted(const UE::Widget::EWidgetMode InWidgetMode) const { AssetEditorViewportInteractionCustomization->GizmoManipulationStarted(ElementWorldHandle, InWidgetMode); }
+	void GizmoManipulationDeltaUpdate(const UE::Widget::EWidgetMode InWidgetMode, const EAxisList::Type InDragAxis, const FInputDeviceState& InInputState, const FTransform& InDeltaTransform, const FVector& InPivotLocation) const { AssetEditorViewportInteractionCustomization->GizmoManipulationDeltaUpdate(ElementWorldHandle, InWidgetMode, InDragAxis, InInputState, InDeltaTransform, InPivotLocation); }
+	void GizmoManipulationStopped(const UE::Widget::EWidgetMode InWidgetMode) const { AssetEditorViewportInteractionCustomization->GizmoManipulationStopped(ElementWorldHandle, InWidgetMode); }
+	void MirrorElement(const FVector& InMirrorScale, const FVector& InPivotLocation) const { AssetEditorViewportInteractionCustomization->MirrorElement(ElementWorldHandle, InMirrorScale, InPivotLocation); }
 
 private:
 	TTypedElement<UTypedElementWorldInterface> ElementWorldHandle;
@@ -114,6 +116,11 @@ public:
 	 */
 	void ApplyDeltaToElement(const FTypedElementHandle& InElementHandle, const UE::Widget::EWidgetMode InWidgetMode, const EAxisList::Type InDragAxis, const FInputDeviceState& InInputState, const FTransform& InDeltaTransform);
 
+	/**
+	 * Apply the given mirror scale to the specified element.
+	 */
+	void MirrorElement(const FTypedElementHandle& InElementHandle, const FVector& InMirrorScale);
+
 private:
 	/**
 	 * Attempt to resolve the selection interface and asset editor viewport interaction customization for the given element, if any.
diff --git a/Engine/Source/Editor/LevelEditor/Private/DlgDeltaTransform.cpp b/Engine/Source/Editor/LevelEditor/Private/DlgDeltaTransform.cpp
--- a/Engine/Source/Editor/LevelEditor/Private/DlgDeltaTransform.cpp
+++ b/Engine/Source/Editor/LevelEditor/Private/DlgDeltaTransform.cpp
@@ -191,7 +191,7 @@ private:
 		if (ButtonID == FDlgDeltaTransform::OK)
 		{
 			GUnrealEd->Exec( GEditor->GetEditorWorldContext().World(), 
-				*FString::Printf(TEXT("ACTOR DELTAMOVE X=%.5f Y=%.5f Z=%.5f"), 
+				*FString::Printf(TEXT("ELEMENT DELTAMOVE X=%.5f Y=%.5f Z=%.5f"), 
 				DeltaTransform.X.GetValue(), DeltaTransform.Y.GetValue(), DeltaTransform.Z.GetValue()) );
 		}
 
diff --git a/Engine/Source/Editor/LevelEditor/Private/LevelEditor.cpp b/Engine/Source/Editor/LevelEditor/Private/LevelEditor.cpp
--- a/Engine/Source/Editor/LevelEditor/Private/LevelEditor.cpp
+++ b/Engine/Source/Editor/LevelEditor/Private/LevelEditor.cpp
@@ -704,10 +704,9 @@ void FLevelEditorModule::BindGlobalLevelEditorCommands()
 		FExecuteAction::CreateStatic( &FLevelEditorActionCallbacks::FindInContentBrowser_Clicked )
 		);
 
-	const FVector* NullVector = nullptr;
 	ActionList.MapAction(
 		Commands.GoHere,
-		FExecuteAction::CreateStatic( &FLevelEditorActionCallbacks::GoHere_Clicked, NullVector )
+		FExecuteAction::CreateStatic( &FLevelEditorActionCallbacks::GoHere_Clicked, (const FVector*)nullptr )
 		);
 
 	ActionList.MapAction( 
@@ -772,41 +771,34 @@ void FLevelEditorModule::BindGlobalLevelEditorCommands()
 		FCanExecuteAction::CreateStatic( &FLevelEditorActionCallbacks::PasteHere_CanExecute )
 		);
 
-	bool bAlign = false;
-	bool bPerActor = false;
 	ActionList.MapAction(
 		Commands.SnapOriginToGrid,
-		FExecuteAction::CreateStatic(&FLevelEditorActionCallbacks::MoveActorToGrid_Clicked, bAlign, bPerActor),
-		FCanExecuteAction::CreateStatic(&FLevelEditorActionCallbacks::ActorSelected_CanExecute)
+		FExecuteAction::CreateStatic(&FLevelEditorActionCallbacks::MoveElementsToGrid_Clicked, /*bAlign*/false, /*bPerElement*/false),
+		FCanExecuteAction::CreateStatic(&FLevelEditorActionCallbacks::ElementSelected_CanExecute)
 		);
 
-	bPerActor = true;
 	ActionList.MapAction(
 		Commands.SnapOriginToGridPerActor,
-		FExecuteAction::CreateStatic(&FLevelEditorActionCallbacks::MoveActorToGrid_Clicked, bAlign, bPerActor),
-		FCanExecuteAction::CreateStatic(&FLevelEditorActionCallbacks::ActorSelected_CanExecute)
+		FExecuteAction::CreateStatic(&FLevelEditorActionCallbacks::MoveElementsToGrid_Clicked, /*bAlign*/false, /*bPerElement*/true),
+		FCanExecuteAction::CreateStatic(&FLevelEditorActionCallbacks::ElementSelected_CanExecute)
 		);
 	
-	bAlign = true;
-	bPerActor = false;
 	ActionList.MapAction(
 		Commands.AlignOriginToGrid,
-		FExecuteAction::CreateStatic(&FLevelEditorActionCallbacks::MoveActorToGrid_Clicked, bAlign, bPerActor),
-		FCanExecuteAction::CreateStatic(&FLevelEditorActionCallbacks::ActorSelected_CanExecute)
+		FExecuteAction::CreateStatic(&FLevelEditorActionCallbacks::MoveElementsToGrid_Clicked, /*bAlign*/true, /*bPerElement*/false),
+		FCanExecuteAction::CreateStatic(&FLevelEditorActionCallbacks::ElementSelected_CanExecute)
 		);
 
-	bAlign = false;
 	ActionList.MapAction(
 		Commands.SnapOriginToActor,
-		FExecuteAction::CreateStatic(&FLevelEditorActionCallbacks::MoveActorToActor_Clicked, bAlign),
-		FCanExecuteAction::CreateStatic(&FLevelEditorActionCallbacks::ActorsSelected_CanExecute)
+		FExecuteAction::CreateStatic(&FLevelEditorActionCallbacks::MoveElementsToElement_Clicked, /*bAlign*/false),
+		FCanExecuteAction::CreateStatic(&FLevelEditorActionCallbacks::ElementsSelected_CanExecute)
 		);
 	
-	bAlign = true;
 	ActionList.MapAction(
 		Commands.AlignOriginToActor,
-		FExecuteAction::CreateStatic(&FLevelEditorActionCallbacks::MoveActorToActor_Clicked, bAlign),
-		FCanExecuteAction::CreateStatic(&FLevelEditorActionCallbacks::ActorsSelected_CanExecute)
+		FExecuteAction::CreateStatic(&FLevelEditorActionCallbacks::MoveElementsToElement_Clicked, /*bAlign*/true),
+		FCanExecuteAction::CreateStatic(&FLevelEditorActionCallbacks::ElementsSelected_CanExecute)
 		);
 
 	ActionList.MapAction(
@@ -850,147 +842,99 @@ void FLevelEditorModule::BindGlobalLevelEditorCommands()
 		FExecuteAction::CreateStatic(&FLevelEditorActionCallbacks::AlignBrushVerticesToGrid_Execute)
 		);
 
-	bAlign = false;
-	bool bUseLineTrace = false;
-	bool bUseBounds = false;
-	bool bUsePivot = false;
 	ActionList.MapAction(
 		Commands.SnapToFloor,
-		FExecuteAction::CreateStatic(&FLevelEditorActionCallbacks::SnapToFloor_Clicked, bAlign, bUseLineTrace, bUseBounds, bUsePivot),
-		FCanExecuteAction::CreateStatic(&FLevelEditorActionCallbacks::ActorSelected_CanExecute)
+		FExecuteAction::CreateStatic(&FLevelEditorActionCallbacks::SnapToFloor_Clicked, /*bAlign*/false, /*bUseLineTrace*/false, /*bUseBounds*/false, /*bUsePivot*/false),
+		FCanExecuteAction::CreateStatic(&FLevelEditorActionCallbacks::ElementSelected_CanExecute)
 		);
 
-	bAlign = true;
-	bUseLineTrace = false;
-	bUseBounds = false;
-	bUsePivot = false;
 	ActionList.MapAction(
 		Commands.AlignToFloor,
-		FExecuteAction::CreateStatic(&FLevelEditorActionCallbacks::SnapToFloor_Clicked, bAlign, bUseLineTrace, bUseBounds, bUsePivot),
-		FCanExecuteAction::CreateStatic(&FLevelEditorActionCallbacks::ActorSelected_CanExecute)
+		FExecuteAction::CreateStatic(&FLevelEditorActionCallbacks::SnapToFloor_Clicked, /*bAlign*/true, /*bUseLineTrace*/false, /*bUseBounds*/false, /*bUsePivot*/false),
+		FCanExecuteAction::CreateStatic(&FLevelEditorActionCallbacks::ElementSelected_CanExecute)
 		);
 
-	bAlign = false;
-	bUseLineTrace = true;
-	bUseBounds = false;
-	bUsePivot = true;
 	ActionList.MapAction(
 		Commands.SnapPivotToFloor,
-		FExecuteAction::CreateStatic(&FLevelEditorActionCallbacks::SnapToFloor_Clicked, bAlign, bUseLineTrace, bUseBounds, bUsePivot),
-		FCanExecuteAction::CreateStatic(&FLevelEditorActionCallbacks::ActorSelected_CanExecute)
+		FExecuteAction::CreateStatic(&FLevelEditorActionCallbacks::SnapToFloor_Clicked, /*bAlign*/false, /*bUseLineTrace*/true, /*bUseBounds*/false, /*bUsePivot*/true),
+		FCanExecuteAction::CreateStatic(&FLevelEditorActionCallbacks::ElementSelected_CanExecute)
 		);
 
-	bAlign = true;
-	bUseLineTrace = true;
-	bUseBounds = false;
-	bUsePivot = true;
 	ActionList.MapAction(
 		Commands.AlignPivotToFloor,
-		FExecuteAction::CreateStatic(&FLevelEditorActionCallbacks::SnapToFloor_Clicked, bAlign, bUseLineTrace, bUseBounds, bUsePivot),
-		FCanExecuteAction::CreateStatic(&FLevelEditorActionCallbacks::ActorSelected_CanExecute)
+		FExecuteAction::CreateStatic(&FLevelEditorActionCallbacks::SnapToFloor_Clicked, /*bAlign*/true, /*bUseLineTrace*/true, /*bUseBounds*/false, /*bUsePivot*/true),
+		FCanExecuteAction::CreateStatic(&FLevelEditorActionCallbacks::ElementSelected_CanExecute)
 		);
 
-	bAlign = false;
-	bUseLineTrace = true;
-	bUseBounds = true;
-	bUsePivot = false;
 	ActionList.MapAction(
 		Commands.SnapBottomCenterBoundsToFloor,
-		FExecuteAction::CreateStatic(&FLevelEditorActionCallbacks::SnapToFloor_Clicked, bAlign, bUseLineTrace, bUseBounds, bUsePivot),
-		FCanExecuteAction::CreateStatic(&FLevelEditorActionCallbacks::ActorSelected_CanExecute)
+		FExecuteAction::CreateStatic(&FLevelEditorActionCallbacks::SnapToFloor_Clicked, /*bAlign*/false, /*bUseLineTrace*/true, /*bUseBounds*/true, /*bUsePivot*/false),
+		FCanExecuteAction::CreateStatic(&FLevelEditorActionCallbacks::ElementSelected_CanExecute)
 		);
 
-	bAlign = true;
-	bUseLineTrace = true;
-	bUseBounds = true;
-	bUsePivot = false;
 	ActionList.MapAction(
 		Commands.AlignBottomCenterBoundsToFloor,
-		FExecuteAction::CreateStatic(&FLevelEditorActionCallbacks::SnapToFloor_Clicked, bAlign, bUseLineTrace, bUseBounds, bUsePivot),
-		FCanExecuteAction::CreateStatic(&FLevelEditorActionCallbacks::ActorSelected_CanExecute)
+		FExecuteAction::CreateStatic(&FLevelEditorActionCallbacks::SnapToFloor_Clicked, /*bAlign*/true, /*bUseLineTrace*/true, /*bUseBounds*/true, /*bUsePivot*/false),
+		FCanExecuteAction::CreateStatic(&FLevelEditorActionCallbacks::ElementSelected_CanExecute)
 		);
 
-	bAlign = false;
-	bUseLineTrace = false;
-	bUseBounds = false;
-	bUsePivot = false;
 	ActionList.MapAction(
 		Commands.SnapToActor,
-		FExecuteAction::CreateStatic(&FLevelEditorActionCallbacks::SnapActorToActor_Clicked, bAlign, bUseLineTrace, bUseBounds, bUsePivot),
-		FCanExecuteAction::CreateStatic(&FLevelEditorActionCallbacks::ActorsSelected_CanExecute)
+		FExecuteAction::CreateStatic(&FLevelEditorActionCallbacks::SnapElementsToElement_Clicked, /*bAlign*/false, /*bUseLineTrace*/false, /*bUseBounds*/false, /*bUsePivot*/false),
+		FCanExecuteAction::CreateStatic(&FLevelEditorActionCallbacks::ElementsSelected_CanExecute)
 		);
 
-	bAlign = true;
-	bUseLineTrace = false;
-	bUseBounds = false;
-	bUsePivot = false;
 	ActionList.MapAction(
 		Commands.AlignToActor,
-		FExecuteAction::CreateStatic(&FLevelEditorActionCallbacks::SnapActorToActor_Clicked, bAlign, bUseLineTrace, bUseBounds, bUsePivot),
-		FCanExecuteAction::CreateStatic(&FLevelEditorActionCallbacks::ActorsSelected_CanExecute)
+		FExecuteAction::CreateStatic(&FLevelEditorActionCallbacks::SnapElementsToElement_Clicked, /*bAlign*/true, /*bUseLineTrace*/false, /*bUseBounds*/false, /*bUsePivot*/false),
+		FCanExecuteAction::CreateStatic(&FLevelEditorActionCallbacks::ElementsSelected_CanExecute)
 		);
 
-	bAlign = false;
-	bUseLineTrace = true;
-	bUseBounds = false;
-	bUsePivot = true;
 	ActionList.MapAction(
 		Commands.SnapPivotToActor,
-		FExecuteAction::CreateStatic(&FLevelEditorActionCallbacks::SnapActorToActor_Clicked, bAlign, bUseLineTrace, bUseBounds, bUsePivot),
-		FCanExecuteAction::CreateStatic(&FLevelEditorActionCallbacks::ActorsSelected_CanExecute)
+		FExecuteAction::CreateStatic(&FLevelEditorActionCallbacks::SnapElementsToElement_Clicked, /*bAlign*/false, /*bUseLineTrace*/true, /*bUseBounds*/false, /*bUsePivot*/true),
+		FCanExecuteAction::CreateStatic(&FLevelEditorActionCallbacks::ElementsSelected_CanExecute)
 		);
 
-	bAlign = true;
-	bUseLineTrace = true;
-	bUseBounds = false;
-	bUsePivot = true;
 	ActionList.MapAction(
 		Commands.AlignPivotToActor,
-		FExecuteAction::CreateStatic(&FLevelEditorActionCallbacks::SnapActorToActor_Clicked, bAlign, bUseLineTrace, bUseBounds, bUsePivot),
-		FCanExecuteAction::CreateStatic(&FLevelEditorActionCallbacks::ActorsSelected_CanExecute)
+		FExecuteAction::CreateStatic(&FLevelEditorActionCallbacks::SnapElementsToElement_Clicked, /*bAlign*/true, /*bUseLineTrace*/true, /*bUseBounds*/false, /*bUsePivot*/true),
+		FCanExecuteAction::CreateStatic(&FLevelEditorActionCallbacks::ElementsSelected_CanExecute)
 		);
 
-	bAlign = false;
-	bUseLineTrace = true;
-	bUseBounds = true;
-	bUsePivot = false;
 	ActionList.MapAction(
 		Commands.SnapBottomCenterBoundsToActor,
-		FExecuteAction::CreateStatic(&FLevelEditorActionCallbacks::SnapActorToActor_Clicked, bAlign, bUseLineTrace, bUseBounds, bUsePivot),
-		FCanExecuteAction::CreateStatic(&FLevelEditorActionCallbacks::ActorsSelected_CanExecute)
+		FExecuteAction::CreateStatic(&FLevelEditorActionCallbacks::SnapElementsToElement_Clicked, /*bAlign*/false, /*bUseLineTrace*/true, /*bUseBounds*/true, /*bUsePivot*/false),
+		FCanExecuteAction::CreateStatic(&FLevelEditorActionCallbacks::ElementsSelected_CanExecute)
 		);
 
-	bAlign = true;
-	bUseLineTrace = true;
-	bUseBounds = true;
-	bUsePivot = false;
 	ActionList.MapAction(
 		Commands.AlignBottomCenterBoundsToActor,
-		FExecuteAction::CreateStatic(&FLevelEditorActionCallbacks::SnapActorToActor_Clicked, bAlign, bUseLineTrace, bUseBounds, bUsePivot),
-		FCanExecuteAction::CreateStatic(&FLevelEditorActionCallbacks::ActorsSelected_CanExecute)
+		FExecuteAction::CreateStatic(&FLevelEditorActionCallbacks::SnapElementsToElement_Clicked, /*bAlign*/true, /*bUseLineTrace*/true, /*bUseBounds*/true, /*bUsePivot*/false),
+		FCanExecuteAction::CreateStatic(&FLevelEditorActionCallbacks::ElementsSelected_CanExecute)
 		);
 
 	ActionList.MapAction(
 		Commands.DeltaTransformToActors, 
 		FExecuteAction::CreateStatic( &FLevelEditorActionCallbacks::DeltaTransform ), 
-		FCanExecuteAction::CreateStatic( &FLevelEditorActionCallbacks::ActorSelected_CanExecute ) );
+		FCanExecuteAction::CreateStatic( &FLevelEditorActionCallbacks::ElementSelected_CanExecute) );
 
 	ActionList.MapAction(
 		Commands.MirrorActorX,
-		FExecuteAction::CreateStatic( &FLevelEditorActionCallbacks::ExecuteExecCommand, FString( TEXT("ACTOR MIRROR X=-1") ) ),
-		FCanExecuteAction::CreateStatic(&FLevelEditorActionCallbacks::ActorSelected_CanExecute)
+		FExecuteAction::CreateStatic( &FLevelEditorActionCallbacks::ExecuteExecCommand, FString( TEXT("ELEMENT MIRROR X=-1") ) ),
+		FCanExecuteAction::CreateStatic(&FLevelEditorActionCallbacks::ElementSelected_CanExecute)
 		);
 
 	ActionList.MapAction(
 		Commands.MirrorActorY,
-		FExecuteAction::CreateStatic( &FLevelEditorActionCallbacks::ExecuteExecCommand, FString( TEXT("ACTOR MIRROR Y=-1") ) ),
-		FCanExecuteAction::CreateStatic(&FLevelEditorActionCallbacks::ActorSelected_CanExecute)
+		FExecuteAction::CreateStatic( &FLevelEditorActionCallbacks::ExecuteExecCommand, FString( TEXT("ELEMENT MIRROR Y=-1") ) ),
+		FCanExecuteAction::CreateStatic(&FLevelEditorActionCallbacks::ElementSelected_CanExecute)
 		);
 
 	ActionList.MapAction(
 		Commands.MirrorActorZ,
-		FExecuteAction::CreateStatic( &FLevelEditorActionCallbacks::ExecuteExecCommand, FString( TEXT("ACTOR MIRROR Z=-1") ) ),
-		FCanExecuteAction::CreateStatic(&FLevelEditorActionCallbacks::ActorSelected_CanExecute)
+		FExecuteAction::CreateStatic( &FLevelEditorActionCallbacks::ExecuteExecCommand, FString( TEXT("ELEMENT MIRROR Z=-1") ) ),
+		FCanExecuteAction::CreateStatic(&FLevelEditorActionCallbacks::ElementSelected_CanExecute)
 		);
 
 	ActionList.MapAction(
@@ -1117,12 +1061,12 @@ void FLevelEditorModule::BindGlobalLevelEditorCommands()
 
 	ActionList.MapAction(
 		Commands.SelectAllActorsOfSameClass,
-		FExecuteAction::CreateStatic( &FLevelEditorActionCallbacks::OnSelectAllActorsOfClass, (bool)false )
+		FExecuteAction::CreateStatic( &FLevelEditorActionCallbacks::OnSelectAllActorsOfClass, false )
 		);
 
 	ActionList.MapAction(
 		Commands.SelectAllActorsOfSameClassWithArchetype,
-		FExecuteAction::CreateStatic( &FLevelEditorActionCallbacks::OnSelectAllActorsOfClass, (bool)true )
+		FExecuteAction::CreateStatic( &FLevelEditorActionCallbacks::OnSelectAllActorsOfClass, true )
 		);
 
 	ActionList.MapAction(
diff --git a/Engine/Source/Editor/LevelEditor/Private/LevelEditorActions.cpp b/Engine/Source/Editor/LevelEditor/Private/LevelEditorActions.cpp
--- a/Engine/Source/Editor/LevelEditor/Private/LevelEditorActions.cpp
+++ b/Engine/Source/Editor/LevelEditor/Private/LevelEditorActions.cpp
@@ -42,6 +42,9 @@
 #include "WorldPartition/WorldPartitionRuntimeHash.h"
 #include "Editor/WorldBrowser/Public/WorldBrowserModule.h"
 
+#include "Elements/Framework/TypedElementSelectionSet.h"
+#include "Elements/Interfaces/TypedElementObjectInterface.h"
+
 #include "LevelEditor.h"
 #include "Matinee/MatineeActor.h"
 #include "Engine/LevelScriptBlueprint.h"
@@ -2786,23 +2789,30 @@ bool FLevelEditorActionCallbacks::IsCoordinateSystemActive( ECoordSystem Coordin
 	return GLevelEditorModeTools().GetCoordSystem() == CoordinateSystem;
 }
 
-void FLevelEditorActionCallbacks::MoveActorToGrid_Clicked( bool InAlign, bool bInPerActor )
+void FLevelEditorActionCallbacks::MoveElementsToGrid_Clicked( bool InAlign, bool InPerElement )
 {
-	const FScopedTransaction Transaction( NSLOCTEXT("UnrealEd", "MoveActorToGrid", "Snap Origin to Grid") );
-	MoveActorTo_Clicked( InAlign, NULL, bInPerActor );
+	// TODO: Ideally this would come from some level editor context
+	if (const UTypedElementSelectionSet* SelectionSet = GEditor->GetSelectedActors()->GetElementSelectionSet())
+	{
+		const FScopedTransaction Transaction(NSLOCTEXT("UnrealEd", "MoveElementsToGrid", "Snap Origin to Grid"));
+		MoveTo_Clicked(SelectionSet, InAlign, InPerElement);
+	}
 }
 
-void FLevelEditorActionCallbacks::MoveActorToActor_Clicked( bool InAlign )
+void FLevelEditorActionCallbacks::MoveElementsToElement_Clicked(bool InAlign)
 {
-	AActor* Actor = GEditor->GetSelectedActors()->GetBottom<AActor>();
-	if( Actor )
+	// TODO: Ideally this would come from some level editor context
+	if (const UTypedElementSelectionSet* SelectionSet = GEditor->GetSelectedActors()->GetElementSelectionSet())
 	{
-		const FScopedTransaction Transaction( NSLOCTEXT("UnrealEd", "MoveActorToActor", "Snap Origin to Actor") );
-		MoveActorTo_Clicked( InAlign, Actor );
+		if (const TTypedElement<UTypedElementWorldInterface> DestElement = SelectionSet->GetBottomSelectedElement<UTypedElementWorldInterface>())
+		{
+			const FScopedTransaction Transaction(NSLOCTEXT("UnrealEd", "MoveElementsToElement", "Snap Origin to Element"));
+			MoveTo_Clicked(SelectionSet, InAlign, /*bPerElement*/false, DestElement);
+		}
 	}
 }
 
-void FLevelEditorActionCallbacks::MoveActorTo_Clicked( const bool InAlign, const AActor* InDestination/* = NULL*/, bool bInPerActor/* = false*/ )
+void FLevelEditorActionCallbacks::MoveTo_Clicked( const UTypedElementSelectionSet* InSelectionSet, const bool InAlign, bool InPerElement, const TTypedElement<UTypedElementWorldInterface>& InDestination )
 {
 	// Fires ULevel::LevelDirtiedEvent when falling out of scope.
 	FScopedLevelDirtied		LevelDirtyCallback;
@@ -2810,57 +2820,75 @@ void FLevelEditorActionCallbacks::MoveActorTo_Clicked( const bool InAlign, const
 	// Update the pivot location.
 	FVector Delta = FVector::ZeroVector;
 	FVector NewLocation = FVector::ZeroVector;
-	FRotator NewRotation = FRotator::ZeroRotator;
+	FQuat NewRotation = FQuat::Identity;
 
-	if(!bInPerActor)
+	if (!InPerElement)
 	{
-		if ( InDestination )
+		if (InDestination)
 		{
-			NewLocation = InDestination->GetActorLocation();
-			NewRotation = InDestination->GetActorRotation();
-			GEditor->SetPivot( NewLocation, false, true );
+			FTransform DestinationTransform;
+			if (InDestination.GetWorldTransform(DestinationTransform))
+			{
+				NewLocation = DestinationTransform.GetLocation();
+				NewRotation = DestinationTransform.GetRotation();
+				GEditor->SetPivot(NewLocation, false, true);
+			}
 		}
 		else
 		{
 			const FVector OldPivot = GEditor->GetPivotLocation();
 			const FVector NewPivot = OldPivot.GridSnap(GEditor->GetGridSize());
 			Delta = NewPivot - OldPivot;
-			GEditor->SetPivot( NewPivot, false, true );
+			GEditor->SetPivot(NewPivot, false, true);
 		}
 	}
 
-	for ( FSelectionIterator It( GEditor->GetSelectedActorIterator() ) ; It ; ++It )
+	InSelectionSet->ForEachSelectedElement<UTypedElementWorldInterface>([InAlign, InPerElement, &InDestination, &Delta, &NewLocation, &NewRotation, &LevelDirtyCallback](const TTypedElement<UTypedElementWorldInterface>& InElement)
 	{
-		AActor* Actor = Cast<AActor>( *It );
-		checkSlow( Actor );
-		if ( Actor == InDestination )	// Early out
+		// Skip moving the destination element
+		if (InElement == InDestination)
 		{
-			continue;
+			return true;
 		}
 
-		Actor->Modify();
-
-		if(!InDestination)
+		FTransform CurrentTransform;
+		if (InElement.GetWorldTransform(CurrentTransform))
 		{
-			if ( bInPerActor )
+			if (!InDestination)
+			{
+				if (InPerElement)
+				{
+					const FVector OldPivot = CurrentTransform.GetLocation();
+					const FVector NewPivot = OldPivot.GridSnap(GEditor->GetGridSize());
+					Delta = NewPivot - OldPivot;
+					GEditor->SetPivot(NewPivot, false, true);
+				}
+
+				NewLocation = CurrentTransform.GetLocation() + Delta;
+			}
+
+			FTransform NewTransform = CurrentTransform;
+			NewTransform.SetLocation(NewLocation);
+			if (InAlign)
+			{
+				NewTransform.SetRotation(NewRotation);
+			}
+
+			FTransform SuitableTransform;
+			if (!InElement.FindSuitableTransformAtPoint(NewTransform, SuitableTransform))
 			{
-				const FVector OldPivot = Actor->GetActorLocation();
-				const FVector NewPivot = OldPivot.GridSnap(GEditor->GetGridSize());
-				Delta = NewPivot - OldPivot;
-				GEditor->SetPivot( NewPivot, false, true );
+				SuitableTransform = NewTransform;
 			}
 
-			NewLocation = Actor->GetActorLocation() + Delta;
+			InElement.NotifyMovementStarted();
+			InElement.SetWorldTransform(SuitableTransform);
+			InElement.NotifyMovementEnded();
+
+			LevelDirtyCallback.Request();
 		}
-		
-		Actor->TeleportTo( NewLocation, ( !InAlign ? Actor->GetActorRotation() : NewRotation ), false, true );
-		Actor->InvalidateLightingCache();
-		Actor->UpdateComponentTransforms();
-		Actor->PostEditMove(true);
 
-		Actor->MarkPackageDirty();
-		LevelDirtyCallback.Request();
-	}
+		return true;
+	});
 
 	GEditor->RedrawLevelEditingViewports();
 	GEditor->RebuildAlteredBSP(); // Update the Bsp of any levels containing a modified brush
@@ -2980,103 +3008,83 @@ void FLevelEditorActionCallbacks::Select2DLayerDeltaAway_Clicked(int32 Delta)
 
 void FLevelEditorActionCallbacks::SnapToFloor_Clicked( bool InAlign, bool InUseLineTrace, bool InUseBounds, bool InUsePivot )
 {
-	const FScopedTransaction Transaction( NSLOCTEXT("UnrealEd", "SnapActorsToFloor", "Snap Actors To Floor") );
-	SnapTo_Clicked( InAlign, InUseLineTrace, InUseBounds, InUsePivot );
+	// TODO: Ideally this would come from some level editor context
+	if (const UTypedElementSelectionSet* SelectionSet = GEditor->GetSelectedActors()->GetElementSelectionSet())
+	{
+		const FScopedTransaction Transaction(NSLOCTEXT("UnrealEd", "SnapActorsToFloor", "Snap Elements To Floor"));
+		SnapTo_Clicked(SelectionSet, InAlign, InUseLineTrace, InUseBounds, InUsePivot);
+	}
 }
 
-void FLevelEditorActionCallbacks::SnapActorToActor_Clicked( bool InAlign, bool InUseLineTrace, bool InUseBounds, bool InUsePivot )
+void FLevelEditorActionCallbacks::SnapElementsToElement_Clicked( bool InAlign, bool InUseLineTrace, bool InUseBounds, bool InUsePivot )
 {
-	AActor* Actor = GEditor->GetSelectedActors()->GetBottom<AActor>();
-	if( Actor )
+	// TODO: Ideally this would come from some level editor context
+	if (const UTypedElementSelectionSet* SelectionSet = GEditor->GetSelectedActors()->GetElementSelectionSet())
 	{
-		const FScopedTransaction Transaction( NSLOCTEXT("UnrealEd", "SnapActorsToActor", "Snap Actors To Actor") );
-		SnapTo_Clicked( InAlign, InUseLineTrace, InUseBounds, InUsePivot, Actor );
+		if (const TTypedElement<UTypedElementWorldInterface> DestElement = SelectionSet->GetBottomSelectedElement<UTypedElementWorldInterface>())
+		{
+			const FScopedTransaction Transaction(NSLOCTEXT("UnrealEd", "SnapElementsToElement", "Snap Elements to Element"));
+			SnapTo_Clicked(SelectionSet, InAlign, InUseLineTrace, InUseBounds, InUsePivot, DestElement);
+		}
 	}
 }
 
-void FLevelEditorActionCallbacks::SnapTo_Clicked( const bool InAlign, const bool InUseLineTrace, const bool InUseBounds, const bool InUsePivot, AActor* InDestination )
+void FLevelEditorActionCallbacks::SnapTo_Clicked( const UTypedElementSelectionSet* InSelectionSet, const bool InAlign, const bool InUseLineTrace, const bool InUseBounds, const bool InUsePivot, const TTypedElement<UTypedElementWorldInterface>& InDestination )
 {
 	// Fires ULevel::LevelDirtiedEvent when falling out of scope.
 	FScopedLevelDirtied		LevelDirtyCallback;
 
-	bool bSnappedComponents = false;
-
 	// Let the component visualizers try to handle the selection.
-	if (GUnrealEd->ComponentVisManager.HandleSnapTo(InAlign, InUseLineTrace, InUseBounds, InUsePivot, InDestination))
-	{
-		bSnappedComponents = true;
-	}
-
-	// Gather the selected actors and components to later add to the ingore list for object snapping
-	TArray<FActorOrComponent> ObjectsToIgnore;
-	for (FSelectedEditableComponentIterator It(GEditor->GetSelectedEditableComponentIterator()); It; ++It)
-	{
-		USceneComponent* SceneComponent = Cast<USceneComponent>(*It);
-		ObjectsToIgnore.Add(FActorOrComponent(SceneComponent));
-	}
-	for (FSelectionIterator It(GEditor->GetSelectedActorIterator()); It; ++It)
+	// TODO: Should this also take an element?
+	if (TTypedElement<UTypedElementObjectInterface> DestinationObjectHandle = InSelectionSet->GetElementList()->GetElement<UTypedElementObjectInterface>(InDestination))
 	{
-		AActor* Actor = Cast<AActor>(*It);
-		ObjectsToIgnore.Add(FActorOrComponent(Actor));
-	}
-
-
-	if( !bSnappedComponents && GEditor->GetSelectedComponentCount() > 0 )
-	{
-		for(FSelectedEditableComponentIterator It(GEditor->GetSelectedEditableComponentIterator()); It; ++It)
+		if (AActor* DestinationActor = Cast<AActor>(DestinationObjectHandle.GetObject()))
 		{
-			USceneComponent* SceneComponent = Cast<USceneComponent>(*It);
-			if(SceneComponent)
+			if (GUnrealEd->ComponentVisManager.HandleSnapTo(InAlign, InUseLineTrace, InUseBounds, InUsePivot, DestinationActor))
 			{
-				SceneComponent->Modify();
-				AActor* ActorOwner = SceneComponent->GetOwner();
-				bSnappedComponents = true;
-				if(ActorOwner)
-				{
-					ActorOwner->Modify();
-					GEditor->SnapObjectTo(FActorOrComponent(SceneComponent), InAlign, InUseLineTrace, InUseBounds, InUsePivot, FActorOrComponent(InDestination), ObjectsToIgnore);
-					ActorOwner->InvalidateLightingCache();
-					ActorOwner->UpdateComponentTransforms();
-
-					LevelDirtyCallback.Request();
-				}
+				return;
 			}
 		}
-
-		USceneComponent* LastComp = GEditor->GetSelectedComponents()->GetBottom<USceneComponent>();
-
-		GEditor->SetPivot(LastComp->GetComponentLocation(), false, true);
 	}
 
-	if( !bSnappedComponents )
+	// Ignore the selected elements when sweeping for the snap location
+	const TArray<FTypedElementHandle> ElementsToIgnore = InSelectionSet->GetSelectedElementHandles();
+
+	// Snap each selected element
+	bool bSnappedElements = false;
+	InSelectionSet->ForEachSelectedElement<UTypedElementWorldInterface>([InAlign, InUseLineTrace, InUseBounds, InUsePivot, &InDestination, &ElementsToIgnore, &LevelDirtyCallback, &bSnappedElements](const TTypedElement<UTypedElementWorldInterface>& InElement)
 	{
-		for(FSelectionIterator It(GEditor->GetSelectedActorIterator()); It; ++It)
+		if (GEditor->SnapElementTo(InElement, InAlign, InUseLineTrace, InUseBounds, InUsePivot, InDestination, ElementsToIgnore))
 		{
-			AActor* Actor = Cast<AActor>(*It);
-			if(Actor)
-			{
-				Actor->Modify();
-				GEditor->SnapObjectTo(FActorOrComponent(Actor), InAlign, InUseLineTrace, InUseBounds, InUsePivot, FActorOrComponent(InDestination), ObjectsToIgnore);
-				Actor->InvalidateLightingCache();
-				Actor->UpdateComponentTransforms();
-
-				LevelDirtyCallback.Request();
-			}
+			bSnappedElements = true;
+			LevelDirtyCallback.Request();
 		}
+		return true;
+	});
 
-
-		AActor* Actor = GEditor->GetSelectedActors()->GetBottom<AActor>();
-		if(Actor)
+	// Update the pivot location
+	if (bSnappedElements)
+	{
+		if (TTypedElement<UTypedElementWorldInterface> LastElement = InSelectionSet->GetBottomSelectedElement<UTypedElementWorldInterface>())
 		{
-			GEditor->SetPivot(Actor->GetActorLocation(), false, true);
-
-			if(UActorGroupingUtils::IsGroupingActive())
+			FTransform LastElementTransform;
+			if (LastElement.GetWorldTransform(LastElementTransform))
 			{
-				// set group pivot for the root-most group
-				AGroupActor* ActorGroupRoot = AGroupActor::GetRootForActor(Actor, true, true);
-				if(ActorGroupRoot)
+				GEditor->SetPivot(LastElementTransform.GetLocation(), false, true);
+
+				if (UActorGroupingUtils::IsGroupingActive())
 				{
-					ActorGroupRoot->CenterGroupLocation();
+					if (TTypedElement<UTypedElementObjectInterface> LastObjectElement = InSelectionSet->GetElementList()->GetElement<UTypedElementObjectInterface>(LastElement))
+					{
+						if (AActor* LastActor = Cast<AActor>(LastObjectElement.GetObject()))
+						{
+							// Set group pivot for the root-most group
+							if (AGroupActor* ActorGroupRoot = AGroupActor::GetRootForActor(LastActor, true, true))
+							{
+								ActorGroupRoot->CenterGroupLocation();
+							}
+						}
+					}
 				}
 			}
 		}
@@ -3093,14 +3101,26 @@ void FLevelEditorActionCallbacks::AlignBrushVerticesToGrid_Execute()
 
 bool FLevelEditorActionCallbacks::ActorSelected_CanExecute()
 {
-	// Had to have something selected
-	return ( ( GEditor->GetSelectedActorCount() > 0 ) ? true : false );
+	return GEditor->GetSelectedActorCount() > 0;
 }
 
 bool FLevelEditorActionCallbacks::ActorsSelected_CanExecute()
 {
-	// Has to have more than one selected
-	return ( ( GEditor->GetSelectedActorCount() > 1 ) ? true : false );
+	return GEditor->GetSelectedActorCount() > 1;
+}
+
+bool FLevelEditorActionCallbacks::ElementSelected_CanExecute()
+{
+	// TODO: Ideally this would come from some level editor context
+	const UTypedElementSelectionSet* SelectionSet = GEditor->GetSelectedActors()->GetElementSelectionSet();
+	return SelectionSet && SelectionSet->GetNumSelectedElements() > 0;
+}
+
+bool FLevelEditorActionCallbacks::ElementsSelected_CanExecute()
+{
+	// TODO: Ideally this would come from some level editor context
+	const UTypedElementSelectionSet* SelectionSet = GEditor->GetSelectedActors()->GetElementSelectionSet();
+	return SelectionSet && SelectionSet->GetNumSelectedElements() > 1;
 }
 
 void FLevelEditorActionCallbacks::GeometryCollection_SelectAllGeometry()
diff --git a/Engine/Source/Editor/LevelEditor/Public/LevelEditorActions.h b/Engine/Source/Editor/LevelEditor/Public/LevelEditorActions.h
--- a/Engine/Source/Editor/LevelEditor/Public/LevelEditorActions.h
+++ b/Engine/Source/Editor/LevelEditor/Public/LevelEditorActions.h
@@ -14,11 +14,14 @@
 #include "EditorStyleSet.h"
 #include "TexAligner/TexAligner.h"
 #include "LightmapResRatioAdjust.h"
+#include "Elements/Framework/TypedElementHandle.h"
+#include "Elements/Interfaces/TypedElementWorldInterface.h"
 
 class AMatineeActor;
 class FLightingBuildOptions;
 class SLevelEditor;
 class UActorFactory;
+class UTypedElementSelectionSet;
 
 /**
  * Unreal level editor actions
@@ -1239,9 +1242,9 @@ public:
 	static class UWorld* GetWorld();
 public:
 	/** 
-	 * Moves an actor to the grid.
+	 * Moves the selected elements to the grid.
 	 */
-	static void MoveActorToGrid_Clicked( bool InAlign, bool bInPerActor );
+	static void MoveElementsToGrid_Clicked( bool InAlign, bool InPerElement );
 
 	/**
 	* Snaps a selected actor to the camera view.
@@ -1249,9 +1252,9 @@ public:
 	static void SnapObjectToView_Clicked();
 
 	/** 
-	 * Moves an actor to another actor.
+	 * Moves the selected elements to the last selected element.
 	 */
-	static void MoveActorToActor_Clicked( bool InAlign );
+	static void MoveElementsToElement_Clicked( bool InAlign );
 
 	/** 
 	 * Snaps an actor to the currently selected 2D snap layer
@@ -1281,7 +1284,7 @@ public:
 	static void Select2DLayerDeltaAway_Clicked(int32 Delta);
 
 	/** 
-	 * Snaps an actor to the floor.  Optionally will align with the trace normal.
+	 * Snaps the selected elements to the floor.  Optionally will align with the trace normal.
 	 * @param InAlign			Whether or not to rotate the actor to align with the trace normal.
 	 * @param InUseLineTrace	Whether or not to only trace with a line through the world.
 	 * @param InUseBounds		Whether or not to base the line trace off of the bounds.
@@ -1290,19 +1293,25 @@ public:
 	static void SnapToFloor_Clicked( bool InAlign, bool InUseLineTrace, bool InUseBounds, bool InUsePivot );
 
 	/**
-	 * Snaps an actor to another actor.  Optionally will align with the trace normal.
+	 * Snaps the selected elements to another element.  Optionally will align with the trace normal.
 	 * @param InAlign			Whether or not to rotate the actor to align with the trace normal.
 	 * @param InUseLineTrace	Whether or not to only trace with a line through the world.
 	 * @param InUseBounds		Whether or not to base the line trace off of the bounds.
 	 * @param InUsePivot		Whether or not to use the pivot position.
 	 */
-	static void SnapActorToActor_Clicked( bool InAlign, bool InUseLineTrace, bool InUseBounds, bool InUsePivot );
+	static void SnapElementsToElement_Clicked( bool InAlign, bool InUseLineTrace, bool InUseBounds, bool InUsePivot );
 
 	/**
 	 * Aligns brush verticies to the nearest grid point.
 	 */
 	static void AlignBrushVerticesToGrid_Execute();
 
+	/**
+	 * Checks to see if at least one actor is selected
+	 *	@return true if it can execute.
+	 */
+	static bool ActorSelected_CanExecute();
+
 	/**
 	 * Checks to see if multiple actors are selected
 	 *	@return true if it can execute.
@@ -1310,31 +1319,36 @@ public:
 	static bool ActorsSelected_CanExecute();
 
 	/**
-	 * Checks to see if at least a single actor is selected
+	 * Checks to see if at least one element is selected
 	 *	@return true if it can execute.
 	 */
-	static bool ActorSelected_CanExecute();
+	static bool ElementSelected_CanExecute();
 
+	/**
+	 * Checks to see if multiple elements are selected
+	 *	@return true if it can execute.
+	 */
+	static bool ElementsSelected_CanExecute();
 
 	/** Called when 'Open Merge Actor' is clicked */
 	static void OpenMergeActor_Clicked();
 
 private:
 	/** 
-	 * Moves an actor...
-	 * @param InDestination		The destination actor we want to move this actor to, NULL assumes we just want to use the grid
+	 * Moves the selected elements.
+	 * @param InDestination		The destination element we want to move this element to, or invalid to move to the grid
 	 */
-	static void MoveActorTo_Clicked( const bool InAlign, const AActor* InDestination = NULL, bool bInPerActor = false );
+	static void MoveTo_Clicked( const UTypedElementSelectionSet* InSelectionSet, const bool InAlign, bool InPerElement, const TTypedElement<UTypedElementWorldInterface>& InDestination = TTypedElement<UTypedElementWorldInterface>() );
 
 	/** 
-	 * Snaps an actor or component...  Optionally will align with the trace normal.
+	 * Snaps the selected elements. Optionally will align with the trace normal.
 	 * @param InAlign			Whether or not to rotate the actor to align with the trace normal.
 	 * @param InUseLineTrace	Whether or not to only trace with a line through the world.
 	 * @param InUseBounds		Whether or not to base the line trace off of the bounds.
 	 * @param InUsePivot		Whether or not to use the pivot position.
-	 * @param InDestination		The destination actor we want to move this actor to, NULL assumes we just want to go towards the floor
+	 * @param InDestination		The destination element we want to move this actor to, or invalid to go towards the floor
 	 */
-	static void SnapTo_Clicked( const bool InAlign, const bool InUseLineTrace, const bool InUseBounds, const bool InUsePivot, AActor* InDestination = NULL );
+	static void SnapTo_Clicked(const UTypedElementSelectionSet* InSelectionSet, const bool InAlign, const bool InUseLineTrace, const bool InUseBounds, const bool InUsePivot, const TTypedElement<UTypedElementWorldInterface>& InDestination = TTypedElement<UTypedElementWorldInterface>() );
 
 	/** 
 	 * Create and apply animation to the SkeletalMeshComponent if Simulating
diff --git a/Engine/Source/Editor/UnrealEd/Classes/Editor/EditorEngine.h b/Engine/Source/Editor/UnrealEd/Classes/Editor/EditorEngine.h
--- a/Engine/Source/Editor/UnrealEd/Classes/Editor/EditorEngine.h
+++ b/Engine/Source/Editor/UnrealEd/Classes/Editor/EditorEngine.h
@@ -64,6 +64,7 @@ class USkeleton;
 class USoundBase;
 class USoundNode;
 class UTextureRenderTarget2D;
+struct FTypedElementHandle;
 struct FAnalyticsEventAttribute;
 class UEditorWorldExtensionManager;
 class ITargetDevice;
@@ -179,82 +180,6 @@ private:
 	TSet<FString> ActorLabels;
 };
 
-/** 
- * Represents an actor or a component for use in editor functionality such as snapping which can operate on either type
- */
-struct FActorOrComponent
-{
-	AActor* Actor;
-	USceneComponent* Component;
-
-	FActorOrComponent()
-		: Actor( nullptr )
-		, Component( nullptr )
-	{}
-
-	FActorOrComponent( AActor* InActor )
-		: Actor( InActor )
-		, Component( nullptr )
-	{}
-
-	FActorOrComponent( USceneComponent* InComponent )
-		: Actor( nullptr )
-		, Component( InComponent )
-	{}
-
-	UWorld* GetWorld() const
-	{
-		return Actor ? Actor->GetWorld() : Component->GetWorld();
-	}
-
-	bool operator==( const FActorOrComponent& Other ) const
-	{
-		return Actor == Other.Actor && Component == Other.Component;
-	}
-
-	const FBoxSphereBounds& GetBounds() const
-	{
-		return Actor ? Actor->GetRootComponent()->Bounds : Component->Bounds;
-	}
-
-	FVector GetWorldLocation() const
-	{
-		return Actor ? Actor->GetActorLocation() : Component->GetComponentLocation();
-	}
-
-	FRotator GetWorldRotation() const
-	{
-		return Actor ? Actor->GetActorRotation() : Component->GetComponentRotation();
-	}
-
-	void SetWorldLocation( const FVector& NewLocation )
-	{
-		if( Actor )
-		{
-			Actor->SetActorLocation( NewLocation );
-		}
-		else
-		{
-			Component->SetWorldLocation( NewLocation );
-		}
-	}
-
-	void SetWorldRotation( const FRotator& NewRotation )
-	{
-		if(Actor)
-		{
-			Actor->SetActorRotation(NewRotation);
-		}
-		else
-		{
-			Component->SetWorldRotation(NewRotation);
-		}
-	}
-
-	/** @return true if this is a valid actor or component but not both */
-	bool IsValid() const { return (Actor != nullptr) ^ (Component != nullptr);}
-};
-
 /**
  * Represents the current selection state of a level (its selected actors and components) from a given point in a time, in a way that can be safely restored later even if the level is reloaded
  */
@@ -1084,23 +1009,23 @@ public:
 	void MoveViewportCamerasToBox(const FBox& BoundingBox, bool bActiveViewportOnly) const;
 
 	/** 
-	 * Snaps an actor in a direction.  Optionally will align with the trace normal.
-	 * @param InActor			Actor to move to the floor.
+	 * Snaps an element in a direction.  Optionally will align with the trace normal.
+	 * @param InElementHandle	Element to move to the floor.
 	 * @param InAlign			Whether or not to rotate the actor to align with the trace normal.
 	 * @param InUseLineTrace	Whether or not to only trace with a line through the world.
 	 * @param InUseBounds		Whether or not to base the line trace off of the bounds.
 	 * @param InUsePivot		Whether or not to use the pivot position.
-	 * @param InDestination		The destination actor we want to move this actor to, NULL assumes we just want to go towards the floor
+	 * @param InDestination		The destination element we want to move this actor to, unset assumes we just want to go towards the floor
 	 * @return					Whether or not the actor was moved.
 	 */
-	bool SnapObjectTo( FActorOrComponent Object, const bool InAlign, const bool InUseLineTrace, const bool InUseBounds, const bool InUsePivot, FActorOrComponent InDestination = FActorOrComponent(), TArray<FActorOrComponent> ObjectsToIgnore = TArray<FActorOrComponent>() );
+	bool SnapElementTo(const FTypedElementHandle& InElementHandle, const bool InAlign, const bool InUseLineTrace, const bool InUseBounds, const bool InUsePivot, const FTypedElementHandle& InDestination = FTypedElementHandle(), TArrayView<const FTypedElementHandle> InElementsToIgnore = TArrayView<const FTypedElementHandle>());
 
 	/**
-	 * Snaps the view of the camera to that of the provided actor.
+	 * Snaps the view of the camera to that of the provided element.
 	 *
-	 * @param	Actor	The actor the camera is going to be snapped to.
+	 * @param InElementHandle	The element the camera is going to be snapped to.
 	 */
-	void SnapViewTo(const FActorOrComponent& Object);
+	void SnapViewTo(const FTypedElementHandle& InElementHandle);
 
 	/**
 	 * Remove the roll, pitch and/or yaw from the perspective viewports' cameras.
diff --git a/Engine/Source/Editor/UnrealEd/Classes/Editor/UnrealEdEngine.h b/Engine/Source/Editor/UnrealEd/Classes/Editor/UnrealEdEngine.h
--- a/Engine/Source/Editor/UnrealEd/Classes/Editor/UnrealEdEngine.h
+++ b/Engine/Source/Editor/UnrealEd/Classes/Editor/UnrealEdEngine.h
@@ -327,6 +327,7 @@ public:
 	bool Exec_Edit( UWorld* InWorld, const TCHAR* Str, FOutputDevice& Ar );
 	bool Exec_Pivot( const TCHAR* Str, FOutputDevice& Ar );
 	bool Exec_Actor( UWorld* InWorld, const TCHAR* Str, FOutputDevice& Ar );
+	bool Exec_Element( UWorld* InWorld, const TCHAR* Str, FOutputDevice& Ar );
 	bool Exec_Mode( const TCHAR* Str, FOutputDevice& Ar );
 	bool Exec_Group( const TCHAR* Str, FOutputDevice& Ar );
 
diff --git a/Engine/Source/Editor/UnrealEd/Private/EditorServer.cpp b/Engine/Source/Editor/UnrealEd/Private/EditorServer.cpp
--- a/Engine/Source/Editor/UnrealEd/Private/EditorServer.cpp
+++ b/Engine/Source/Editor/UnrealEd/Private/EditorServer.cpp
@@ -126,6 +126,11 @@
 #include "AssetRegistryModule.h"
 #include "SnappingUtils.h"
 
+#include "Elements/Framework/TypedElementHandle.h"
+#include "Elements/Framework/TypedElementRegistry.h"
+#include "Elements/Framework/TypedElementSelectionSet.h"
+#include "Elements/Interfaces/TypedElementWorldInterface.h"
+#include "Elements/Interfaces/TypedElementObjectInterface.h"
 
 #include "Editor/ActorPositioning.h"
 
@@ -4935,32 +4940,46 @@ void UEditorEngine::MoveViewportCamerasToBox(const FBox& BoundingBox, bool bActi
 	}
 }
 
-/** 
- * Snaps an actor in a direction.  Optionally will align with the trace normal.
- * @param InActor			Actor to move to the floor.
- * @param InAlign			Whether or not to rotate the actor to align with the trace normal.
- * @param InUseLineTrace	Whether or not to only trace with a line through the world.
- * @param InUseBounds		Whether or not to base the line trace off of the bounds.
- * @param InUsePivot		Whether or not to use the pivot position.
- * @param InDestination		The destination actor we want to move this actor to, NULL assumes we just want to go towards the floor
- * @return					Whether or not the actor was moved.
- */
-bool UEditorEngine::SnapObjectTo( FActorOrComponent Object, const bool InAlign, const bool InUseLineTrace, const bool InUseBounds, const bool InUsePivot, FActorOrComponent InDestination, TArray<FActorOrComponent> ObjectsToIgnore)
+bool UEditorEngine::SnapElementTo(const FTypedElementHandle& InElementHandle, const bool InAlign, const bool InUseLineTrace, const bool InUseBounds, const bool InUsePivot, const FTypedElementHandle& InDestination, TArrayView<const FTypedElementHandle> InElementsToIgnore)
 {
-	if ( !Object.IsValid() || Object == InDestination )	// Early out
+	if (!InElementHandle || InElementHandle == InDestination)
+	{
+		return false;
+	}
+
+	const UTypedElementRegistry* Registry = UTypedElementRegistry::GetInstance();
+
+	TTypedElement<UTypedElementWorldInterface> ElementWorldHandle = Registry->GetElement<UTypedElementWorldInterface>(InElementHandle);
+	if (!ElementWorldHandle)
 	{
 		return false;
 	}
 
+	FTransform ElementTransform;
+	if (!ElementWorldHandle.GetWorldTransform(ElementTransform))
+	{
+		return false;
+	}
 
-	FVector	StartLocation = Object.GetWorldLocation();
+	FBoxSphereBounds ElementBounds;
+	if (!ElementWorldHandle.GetBounds(ElementBounds))
+	{
+		return false;
+	}
+
+	ABrush* Brush = nullptr;
+	if (TTypedElement<UTypedElementObjectInterface> ElementObjectHandle = Registry->GetElement<UTypedElementObjectInterface>(InElementHandle))
+	{
+		Brush = Cast<ABrush>(ElementObjectHandle.GetObject());
+	}
+
+	FVector	StartLocation = ElementTransform.GetLocation();
 	FVector	LocationOffset = FVector::ZeroVector;
 	FVector	Extent = FVector::ZeroVector;
-	ABrush* Brush = Cast< ABrush >( Object.Actor );
-	bool UseLineTrace = Brush ? true: InUseLineTrace;
-	bool UseBounds = Brush ? true: InUseBounds;
+	bool UseLineTrace = Brush ? true : InUseLineTrace;
+	bool UseBounds = Brush ? true : InUseBounds;
 
-	if( UseLineTrace && UseBounds )
+	if (UseLineTrace && UseBounds)
 	{
 		if (InUsePivot)
 		{
@@ -4970,72 +4989,75 @@ bool UEditorEngine::SnapObjectTo( FActorOrComponent Object, const bool InAlign,
 		else
 		{
 			// Will do a line trace from the center bottom of the bounds through the world. Will begin at the bottom center of the component's bounds.
-			StartLocation = Object.GetBounds().Origin;
-			StartLocation.Z -= Object.GetBounds().BoxExtent.Z;
+			StartLocation = ElementBounds.Origin;
+			StartLocation.Z -= ElementBounds.BoxExtent.Z;
 		}
 
 		// Forces a line trace.
 		Extent = FVector::ZeroVector;
-		LocationOffset = StartLocation - Object.GetWorldLocation();
+		LocationOffset = StartLocation - ElementTransform.GetLocation();
 	}
-	else if( UseLineTrace )
+	else if (UseLineTrace)
 	{
 		// This will be false if multiple objects are selected. In that case the actor's position should be used so all the objects do not go to the same point.
-		if( InUsePivot && !InDestination.IsValid() )	// @todo: If the destination actor is part of the selection tho, we can't use the pivot! (remove check if not)
+		if (InUsePivot && !InDestination)	// @todo: If the destination actor is part of the selection tho, we can't use the pivot! (remove check if not)
 		{
 			StartLocation = GetPivotLocation();
 		}
 		else
 		{
-			StartLocation = Object.GetWorldLocation();
+			StartLocation = ElementTransform.GetLocation();
 		}
 
 		// Forces a line trace.
 		Extent = FVector::ZeroVector;
-		LocationOffset = StartLocation - Object.GetWorldLocation();
+		LocationOffset = StartLocation - ElementTransform.GetLocation();
 	}
-	else 
+	else
 	{
-		StartLocation = Object.GetBounds().Origin;
+		StartLocation = ElementBounds.Origin;
 
-		Extent = Object.GetBounds().BoxExtent;
-		LocationOffset = StartLocation - Object.GetWorldLocation();
+		Extent = ElementBounds.BoxExtent;
+		LocationOffset = StartLocation - ElementTransform.GetLocation();
 	}
 
-
-	FVector Direction = FVector(0.f,0.f,-1.f);
-	if ( InDestination.IsValid() )	// If a destination actor was specified, work out the direction
+	FVector Direction = FVector(0.f, 0.f, -1.f);
+	if (TTypedElement<UTypedElementWorldInterface> DestinationWorldHandle = Registry->GetElement<UTypedElementWorldInterface>(InDestination))	// If a destination actor was specified, work out the direction
 	{
-		FVector	EndLocation = InDestination.GetWorldLocation();
+		FTransform DestinationTransform;
+		if (!DestinationWorldHandle.GetWorldTransform(DestinationTransform))
+		{
+			return false;
+		}
+
+		FBoxSphereBounds DestinationBounds;
+		if (!DestinationWorldHandle.GetBounds(DestinationBounds))
+		{
+			return false;
+		}
+
+		FVector	EndLocation = DestinationTransform.GetLocation();
 
 		// Code here assumes you want to same type of end point as the start point used, comment out to just use the destination actors origin!
-		if( UseLineTrace && UseBounds )
+		if (UseLineTrace && UseBounds)
 		{
-			EndLocation = InDestination.GetBounds().Origin;
-			EndLocation.Z -= InDestination.GetBounds().BoxExtent.Z;
+			EndLocation = DestinationBounds.Origin;
+			EndLocation.Z -= DestinationBounds.BoxExtent.Z;
 		}
-		else if( UseLineTrace )
+		else if (UseLineTrace)
 		{
-			// This will be false if multiple objects are selected. In that case the actor's position should be used so all the objects do not go to the same point.
-			if( InUsePivot && !InDestination.IsValid() )	// @todo: If the destination actor is part of the selection tho, we can't use the pivot! (remove check if not)
-			{
-				EndLocation = GetPivotLocation();
-			}
-			else
-			{
-				EndLocation = InDestination.GetWorldLocation();
-			}
+			EndLocation = DestinationTransform.GetLocation();
 		}
 		else
 		{
-			EndLocation = InDestination.GetBounds().Origin;
+			EndLocation = DestinationBounds.Origin;
 		}
 
-		if ( EndLocation.Equals( StartLocation ) )
+		if (EndLocation.Equals(StartLocation))
 		{
 			return false;
 		}
-		Direction = ( EndLocation - StartLocation );
+		Direction = (EndLocation - StartLocation);
 		Direction.Normalize();
 	}
 
@@ -5047,72 +5069,25 @@ bool UEditorEngine::SnapObjectTo( FActorOrComponent Object, const bool InAlign,
 		StartLocation.Z = Brush->GetRootComponent()->Bounds.Origin.Z - Brush->GetRootComponent()->Bounds.BoxExtent.Z - fTinyOffset;
 	}
 
-	// Do the actual actor->world check.  We try to collide against the world, straight down from our current position.
+	// Do the actual sweep test. We try to collide against the world, straight down from our current position.
 	// If we hit anything, we will move the actor to a position that lets it rest on the floor.
-	FHitResult Hit(1.0f);
-	FCollisionQueryParams Params(SCENE_QUERY_STAT(MoveActorToTrace), false);
-	for (FActorOrComponent ObjectToIgnore : ObjectsToIgnore)
+	FTransform NewTransform;
+	if (ElementWorldHandle.FindSuitableTransformAlongPath(StartLocation, StartLocation + Direction * WORLD_MAX, FCollisionShape::MakeBox(Extent), InElementsToIgnore, NewTransform))
 	{
-		if (ObjectToIgnore.Actor)
+		NewTransform.SetTranslation(NewTransform.GetTranslation() - LocationOffset);
+		if (!InAlign)
 		{
-			Params.AddIgnoredActor(ObjectToIgnore.Actor);
-		}
-		else
-		{
-			Params.AddIgnoredComponent(Cast<UPrimitiveComponent>(ObjectToIgnore.Component));
-		}
-	}
-	if( Object.Actor )
-	{
-		Params.AddIgnoredActor( Object.Actor );
-		TArray<AActor*> ChildActors;
-		Object.Actor->GetAllChildActors(ChildActors);
-		Params.AddIgnoredActors(ChildActors);
-	}
-	else
-	{
-		Params.AddIgnoredComponent( Cast<UPrimitiveComponent>(Object.Component) );
-	}
-
-	if (Object.GetWorld()->SweepSingleByChannel(Hit, StartLocation, StartLocation + Direction*WORLD_MAX, FQuat::Identity, ECC_WorldStatic, FCollisionShape::MakeBox(Extent), Params))
-	{
-		FVector NewLocation = Hit.Location - LocationOffset;
-		NewLocation.Z += KINDA_SMALL_NUMBER;	// Move the new desired location up by an error tolerance
-		
-		if (Object.Actor)
-		{
-			BroadcastBeginObjectMovement(*Object.Actor);
-		}
-		else
-		{
-			BroadcastBeginObjectMovement(*Object.Component);
-		}
-
-		Object.SetWorldLocation( NewLocation );
-		//InActor->TeleportTo( NewLocation, InActor->GetActorRotation(), false,true );
-		
-		if( InAlign )
-		{
-			//@todo: This doesn't take into account that rotating the actor changes LocationOffset.
-			FRotator NewRotation( Hit.Normal.Rotation() );
-			NewRotation.Pitch -= 90.f;
-			Object.SetWorldRotation( NewRotation );
-		}
-
-		if (Object.Actor)
-		{
-			BroadcastEndObjectMovement(*Object.Actor);
-		}
-		else
-		{
-			BroadcastEndObjectMovement(*Object.Component);
+			NewTransform.SetRotation(ElementTransform.GetRotation());
 		}
+		NewTransform.SetScale3D(ElementTransform.GetScale3D());
 
 		// Switch to the pie world if we have one
-		FScopedConditionalWorldSwitcher WorldSwitcher( GCurrentLevelEditingViewportClient );
+		FScopedConditionalWorldSwitcher WorldSwitcher(GCurrentLevelEditingViewportClient);
+
+		ElementWorldHandle.NotifyMovementStarted();
+		ElementWorldHandle.SetWorldTransform(NewTransform);
+		ElementWorldHandle.NotifyMovementEnded();
 
-		Object.Actor ? Object.Actor->PostEditMove(true) : Object.Component->GetOwner()->PostEditMove(true);
-		//InActor->PostEditMove( true );
 		if (Brush)
 		{
 			RebuildAlteredBSP();
@@ -5126,7 +5101,6 @@ bool UEditorEngine::SnapObjectTo( FActorOrComponent Object, const bool InAlign,
 	return false;
 }
 
-
 void UEditorEngine::MoveActorInFrontOfCamera( AActor& InActor, const FVector& InCameraOrigin, const FVector& InCameraDirection )
 {
 	const FVector NewLocation = FActorPositioning::GetActorPositionInFrontOfCamera(InActor, InCameraOrigin, InCameraDirection);
@@ -5148,16 +5122,22 @@ void UEditorEngine::MoveActorInFrontOfCamera( AActor& InActor, const FVector& In
 	InActor.PostEditMove( true );
 }
 
-
-void UEditorEngine::SnapViewTo(const FActorOrComponent& Object)
+void UEditorEngine::SnapViewTo(const FTypedElementHandle& InElementHandle)
 {
-	for(FLevelEditorViewportClient* ViewportClient : GetLevelViewportClients())
+	if (TTypedElement<UTypedElementWorldInterface> ElementWorldHandle = UTypedElementRegistry::GetInstance()->GetElement<UTypedElementWorldInterface>(InElementHandle))
 	{
-		if ( ViewportClient->IsPerspective()  )
+		FTransform ElementTransform;
+		if (ElementWorldHandle.GetWorldTransform(ElementTransform))
 		{
-			ViewportClient->SetViewLocation( Object.GetWorldLocation() );
-			ViewportClient->SetViewRotation( Object.GetWorldRotation() );
-			ViewportClient->Invalidate();
+			for (FLevelEditorViewportClient* ViewportClient : GetLevelViewportClients())
+			{
+				if (ViewportClient->IsPerspective())
+				{
+					ViewportClient->SetViewLocation(ElementTransform.GetLocation());
+					ViewportClient->SetViewRotation(ElementTransform.Rotator());
+					ViewportClient->Invalidate();
+				}
+			}
 		}
 	}
 }
@@ -5269,17 +5249,15 @@ bool UEditorEngine::Exec_Camera( const TCHAR* Str, FOutputDevice& Ar )
 	}
 	else if ( bSnap )
 	{
-		FActorOrComponent SelectedObject(GetSelectedComponents()->GetTop<USceneComponent>());
-		if (!SelectedObject.IsValid())
+		const UTypedElementSelectionSet* SelectionSet = GetSelectedActors()->GetElementSelectionSet();
+		if (SelectionSet && SelectionSet->HasSelectedElements())
 		{
-			SelectedObject.Actor = GetSelectedActors()->GetTop<AActor>();
-		}
-
-		if (SelectedObject.IsValid())
-		{
-			// Set perspective viewport camera parameters to that of the selected camera.
-			SnapViewTo(SelectedObject);
-			Ar.Log( TEXT("Snapped camera to the first selected object.") );
+			if (FTypedElementHandle SelectedElement = SelectionSet->GetElementList()->GetElementHandleAt(0))
+			{
+				// Set perspective viewport camera parameters to that of the selected camera.
+				SnapViewTo(SelectedElement);
+				Ar.Log(TEXT("Snapped camera to the first selected element."));
+			}
 		}
 	}
 
diff --git a/Engine/Source/Editor/UnrealEd/Private/Elements/Actor/ActorElementEditorViewportInteractionCustomization.cpp b/Engine/Source/Editor/UnrealEd/Private/Elements/Actor/ActorElementEditorViewportInteractionCustomization.cpp
--- a/Engine/Source/Editor/UnrealEd/Private/Elements/Actor/ActorElementEditorViewportInteractionCustomization.cpp
+++ b/Engine/Source/Editor/UnrealEd/Private/Elements/Actor/ActorElementEditorViewportInteractionCustomization.cpp
@@ -8,6 +8,7 @@
 #include "EditorModeManager.h"
 #include "Toolkits/IToolkitHost.h"
 #include "AI/NavigationSystemBase.h"
+#include "Components/BrushComponent.h"
 
 bool FActorElementEditorViewportInteractionCustomization::GetGizmoPivotLocation(const TTypedElement<UTypedElementWorldInterface>& InElementWorldHandle, const UE::Widget::EWidgetMode InWidgetMode, FVector& OutPivotLocation)
 {
@@ -31,6 +32,27 @@ void FActorElementEditorViewportInteractionCustomization::GizmoManipulationDelta
 	FActorElementEditorViewportInteractionCustomization::ApplyDeltaToActor(Actor, /*bDelta*/true, &DeltaTranslation, &DeltaRotation, &DeltaScale3D, InPivotLocation, InInputState);
 }
 
+void FActorElementEditorViewportInteractionCustomization::MirrorElement(const TTypedElement<UTypedElementWorldInterface>& InElementWorldHandle, const FVector& InMirrorScale, const FVector& InPivotLocation)
+{
+	AActor* Actor = ActorElementDataUtil::GetActorFromHandleChecked(InElementWorldHandle);
+
+	Actor->Modify();
+	Actor->EditorApplyMirror(InMirrorScale, InPivotLocation);
+
+	if (ABrush* Brush = Cast<ABrush>(Actor))
+	{
+		if (UBrushComponent* BrushComponent = Brush->GetBrushComponent())
+		{
+			BrushComponent->RequestUpdateBrushCollision();
+		}
+	}
+
+	Actor->InvalidateLightingCache();
+	Actor->PostEditMove(true);
+
+	Actor->MarkPackageDirty();
+}
+
 void FActorElementEditorViewportInteractionCustomization::ApplyDeltaToActor(AActor* InActor, const bool InIsDelta, const FVector* InDeltaTranslationPtr, const FRotator* InDeltaRotationPtr, const FVector* InDeltaScalePtr, const FVector& InPivotLocation, const FInputDeviceState& InInputState)
 {
 	const bool bIsSimulatingInEditor = GEditor->IsSimulatingInEditor();
diff --git a/Engine/Source/Editor/UnrealEd/Private/Elements/Actor/ActorElementEditorWorldInterface.cpp b/Engine/Source/Editor/UnrealEd/Private/Elements/Actor/ActorElementEditorWorldInterface.cpp
new file mode 100644
--- /dev/null
+++ b/Engine/Source/Editor/UnrealEd/Private/Elements/Actor/ActorElementEditorWorldInterface.cpp
@@ -0,0 +1,36 @@
+// Copyright Epic Games, Inc. All Rights Reserved.
+
+#include "Elements/Actor/ActorElementEditorWorldInterface.h"
+#include "Elements/Actor/ActorElementData.h"
+#include "GameFramework/Actor.h"
+
+#include "Editor.h"
+
+void UActorElementEditorWorldInterface::NotifyMovementStarted(const FTypedElementHandle& InElementHandle)
+{
+	if (AActor* Actor = ActorElementDataUtil::GetActorFromHandle(InElementHandle))
+	{
+		GEditor->BroadcastBeginObjectMovement(*Actor);
+	}
+}
+
+void UActorElementEditorWorldInterface::NotifyMovementOngoing(const FTypedElementHandle& InElementHandle)
+{
+	if (AActor* Actor = ActorElementDataUtil::GetActorFromHandle(InElementHandle))
+	{
+		Actor->PostEditMove(false);
+	}
+}
+
+void UActorElementEditorWorldInterface::NotifyMovementEnded(const FTypedElementHandle& InElementHandle)
+{
+	if (AActor* Actor = ActorElementDataUtil::GetActorFromHandle(InElementHandle))
+	{
+		GEditor->BroadcastEndObjectMovement(*Actor);
+		Actor->PostEditMove(true);
+
+		Actor->InvalidateLightingCache();
+		Actor->UpdateComponentTransforms();
+		Actor->MarkPackageDirty();
+	}
+}
diff --git a/Engine/Source/Editor/UnrealEd/Private/Elements/Actor/ActorElementEditorWorldInterface.h b/Engine/Source/Editor/UnrealEd/Private/Elements/Actor/ActorElementEditorWorldInterface.h
new file mode 100644
--- /dev/null
+++ b/Engine/Source/Editor/UnrealEd/Private/Elements/Actor/ActorElementEditorWorldInterface.h
@@ -0,0 +1,17 @@
+// Copyright Epic Games, Inc. All Rights Reserved.
+
+#pragma once
+
+#include "Elements/Actor/ActorElementWorldInterface.h"
+#include "ActorElementEditorWorldInterface.generated.h"
+
+UCLASS()
+class UActorElementEditorWorldInterface : public UActorElementWorldInterface
+{
+	GENERATED_BODY()
+
+public:
+	virtual void NotifyMovementStarted(const FTypedElementHandle& InElementHandle) override;
+	virtual void NotifyMovementOngoing(const FTypedElementHandle& InElementHandle) override;
+	virtual void NotifyMovementEnded(const FTypedElementHandle& InElementHandle) override;
+};
diff --git a/Engine/Source/Editor/UnrealEd/Private/Elements/Component/ComponentElementEditorWorldInterface.cpp b/Engine/Source/Editor/UnrealEd/Private/Elements/Component/ComponentElementEditorWorldInterface.cpp
new file mode 100644
--- /dev/null
+++ b/Engine/Source/Editor/UnrealEd/Private/Elements/Component/ComponentElementEditorWorldInterface.cpp
@@ -0,0 +1,41 @@
+// Copyright Epic Games, Inc. All Rights Reserved.
+
+#include "Elements/Component/ComponentElementEditorWorldInterface.h"
+#include "Elements/Component/ComponentElementData.h"
+#include "Components/ActorComponent.h"
+
+#include "Editor.h"
+
+void UComponentElementEditorWorldInterface::NotifyMovementStarted(const FTypedElementHandle& InElementHandle)
+{
+	if (UActorComponent* Component = ComponentElementDataUtil::GetComponentFromHandle(InElementHandle))
+	{
+		GEditor->BroadcastBeginObjectMovement(*Component);
+	}
+}
+
+void UComponentElementEditorWorldInterface::NotifyMovementOngoing(const FTypedElementHandle& InElementHandle)
+{
+	if (UActorComponent* Component = ComponentElementDataUtil::GetComponentFromHandle(InElementHandle))
+	{
+		if (AActor* Actor = Component->GetOwner())
+		{
+			Actor->PostEditMove(false);
+		}
+	}
+}
+
+void UComponentElementEditorWorldInterface::NotifyMovementEnded(const FTypedElementHandle& InElementHandle)
+{
+	if (UActorComponent* Component = ComponentElementDataUtil::GetComponentFromHandle(InElementHandle))
+	{
+		GEditor->BroadcastEndObjectMovement(*Component);
+		if (AActor* Actor = Component->GetOwner())
+		{
+			Actor->PostEditMove(true);
+			Actor->InvalidateLightingCache();
+		}
+
+		Component->MarkPackageDirty();
+	}
+}
diff --git a/Engine/Source/Editor/UnrealEd/Private/Elements/Component/ComponentElementEditorWorldInterface.h b/Engine/Source/Editor/UnrealEd/Private/Elements/Component/ComponentElementEditorWorldInterface.h
new file mode 100644
--- /dev/null
+++ b/Engine/Source/Editor/UnrealEd/Private/Elements/Component/ComponentElementEditorWorldInterface.h
@@ -0,0 +1,17 @@
+// Copyright Epic Games, Inc. All Rights Reserved.
+
+#pragma once
+
+#include "Elements/Component/ComponentElementWorldInterface.h"
+#include "ComponentElementEditorWorldInterface.generated.h"
+
+UCLASS()
+class UComponentElementEditorWorldInterface : public UComponentElementWorldInterface
+{
+	GENERATED_BODY()
+
+public:
+	virtual void NotifyMovementStarted(const FTypedElementHandle& InElementHandle) override;
+	virtual void NotifyMovementOngoing(const FTypedElementHandle& InElementHandle) override;
+	virtual void NotifyMovementEnded(const FTypedElementHandle& InElementHandle) override;
+};
diff --git a/Engine/Source/Editor/UnrealEd/Private/Elements/Framework/EditorElements.cpp b/Engine/Source/Editor/UnrealEd/Private/Elements/Framework/EditorElements.cpp
--- a/Engine/Source/Editor/UnrealEd/Private/Elements/Framework/EditorElements.cpp
+++ b/Engine/Source/Editor/UnrealEd/Private/Elements/Framework/EditorElements.cpp
@@ -5,8 +5,10 @@
 
 #include "Elements/Object/ObjectElementEditorSelectionInterface.h"
 
+#include "Elements/Actor/ActorElementEditorWorldInterface.h"
 #include "Elements/Actor/ActorElementEditorSelectionInterface.h"
 
+#include "Elements/Component/ComponentElementEditorWorldInterface.h"
 #include "Elements/Component/ComponentElementEditorSelectionInterface.h"
 
 FSimpleMulticastDelegate OnRegisterEditorElementsDelegate;
@@ -22,6 +24,7 @@ void RegisterEditorActorElements()
 {
 	UTypedElementRegistry* Registry = UTypedElementRegistry::GetInstance();
 
+	Registry->RegisterElementInterface<UTypedElementWorldInterface>(NAME_Actor, NewObject<UActorElementEditorWorldInterface>(), /*bAllowOverride*/true);
 	Registry->RegisterElementInterface<UTypedElementSelectionInterface>(NAME_Actor, NewObject<UActorElementEditorSelectionInterface>(), /*bAllowOverride*/true);
 }
 
@@ -29,6 +32,7 @@ void RegisterEditorComponentElements()
 {
 	UTypedElementRegistry* Registry = UTypedElementRegistry::GetInstance();
 
+	Registry->RegisterElementInterface<UTypedElementWorldInterface>(NAME_Components, NewObject<UComponentElementEditorWorldInterface>(), /*bAllowOverride*/true);
 	Registry->RegisterElementInterface<UTypedElementSelectionInterface>(NAME_Components, NewObject<UComponentElementEditorSelectionInterface>(), /*bAllowOverride*/true);
 }
 
diff --git a/Engine/Source/Editor/UnrealEd/Private/LevelEditorViewport.cpp b/Engine/Source/Editor/UnrealEd/Private/LevelEditorViewport.cpp
--- a/Engine/Source/Editor/UnrealEd/Private/LevelEditorViewport.cpp
+++ b/Engine/Source/Editor/UnrealEd/Private/LevelEditorViewport.cpp
@@ -3632,6 +3632,46 @@ void FLevelEditorViewportClient::ApplyDeltaToElement(const FTypedElementHandle&
 	ViewportInteraction->ApplyDeltaToElement(InElementHandle, GetWidgetMode(), Widget ? Widget->GetCurrentAxis() : EAxisList::None, InputState, InDeltaTransform);
 }
 
+void FLevelEditorViewportClient::MirrorSelectedActors(const FVector& InMirrorScale)
+{
+	FScopedLevelDirtied LevelDirtyCallback;
+
+	const UTypedElementSelectionSet* SelectionSet = GetSelectionSet();
+	SelectionSet->ForEachSelectedObject<AActor>([this, &InMirrorScale, &LevelDirtyCallback](AActor* InActor)
+	{
+		ViewportInteraction->MirrorElement(UEngineElementsLibrary::AcquireEditorActorElementHandle(InActor), InMirrorScale);
+		LevelDirtyCallback.Request();
+		return true;
+	});
+
+	if (UBrushEditingSubsystem* BrushSubsystem = GEditor->GetEditorSubsystem<UBrushEditingSubsystem>())
+	{
+		BrushSubsystem->UpdateGeometryFromSelectedBrushes();
+	}
+
+	RedrawAllViewportsIntoThisScene();
+}
+
+void FLevelEditorViewportClient::MirrorSelectedElements(const FVector& InMirrorScale)
+{
+	FScopedLevelDirtied LevelDirtyCallback;
+
+	const UTypedElementSelectionSet* SelectionSet = GetSelectionSet();
+	SelectionSet->ForEachSelectedElementHandle([this, &InMirrorScale, &LevelDirtyCallback](const FTypedElementHandle& InElementHandle)
+	{
+		ViewportInteraction->MirrorElement(InElementHandle, InMirrorScale);
+		LevelDirtyCallback.Request();
+		return true;
+	});
+
+	if (UBrushEditingSubsystem* BrushSubsystem = GEditor->GetEditorSubsystem<UBrushEditingSubsystem>())
+	{
+		BrushSubsystem->UpdateGeometryFromSelectedBrushes();
+	}
+
+	RedrawAllViewportsIntoThisScene();
+}
+
 FTransform FLevelEditorViewportClient::CachePreDragActorTransform(const AActor* InActor)
 {
 	if (const FTransform* PreDragTransform = PreDragActorTransforms.Find(InActor))
diff --git a/Engine/Source/Editor/UnrealEd/Private/UnrealEdSrv.cpp b/Engine/Source/Editor/UnrealEd/Private/UnrealEdSrv.cpp
--- a/Engine/Source/Editor/UnrealEd/Private/UnrealEdSrv.cpp
+++ b/Engine/Source/Editor/UnrealEd/Private/UnrealEdSrv.cpp
@@ -104,6 +104,7 @@
 #endif
 #include "ActorGroupingUtils.h"
 #include "EdMode.h"
+#include "ILevelEditor.h"
 #include "Subsystems/BrushEditingSubsystem.h"
 #include "Subsystems/EditorActorSubsystem.h"
 
@@ -480,7 +481,7 @@
 				if( !bSilent )
 				{
 					GWarn->UpdateProgress( 1, 1 );
-					GWarn->EndSlowTask();
+					GWarn->EndSlowTask();
 				}
 			}
 		}
@@ -763,6 +764,13 @@
 		return Exec_Actor( InWorld, Str, Ar );
 	}
 	//------------------------------------------------------------------------------------
+	// ELEMENT: Element-related functions
+	//
+	else if (FParse::Command(&Str,TEXT("ELEMENT")))
+	{
+		return Exec_Element( InWorld, Str, Ar );
+	}
+	//------------------------------------------------------------------------------------
 	// MODE management (Global EDITOR mode):
 	//
 	else if( FParse::Command(&Str,TEXT("MODE")) )
@@ -1693,45 +1701,6 @@
 	return false;
 }
 
-static void MirrorActors(const FVector& MirrorScale)
-{
-	const FScopedTransaction Transaction( NSLOCTEXT("UnrealEd", "MirroringActors", "Mirroring Actors") );
-
-	// Fires ULevel::LevelDirtiedEvent when falling out of scope.
-	FScopedLevelDirtied		LevelDirtyCallback;
-
-	for ( FSelectionIterator It( GEditor->GetSelectedActorIterator() ) ; It ; ++It )
-	{
-		AActor* Actor = static_cast<AActor*>( *It );
-		checkSlow( Actor->IsA(AActor::StaticClass()) );
-
-		const FVector PivotLocation = GLevelEditorModeTools().PivotLocation;
-
-		Actor->Modify();
-		Actor->EditorApplyMirror( MirrorScale, PivotLocation );
-
-		ABrush* Brush = Cast< ABrush >(Actor);
-		if (Brush && Brush->GetBrushComponent())
-		{
-			Brush->GetBrushComponent()->RequestUpdateBrushCollision();
-		}
-
-		Actor->InvalidateLightingCache();
-		Actor->PostEditMove( true );
-
-		Actor->MarkPackageDirty();
-		LevelDirtyCallback.Request();
-	}
-
-	if (UBrushEditingSubsystem* BrushSubsystem = GEditor->GetEditorSubsystem<UBrushEditingSubsystem>())
-	{
-		BrushSubsystem->UpdateGeometryFromSelectedBrushes();
-	}
-
-	GEditor->RedrawLevelEditingViewports();
-}
-
-
 /**
 * Gathers up a list of selection FPolys from selected static meshes.
 *
@@ -2230,9 +2199,13 @@
 		if( !MirrorScale.Y )		MirrorScale.Y = 1;
 		if( !MirrorScale.Z )		MirrorScale.Z = 1;
 
-		const FScopedTransaction Transaction(NSLOCTEXT("UnrealEd", "MirroringActors", "Mirroring Actors"));
-		MirrorActors(MirrorScale);
-		RebuildAlteredBSP(); // Update the Bsp of any levels containing a modified brush
+		if (GCurrentLevelEditingViewportClient)
+		{
+			const FScopedTransaction Transaction(NSLOCTEXT("UnrealEd", "MirroringActors", "Mirroring Actors"));
+			GCurrentLevelEditingViewportClient->MirrorSelectedActors(MirrorScale);
+			RebuildAlteredBSP(); // Update the Bsp of any levels containing a modified brush
+		}
+
 		return true;
 	}
 	else if( FParse::Command(&Str,TEXT("DELTAMOVE")) )
@@ -2241,14 +2214,17 @@
 		FVector DeltaMove = FVector::ZeroVector;
 		GetFVECTOR( Str, DeltaMove );
 
-		FEditorModeTools& Tools = GLevelEditorModeTools();
-		Tools.SetPivotLocation( Tools.PivotLocation + DeltaMove, false );
-
 		if (GCurrentLevelEditingViewportClient)
 		{
+			if (TSharedPtr<ILevelEditor> LevelEditor = GCurrentLevelEditingViewportClient->ParentLevelEditor.Pin())
+			{
+				FEditorModeTools& Tools = LevelEditor->GetEditorModeManager();
+				Tools.SetPivotLocation(Tools.PivotLocation + DeltaMove, false);
+			}
+
 			GCurrentLevelEditingViewportClient->ApplyDeltaToActors(DeltaMove, FRotator::ZeroRotator, FVector::ZeroVector);
+			RedrawLevelEditingViewports();
 		}
-		RedrawLevelEditingViewports();
 
 		return true;
 	}
@@ -2825,7 +2801,54 @@
 	return false;
 }
 
+bool UUnrealEdEngine::Exec_Element( UWorld* InWorld, const TCHAR* Str, FOutputDevice& Ar )
+{
+	// Keep a pointer to the beginning of the string to use for message displaying purposes
+	const TCHAR* const FullStr = Str;
 
+	if (FParse::Command(&Str, TEXT("MIRROR")))
+	{
+		FVector MirrorScale(1, 1, 1);
+		GetFVECTOR(Str, MirrorScale);
+		// We can't have zeroes in the vector
+		if (!MirrorScale.X)		MirrorScale.X = 1;
+		if (!MirrorScale.Y)		MirrorScale.Y = 1;
+		if (!MirrorScale.Z)		MirrorScale.Z = 1;
+
+		if (GCurrentLevelEditingViewportClient)
+		{
+			const FScopedTransaction Transaction(NSLOCTEXT("UnrealEd", "MirroringElements", "Mirroring Elements"));
+			GCurrentLevelEditingViewportClient->MirrorSelectedElements(MirrorScale);
+			RebuildAlteredBSP(); // Update the Bsp of any levels containing a modified brush
+		}
+
+		return true;
+	}
+	
+	if (FParse::Command(&Str, TEXT("DELTAMOVE")))
+	{
+		const FScopedTransaction Transaction(NSLOCTEXT("UnrealEd", "DeltaMovElements", "Move Elements by Delta"));
+		FVector DeltaMove = FVector::ZeroVector;
+		GetFVECTOR(Str, DeltaMove);
+
+		if (GCurrentLevelEditingViewportClient)
+		{
+			if (TSharedPtr<ILevelEditor> LevelEditor = GCurrentLevelEditingViewportClient->ParentLevelEditor.Pin())
+			{
+				FEditorModeTools& Tools = LevelEditor->GetEditorModeManager();
+				Tools.SetPivotLocation(Tools.PivotLocation + DeltaMove, false);
+			}
+
+			GCurrentLevelEditingViewportClient->ApplyDeltaToSelectedElements(FTransform(FRotator::ZeroRotator, DeltaMove, FVector::ZeroVector));
+			RedrawLevelEditingViewports();
+		}
+
+		return true;
+	}
+
+	return false;
+}
+
 bool UUnrealEdEngine::Exec_Mode( const TCHAR* Str, FOutputDevice& Ar )
 {
 	int32 DWord1;
diff --git a/Engine/Source/Editor/UnrealEd/Public/Elements/Actor/ActorElementEditorViewportInteractionCustomization.h b/Engine/Source/Editor/UnrealEd/Public/Elements/Actor/ActorElementEditorViewportInteractionCustomization.h
--- a/Engine/Source/Editor/UnrealEd/Public/Elements/Actor/ActorElementEditorViewportInteractionCustomization.h
+++ b/Engine/Source/Editor/UnrealEd/Public/Elements/Actor/ActorElementEditorViewportInteractionCustomization.h
@@ -12,6 +12,7 @@ class UNREALED_API FActorElementEditorViewportInteractionCustomization : public
 public:
 	virtual bool GetGizmoPivotLocation(const TTypedElement<UTypedElementWorldInterface>& InElementWorldHandle, const UE::Widget::EWidgetMode InWidgetMode, FVector& OutPivotLocation) override;
 	virtual void GizmoManipulationDeltaUpdate(const TTypedElement<UTypedElementWorldInterface>& InElementWorldHandle, const UE::Widget::EWidgetMode InWidgetMode, const EAxisList::Type InDragAxis, const FInputDeviceState& InInputState, const FTransform& InDeltaTransform, const FVector& InPivotLocation) override;
+	virtual void MirrorElement(const TTypedElement<UTypedElementWorldInterface>& InElementWorldHandle, const FVector& InMirrorScale, const FVector& InPivotLocation) override;
 
 	static void ApplyDeltaToActor(AActor* InActor, const bool InIsDelta, const FVector* InDeltaTranslationPtr, const FRotator* InDeltaRotationPtr, const FVector* InDeltaScalePtr, const FVector& InPivotLocation, const FInputDeviceState& InInputState);
 };
diff --git a/Engine/Source/Editor/UnrealEd/Public/LevelEditorViewport.h b/Engine/Source/Editor/UnrealEd/Public/LevelEditorViewport.h
--- a/Engine/Source/Editor/UnrealEd/Public/LevelEditorViewport.h
+++ b/Engine/Source/Editor/UnrealEd/Public/LevelEditorViewport.h
@@ -322,6 +322,9 @@ public:
 	void ApplyDeltaToSelectedElements(const FTransform& InDeltaTransform);
 	void ApplyDeltaToElement(const FTypedElementHandle& InElementHandle, const FTransform& InDeltaTransform);
 
+	void MirrorSelectedActors(const FVector& InMirrorScale);
+	void MirrorSelectedElements(const FVector& InMirrorScale);
+
 	virtual void SetIsSimulateInEditorViewport( bool bInIsSimulateInEditorViewport ) override;
 
 	/**
diff --git a/Engine/Source/Runtime/Engine/Private/Elements/Actor/ActorElementWorldInterface.cpp b/Engine/Source/Runtime/Engine/Private/Elements/Actor/ActorElementWorldInterface.cpp
--- a/Engine/Source/Runtime/Engine/Private/Elements/Actor/ActorElementWorldInterface.cpp
+++ b/Engine/Source/Runtime/Engine/Private/Elements/Actor/ActorElementWorldInterface.cpp
@@ -4,13 +4,16 @@
 #include "Elements/Actor/ActorElementData.h"
 #include "GameFramework/Actor.h"
 
+#include "Elements/Component/ComponentElementData.h"
+#include "Components/PrimitiveComponent.h"
+
 UWorld* UActorElementWorldInterface::GetOwnerWorld(const FTypedElementHandle& InElementHandle)
 {
 	const AActor* Actor = ActorElementDataUtil::GetActorFromHandle(InElementHandle);
 	return Actor ? Actor->GetWorld() : nullptr;
 }
 
-bool UActorElementWorldInterface::GetWorldBounds(const FTypedElementHandle& InElementHandle, FBoxSphereBounds& OutBounds)
+bool UActorElementWorldInterface::GetBounds(const FTypedElementHandle& InElementHandle, FBoxSphereBounds& OutBounds)
 {
 	if (const AActor* Actor = ActorElementDataUtil::GetActorFromHandle(InElementHandle))
 	{
@@ -43,3 +46,143 @@ bool UActorElementWorldInterface::SetWorldTransform(const FTypedElementHandle& I
 
 	return false;
 }
+
+bool UActorElementWorldInterface::GetRelativeTransform(const FTypedElementHandle& InElementHandle, FTransform& OutTransform)
+{
+	if (const AActor* Actor = ActorElementDataUtil::GetActorFromHandle(InElementHandle))
+	{
+		if (const USceneComponent* RootComponent = Actor->GetRootComponent())
+		{
+			OutTransform = RootComponent->GetRelativeTransform();
+		}
+		else
+		{
+			OutTransform = FTransform::Identity;
+		}
+		return true;
+	}
+
+	return false;
+}
+
+bool UActorElementWorldInterface::SetRelativeTransform(const FTypedElementHandle& InElementHandle, const FTransform& InTransform)
+{
+	if (AActor* Actor = ActorElementDataUtil::GetActorFromHandle(InElementHandle))
+	{
+		if (USceneComponent* RootComponent = Actor->GetRootComponent())
+		{
+			Actor->Modify();
+			RootComponent->SetRelativeTransform(InTransform);
+			return true;
+		}
+	}
+
+	return false;
+}
+
+bool UActorElementWorldInterface::FindSuitableTransformAtPoint(const FTypedElementHandle& InElementHandle, const FTransform& InPotentialTransform, FTransform& OutSuitableTransform)
+{
+	if (const AActor* Actor = ActorElementDataUtil::GetActorFromHandle(InElementHandle))
+	{
+		UWorld* World = Actor->GetWorld();
+		const UPrimitiveComponent* PrimComponent = Cast<UPrimitiveComponent>(Actor->GetRootComponent());
+
+		if (World && PrimComponent && PrimComponent->IsQueryCollisionEnabled())
+		{
+			const FVector PivotOffset = PrimComponent->Bounds.Origin - Actor->GetActorLocation();
+
+			FVector NewLocation = InPotentialTransform.GetTranslation();
+			FRotator NewRotation = InPotentialTransform.Rotator();
+
+			// Apply the pivot offset to the desired location
+			NewLocation += PivotOffset;
+
+			// Check if able to find an acceptable destination for this actor that doesn't embed it in world geometry
+			if (World->FindTeleportSpot(Actor, NewLocation, NewRotation))
+			{
+				// Undo the pivot offset
+				NewLocation -= PivotOffset;
+
+				OutSuitableTransform = InPotentialTransform;
+				OutSuitableTransform.SetTranslation(NewLocation);
+				OutSuitableTransform.SetRotation(NewRotation.Quaternion());
+				return true;
+			}
+		}
+		else
+		{
+			OutSuitableTransform = InPotentialTransform;
+			return true;
+		}
+	}
+
+	return false;
+}
+
+bool UActorElementWorldInterface::FindSuitableTransformAlongPath(const FTypedElementHandle& InElementHandle, const FVector& InPathStart, const FVector& InPathEnd, const FCollisionShape& InTestShape, TArrayView<const FTypedElementHandle> InElementsToIgnore, FTransform& OutSuitableTransform)
+{
+	if (const AActor* Actor = ActorElementDataUtil::GetActorFromHandle(InElementHandle))
+	{	
+		if (const UWorld* World = Actor->GetWorld())
+		{
+			FCollisionQueryParams Params(SCENE_QUERY_STAT(FindSuitableTransformAlongPath), false);
+			
+			// Don't hit ourself
+			{
+				Params.AddIgnoredActor(Actor);
+
+				TArray<AActor*> ChildActors;
+				Actor->GetAllChildActors(ChildActors);
+				Params.AddIgnoredActors(ChildActors);
+			}
+
+			return UActorElementWorldInterface::FindSuitableTransformAlongPath_WorldSweep(World, InPathStart, InPathEnd, InTestShape, InElementsToIgnore, Params, OutSuitableTransform);
+		}
+	}
+
+	return false;
+}
+
+bool UActorElementWorldInterface::FindSuitableTransformAlongPath_WorldSweep(const UWorld* InWorld, const FVector& InPathStart, const FVector& InPathEnd, const FCollisionShape& InTestShape, TArrayView<const FTypedElementHandle> InElementsToIgnore, FCollisionQueryParams& InOutParams, FTransform& OutSuitableTransform)
+{
+	for (const FTypedElementHandle& ElementToIgnore : InElementsToIgnore)
+	{
+		AddIgnoredCollisionQueryElement(ElementToIgnore, InOutParams);
+	}
+
+	FHitResult Hit(1.0f);
+	if (InWorld->SweepSingleByChannel(Hit, InPathStart, InPathEnd, FQuat::Identity, ECC_WorldStatic, InTestShape, InOutParams))
+	{
+		FVector NewLocation = Hit.Location;
+		NewLocation.Z += KINDA_SMALL_NUMBER; // Move the new desired location up by an error tolerance
+
+		//@todo: This doesn't take into account that rotating the actor changes LocationOffset.
+		FRotator NewRotation = Hit.Normal.Rotation();
+		NewRotation.Pitch -= 90.f;
+
+		OutSuitableTransform.SetTranslation(NewLocation);
+		OutSuitableTransform.SetRotation(NewRotation.Quaternion());
+		OutSuitableTransform.SetScale3D(FVector::OneVector);
+		return true;
+	}
+
+	return false;
+}
+
+void UActorElementWorldInterface::AddIgnoredCollisionQueryElement(const FTypedElementHandle& InElementHandle, FCollisionQueryParams& InOutParams)
+{
+	if (const AActor* Actor = ActorElementDataUtil::GetActorFromHandle(InElementHandle, /*bSilent*/true))
+	{
+		InOutParams.AddIgnoredActor(Actor);
+		return;
+	}
+
+	if (const UActorComponent* Component = ComponentElementDataUtil::GetComponentFromHandle(InElementHandle, /*bSilent*/true))
+	{
+		if (const UPrimitiveComponent* PrimComponent = Cast<UPrimitiveComponent>(Component))
+		{
+			InOutParams.AddIgnoredComponent(PrimComponent);
+		}
+		return;
+	}
+}
diff --git a/Engine/Source/Runtime/Engine/Private/Elements/Component/ComponentElementWorldInterface.cpp b/Engine/Source/Runtime/Engine/Private/Elements/Component/ComponentElementWorldInterface.cpp
--- a/Engine/Source/Runtime/Engine/Private/Elements/Component/ComponentElementWorldInterface.cpp
+++ b/Engine/Source/Runtime/Engine/Private/Elements/Component/ComponentElementWorldInterface.cpp
@@ -2,8 +2,9 @@
 
 #include "Elements/Component/ComponentElementWorldInterface.h"
 #include "Elements/Component/ComponentElementData.h"
-#include "Components/ActorComponent.h"
-#include "Components/SceneComponent.h"
+#include "Components/PrimitiveComponent.h"
+
+#include "Elements/Actor/ActorElementWorldInterface.h"
 
 bool UComponentElementWorldInterface::CanEditElement(const FTypedElementHandle& InElementHandle)
 {
@@ -17,7 +18,7 @@ UWorld* UComponentElementWorldInterface::GetOwnerWorld(const FTypedElementHandle
 	return Component ? Component->GetWorld() : nullptr;
 }
 
-bool UComponentElementWorldInterface::GetWorldBounds(const FTypedElementHandle& InElementHandle, FBoxSphereBounds& OutBounds)
+bool UComponentElementWorldInterface::GetBounds(const FTypedElementHandle& InElementHandle, FBoxSphereBounds& OutBounds)
 {
 	if (const UActorComponent* Component = ComponentElementDataUtil::GetComponentFromHandleChecked(InElementHandle))
 	{
@@ -59,3 +60,53 @@ bool UComponentElementWorldInterface::SetWorldTransform(const FTypedElementHandl
 	
 	return false;
 }
+
+bool UComponentElementWorldInterface::GetRelativeTransform(const FTypedElementHandle& InElementHandle, FTransform& OutTransform)
+{
+	if (const UActorComponent* Component = ComponentElementDataUtil::GetComponentFromHandleChecked(InElementHandle))
+	{
+		if (const USceneComponent* SceneComponent = Cast<USceneComponent>(Component))
+		{
+			OutTransform = SceneComponent->GetRelativeTransform();
+			return true;
+		}
+	}
+
+	return false;
+}
+
+bool UComponentElementWorldInterface::SetRelativeTransform(const FTypedElementHandle& InElementHandle, const FTransform& InTransform)
+{
+	if (UActorComponent* Component = ComponentElementDataUtil::GetComponentFromHandleChecked(InElementHandle))
+	{
+		if (USceneComponent* SceneComponent = Cast<USceneComponent>(Component))
+		{
+			SceneComponent->Modify();
+			SceneComponent->SetRelativeTransform(InTransform);
+			return true;
+		}
+	}
+
+	return false;
+}
+
+bool UComponentElementWorldInterface::FindSuitableTransformAlongPath(const FTypedElementHandle& InElementHandle, const FVector& InPathStart, const FVector& InPathEnd, const FCollisionShape& InTestShape, TArrayView<const FTypedElementHandle> InElementsToIgnore, FTransform& OutSuitableTransform)
+{
+	if (const UActorComponent* Component = ComponentElementDataUtil::GetComponentFromHandleChecked(InElementHandle))
+	{
+		if (const UWorld* World = Component->GetWorld())
+		{
+			FCollisionQueryParams Params(SCENE_QUERY_STAT(FindSuitableTransformAlongPath), false);
+
+			// Don't hit ourself
+			if (const UPrimitiveComponent* PrimComponent = Cast<UPrimitiveComponent>(Component))
+			{
+				Params.AddIgnoredComponent(PrimComponent);
+			}
+
+			return UActorElementWorldInterface::FindSuitableTransformAlongPath_WorldSweep(World, InPathStart, InPathEnd, InTestShape, InElementsToIgnore, Params, OutSuitableTransform);
+		}
+	}
+
+	return false;
+}
diff --git a/Engine/Source/Runtime/Engine/Public/Elements/Actor/ActorElementWorldInterface.h b/Engine/Source/Runtime/Engine/Public/Elements/Actor/ActorElementWorldInterface.h
--- a/Engine/Source/Runtime/Engine/Public/Elements/Actor/ActorElementWorldInterface.h
+++ b/Engine/Source/Runtime/Engine/Public/Elements/Actor/ActorElementWorldInterface.h
@@ -5,6 +5,8 @@
 #include "Elements/Interfaces/TypedElementWorldInterface.h"
 #include "ActorElementWorldInterface.generated.h"
 
+struct FCollisionQueryParams;
+
 UCLASS()
 class ENGINE_API UActorElementWorldInterface : public UTypedElementWorldInterface
 {
@@ -12,7 +14,14 @@ class ENGINE_API UActorElementWorldInterface : public UTypedElementWorldInterfac
 
 public:
 	virtual UWorld* GetOwnerWorld(const FTypedElementHandle& InElementHandle) override;
-	virtual bool GetWorldBounds(const FTypedElementHandle& InElementHandle, FBoxSphereBounds& OutBounds) override;
+	virtual bool GetBounds(const FTypedElementHandle& InElementHandle, FBoxSphereBounds& OutBounds) override;
 	virtual bool GetWorldTransform(const FTypedElementHandle& InElementHandle, FTransform& OutTransform) override;
 	virtual bool SetWorldTransform(const FTypedElementHandle& InElementHandle, const FTransform& InTransform) override;
+	virtual bool GetRelativeTransform(const FTypedElementHandle& InElementHandle, FTransform& OutTransform) override;
+	virtual bool SetRelativeTransform(const FTypedElementHandle& InElementHandle, const FTransform& InTransform) override;
+	virtual bool FindSuitableTransformAtPoint(const FTypedElementHandle& InElementHandle, const FTransform& InPotentialTransform, FTransform& OutSuitableTransform) override;
+	virtual bool FindSuitableTransformAlongPath(const FTypedElementHandle& InElementHandle, const FVector& InPathStart, const FVector& InPathEnd, const FCollisionShape& InTestShape, TArrayView<const FTypedElementHandle> InElementsToIgnore, FTransform& OutSuitableTransform) override;
+
+	static bool FindSuitableTransformAlongPath_WorldSweep(const UWorld* InWorld, const FVector& InPathStart, const FVector& InPathEnd, const FCollisionShape& InTestShape, TArrayView<const FTypedElementHandle> InElementsToIgnore, FCollisionQueryParams& InOutParams, FTransform& OutSuitableTransform);
+	static void AddIgnoredCollisionQueryElement(const FTypedElementHandle& InElementHandle, FCollisionQueryParams& InOutParams);
 };
diff --git a/Engine/Source/Runtime/Engine/Public/Elements/Component/ComponentElementWorldInterface.h b/Engine/Source/Runtime/Engine/Public/Elements/Component/ComponentElementWorldInterface.h
--- a/Engine/Source/Runtime/Engine/Public/Elements/Component/ComponentElementWorldInterface.h
+++ b/Engine/Source/Runtime/Engine/Public/Elements/Component/ComponentElementWorldInterface.h
@@ -13,7 +13,10 @@ class ENGINE_API UComponentElementWorldInterface : public UTypedElementWorldInte
 public:
 	virtual bool CanEditElement(const FTypedElementHandle& InElementHandle) override;
 	virtual UWorld* GetOwnerWorld(const FTypedElementHandle& InElementHandle) override;
-	virtual bool GetWorldBounds(const FTypedElementHandle& InElementHandle, FBoxSphereBounds& OutBounds) override;
+	virtual bool GetBounds(const FTypedElementHandle& InElementHandle, FBoxSphereBounds& OutBounds) override;
 	virtual bool GetWorldTransform(const FTypedElementHandle& InElementHandle, FTransform& OutTransform) override;
 	virtual bool SetWorldTransform(const FTypedElementHandle& InElementHandle, const FTransform& InTransform) override;
+	virtual bool GetRelativeTransform(const FTypedElementHandle& InElementHandle, FTransform& OutTransform) override;
+	virtual bool SetRelativeTransform(const FTypedElementHandle& InElementHandle, const FTransform& InTransform) override;
+	virtual bool FindSuitableTransformAlongPath(const FTypedElementHandle& InElementHandle, const FVector& InPathStart, const FVector& InPathEnd, const FCollisionShape& InTestShape, TArrayView<const FTypedElementHandle> InElementsToIgnore, FTransform& OutSuitableTransform) override;
 };
diff --git a/Engine/Source/Runtime/Engine/Public/Elements/Interfaces/TypedElementWorldInterface.h b/Engine/Source/Runtime/Engine/Public/Elements/Interfaces/TypedElementWorldInterface.h
--- a/Engine/Source/Runtime/Engine/Public/Elements/Interfaces/TypedElementWorldInterface.h
+++ b/Engine/Source/Runtime/Engine/Public/Elements/Interfaces/TypedElementWorldInterface.h
@@ -6,6 +6,7 @@
 #include "TypedElementWorldInterface.generated.h"
 
 class UWorld;
+struct FCollisionShape;
 
 UCLASS(Abstract)
 class ENGINE_API UTypedElementWorldInterface : public UTypedElementInterface
@@ -26,10 +27,10 @@ public:
 	virtual UWorld* GetOwnerWorld(const FTypedElementHandle& InElementHandle) { return nullptr; }
 
 	/**
-	 * Get the bounds of this element within its owner world, if any.
+	 * Get the bounds of this element, if any.
 	 */
 	UFUNCTION(BlueprintPure, Category="TypedElementInterfaces|World")
-	virtual bool GetWorldBounds(const FTypedElementHandle& InElementHandle, FBoxSphereBounds& OutBounds) { return false; }
+	virtual bool GetBounds(const FTypedElementHandle& InElementHandle, FBoxSphereBounds& OutBounds) { return false; }
 
 	/**
 	 * Get the transform of this element within its owner world, if any.
@@ -42,6 +43,53 @@ public:
 	 */
 	UFUNCTION(BlueprintCallable, Category="TypedElementInterfaces|World")
 	virtual bool SetWorldTransform(const FTypedElementHandle& InElementHandle, const FTransform& InTransform) { return false; }
+
+	/**
+	 * Get the transform of this element relative to its parent, if any.
+	 */
+	UFUNCTION(BlueprintPure, Category="TypedElementInterfaces|World")
+	virtual bool GetRelativeTransform(const FTypedElementHandle& InElementHandle, FTransform& OutTransform) { return GetWorldTransform(InElementHandle, OutTransform); }
+	
+	/**
+	 * Attempt to set the transform of this element relative to its parent.
+	 */
+	UFUNCTION(BlueprintCallable, Category="TypedElementInterfaces|World")
+	virtual bool SetRelativeTransform(const FTypedElementHandle& InElementHandle, const FTransform& InTransform) { return SetWorldTransform(InElementHandle, InTransform); }
+
+	/**
+	 * Notify that this element is about to be moved.
+	 */
+	UFUNCTION(BlueprintCallable, Category="TypedElementInterfaces|World")
+	virtual void NotifyMovementStarted(const FTypedElementHandle& InElementHandle) {}
+
+	/**
+	 * Notify that this element is currently being moved.
+	 */
+	UFUNCTION(BlueprintCallable, Category="TypedElementInterfaces|World")
+	virtual void NotifyMovementOngoing(const FTypedElementHandle& InElementHandle) {}
+
+	/**
+	 * Notify that this element is done being moved.
+	 */
+	UFUNCTION(BlueprintCallable, Category="TypedElementInterfaces|World")
+	virtual void NotifyMovementEnded(const FTypedElementHandle& InElementHandle) {}
+
+	/**
+	 * Attempt to find a suitable (non-intersecting) transform for the given element at the given point.
+	 */
+	virtual bool FindSuitableTransformAtPoint(const FTypedElementHandle& InElementHandle, const FTransform& InPotentialTransform, FTransform& OutSuitableTransform)
+	{
+		OutSuitableTransform = InPotentialTransform;
+		return true;
+	}
+
+	/**
+	 * Attempt to find a suitable (non-intersecting) transform for the given element along the given path.
+	 */
+	virtual bool FindSuitableTransformAlongPath(const FTypedElementHandle& InElementHandle, const FVector& InPathStart, const FVector& InPathEnd, const FCollisionShape& InTestShape, TArrayView<const FTypedElementHandle> InElementsToIgnore, FTransform& OutSuitableTransform)
+	{
+		return false;
+	}
 };
 
 template <>
@@ -49,7 +97,14 @@ struct TTypedElement<UTypedElementWorldInterface> : public TTypedElementBase<UTy
 {
 	bool CanEditElement() const { return InterfacePtr->CanEditElement(*this); }
 	UWorld* GetOwnerWorld() const { return InterfacePtr->GetOwnerWorld(*this); }
-	bool GetWorldBounds(FBoxSphereBounds& OutBounds) const { return InterfacePtr->GetWorldBounds(*this, OutBounds); }
+	bool GetBounds(FBoxSphereBounds& OutBounds) const { return InterfacePtr->GetBounds(*this, OutBounds); }
 	bool GetWorldTransform(FTransform& OutTransform) const { return InterfacePtr->GetWorldTransform(*this, OutTransform); }
 	bool SetWorldTransform(const FTransform& InTransform) const { return InterfacePtr->SetWorldTransform(*this, InTransform); }
+	bool GetRelativeTransform(FTransform& OutTransform) const { return InterfacePtr->GetRelativeTransform(*this, OutTransform); }
+	bool SetRelativeTransform(const FTransform& InTransform) const { return InterfacePtr->SetRelativeTransform(*this, InTransform); }
+	void NotifyMovementStarted() const { InterfacePtr->NotifyMovementStarted(*this); }
+	void NotifyMovementOngoing() const { InterfacePtr->NotifyMovementOngoing(*this); }
+	void NotifyMovementEnded() const { InterfacePtr->NotifyMovementEnded(*this); }
+	bool FindSuitableTransformAtPoint(const FTransform& InPotentialTransform, FTransform& OutSuitableTransform) const { return InterfacePtr->FindSuitableTransformAtPoint(*this, InPotentialTransform, OutSuitableTransform); }
+	bool FindSuitableTransformAlongPath(const FVector& InPathStart, const FVector& InPathEnd, const FCollisionShape& InTestShape, TArrayView<const FTypedElementHandle> InElementsToIgnore, FTransform& OutSuitableTransform) const { return InterfacePtr->FindSuitableTransformAlongPath(*this, InPathStart, InPathEnd, InTestShape, InElementsToIgnore, OutSuitableTransform); }
 };
