Fix Snap To Floor for spline points. Destination is optional for component visualizer Snap To actions.

diff --git a/Engine/Source/Editor/LevelEditor/Private/LevelEditorActions.cpp b/Engine/Source/Editor/LevelEditor/Private/LevelEditorActions.cpp
--- a/Engine/Source/Editor/LevelEditor/Private/LevelEditorActions.cpp
+++ b/Engine/Source/Editor/LevelEditor/Private/LevelEditorActions.cpp
@@ -3036,14 +3036,16 @@ void FLevelEditorActionCallbacks::SnapTo_Clicked( const UTypedElementSelectionSe
 
 	// Let the component visualizers try to handle the selection.
 	// TODO: Should this also take an element?
-	if (TTypedElement<UTypedElementObjectInterface> DestinationObjectHandle = InSelectionSet->GetElementList()->GetElement<UTypedElementObjectInterface>(InDestination))
 	{
-		if (AActor* DestinationActor = Cast<AActor>(DestinationObjectHandle.GetObject()))
+		AActor* DestinationActor = nullptr;
+		if (TTypedElement<UTypedElementObjectInterface> DestinationObjectHandle = InSelectionSet->GetElementList()->GetElement<UTypedElementObjectInterface>(InDestination))
 		{
-			if (GUnrealEd->ComponentVisManager.HandleSnapTo(InAlign, InUseLineTrace, InUseBounds, InUsePivot, DestinationActor))
-			{
-				return;
-			}
+			DestinationActor = Cast<AActor>(DestinationObjectHandle.GetObject());
+		}
+
+		if (GUnrealEd->ComponentVisManager.HandleSnapTo(InAlign, InUseLineTrace, InUseBounds, InUsePivot, DestinationActor))
+		{
+			return;
 		}
 	}
 
