Fix non-unity compile error.

#rnx

diff --git a/Engine/Source/Editor/LevelEditor/Private/SLevelEditorModeContent.cpp b/Engine/Source/Editor/LevelEditor/Private/SLevelEditorModeContent.cpp
--- a/Engine/Source/Editor/LevelEditor/Private/SLevelEditorModeContent.cpp
+++ b/Engine/Source/Editor/LevelEditor/Private/SLevelEditorModeContent.cpp
@@ -7,6 +7,7 @@
 #include "Editor/EditorPerProjectUserSettings.h"
 #include "EditorModeManager.h"
 #include "Widgets/Docking/SDockTab.h"
+#include "Toolkits/IToolkit.h"
 
 #define LOCTEXT_NAMESPACE "SLevelEditorModeContent"
 
