Ensure that dynamic viewport commands are registered/updated prior to constructing a new level viewport client.

Fixes the potential for encountering a regression crash in the editor that was introduced with CL 14886456.

#rnx

diff --git a/Engine/Source/Editor/LevelEditor/Private/SLevelViewport.cpp b/Engine/Source/Editor/LevelEditor/Private/SLevelViewport.cpp
--- a/Engine/Source/Editor/LevelEditor/Private/SLevelViewport.cpp
+++ b/Engine/Source/Editor/LevelEditor/Private/SLevelViewport.cpp
@@ -191,6 +191,12 @@ void SLevelViewport::Construct(const FArguments& InArgs, const FAssetEditorViewp
 	StartingSimulateBorder = FEditorStyle::GetBrush( "LevelViewport.StartingSimulateBorder" );
 	ReturningToEditorBorder = FEditorStyle::GetBrush( "LevelViewport.ReturningToEditorBorder" );
 
+	// Register/update viewport commands that are somewhat dynamic (can grow in size) based on loaded assets.
+	// Note: This needs to happen prior to constructing the viewport client and SEditorViewport instances below.
+	FLevelViewportCommands& LevelViewportCommands = FLevelViewportCommands::Get();
+	LevelViewportCommands.RegisterShowSpriteCommands();
+	LevelViewportCommands.RegisterShowVolumeCommands();
+
 	// Default level viewport client values for settings that could appear in layout config ini
 	FLevelEditorViewportInstanceSettings ViewportInstanceSettings;
 	ViewportInstanceSettings.ViewportType = InConstructionArguments.ViewportType;
@@ -1564,8 +1570,7 @@ void SLevelViewport::BindShowCommands( FUICommandList& OutCommandList )
 			LevelViewportCommands.HideAllVolumes,
 			FExecuteAction::CreateSP( this, &SLevelViewport::OnToggleAllVolumeActors, false ) );
 		
-
-		LevelViewportCommands.RegisterShowVolumeCommands();
+		// Note: This list is somewhat dynamic and is registered/updated at viewport construction time, @see Construct().
 		const TArray<FLevelViewportCommands::FShowMenuCommand>& ShowVolumeCommands = LevelViewportCommands.ShowVolumeCommands;
 		for (int32 VolumeCommandIndex = 0; VolumeCommandIndex < ShowVolumeCommands.Num(); ++VolumeCommandIndex)
 		{
@@ -1618,7 +1623,7 @@ void SLevelViewport::BindShowCommands( FUICommandList& OutCommandList )
 			FExecuteAction::CreateSP( this, &SLevelViewport::OnToggleAllSpriteCategories, false ) );
 
 		// Bind each show flag to the same delegate.  We use the delegate payload system to figure out what show flag we are dealing with
-		LevelViewportCommands.RegisterShowSpriteCommands();
+		// Note: This list is somewhat dynamic and is registered/updated at viewport construction time, @see Construct().
 		const TArray<FLevelViewportCommands::FShowMenuCommand>& ShowSpriteCommands = LevelViewportCommands.ShowSpriteCommands;
 		for (int32 SpriteCommandIndex = 0; SpriteCommandIndex < ShowSpriteCommands.Num(); ++SpriteCommandIndex)
 		{
