Narrowed fix to avoid an editor crash when a new sprite category is introduced by a Blueprint asset on editor load.

diff --git a/Engine/Source/Editor/LevelEditor/Private/LevelViewportActions.cpp b/Engine/Source/Editor/LevelEditor/Private/LevelViewportActions.cpp
--- a/Engine/Source/Editor/LevelEditor/Private/LevelViewportActions.cpp
+++ b/Engine/Source/Editor/LevelEditor/Private/LevelViewportActions.cpp
@@ -10,6 +10,7 @@
 #include "Stats/StatsData.h"
 #include "BufferVisualizationData.h"
 #include "Bookmarks/BookmarkUI.h"
+#include "UnrealEdGlobals.h"
 
 #define LOCTEXT_NAMESPACE "LevelViewportActions"
 
@@ -301,7 +302,13 @@ void FLevelViewportCommands::RegisterShowSpriteCommands()
 				= FUICommandInfoDecl(this->AsShared(), CommandName, LocalizedName, SpriteInfo.Description)
 				.UserInterfaceType(EUserInterfaceActionType::ToggleButton);
 
-			ShowSpriteCommands.Add(FLevelViewportCommands::FShowMenuCommand(ShowSpriteCommand, SpriteInfo.DisplayName));
+			const int32 ShowSpriteCommandIndex = ShowSpriteCommands.Add(FLevelViewportCommands::FShowMenuCommand(ShowSpriteCommand, SpriteInfo.DisplayName));
+
+			// Update the global table that maps each sprite category to its corresponding visibility command.
+			if (GUnrealEd)
+			{
+				GUnrealEd->SpriteIDToIndexMap.Add(SpriteInfo.Category, ShowSpriteCommandIndex);
+			}
 		}
 	}
 }
