Fix padding on level editor toolbar

diff --git a/Engine/Source/Editor/LevelEditor/Private/LevelEditorToolBar.cpp b/Engine/Source/Editor/LevelEditor/Private/LevelEditorToolBar.cpp
--- a/Engine/Source/Editor/LevelEditor/Private/LevelEditorToolBar.cpp
+++ b/Engine/Source/Editor/LevelEditor/Private/LevelEditorToolBar.cpp
@@ -1426,6 +1426,7 @@ TSharedRef< SWidget > FLevelEditorToolBar::MakeLevelEditorToolBar( const TShared
 			SNew(SHorizontalBox)
 			.IsEnabled(FSlateApplication::Get().GetNormalExecutionAttribute())
 			+ SHorizontalBox::Slot()
+			.Padding(7.0f, 0.0f, 0.0f, 0.0f)
 			.AutoWidth()
 			[
 				UToolMenus::Get()->GenerateWidget("LevelEditor.LevelEditorToolBar.AssetsToolBar", MenuContext)
