Move play and platform buttons to the center of the level editor toolbar

diff --git a/Engine/Source/Editor/LevelEditor/Private/LevelEditorToolBar.cpp b/Engine/Source/Editor/LevelEditor/Private/LevelEditorToolBar.cpp
--- a/Engine/Source/Editor/LevelEditor/Private/LevelEditorToolBar.cpp
+++ b/Engine/Source/Editor/LevelEditor/Private/LevelEditorToolBar.cpp
@@ -1294,7 +1294,6 @@ void FLevelEditorToolBar::RegisterLevelEditorToolBar( const TSharedRef<FUIComman
 	}
 
 	UToolMenu* SettingsToolBar = UToolMenus::Get()->RegisterMenu("LevelEditor.LevelEditorToolBar.MarketplaceToolBar", NAME_None, EMultiBoxType::SlimHorizontalToolBar);
-
 	SettingsToolBar->StyleName = "AssetEditorToolbar";
 	{
 		FToolMenuSection& SettingsSection = SettingsToolBar->AddSection("Settings");
@@ -1306,7 +1305,6 @@ void FLevelEditorToolBar::RegisterLevelEditorToolBar( const TSharedRef<FUIComman
 
 	UToolMenu* PlayToolBar = UToolMenus::Get()->RegisterMenu("LevelEditor.LevelEditorToolBar.PlayToolBar", NAME_None, EMultiBoxType::SlimHorizontalToolBar);
 	PlayToolBar->StyleName = "AssetEditorToolbar";
-
 	{
 		FToolMenuSection& PlaySection = PlayToolBar->AddSection("Play");
 
@@ -1381,6 +1379,13 @@ void FLevelEditorToolBar::RegisterLevelEditorToolBar( const TSharedRef<FUIComman
 		// Add the shared play-world commands that will be shown on the Kismet toolbar as well
 		FPlayWorldCommands::BuildToolbar(PlaySection, true);
 
+
+	}
+
+	UToolMenu* SettingsToolbar = UToolMenus::Get()->RegisterMenu("LevelEditor.LevelEditorToolBar.SettingsToolbar", NAME_None, EMultiBoxType::SlimHorizontalToolBar);
+	SettingsToolbar->StyleName = "AssetEditorToolbar";
+	{
+		FToolMenuSection& SettingsSection = SettingsToolbar->AddSection("ProjectSettings");
 		FToolMenuEntry SettingsEntry =
 			FToolMenuEntry::InitComboButton(
 				"LevelToolbarQuickSettings",
@@ -1393,10 +1398,8 @@ void FLevelEditorToolBar::RegisterLevelEditorToolBar( const TSharedRef<FUIComman
 				"LevelToolbarQuickSettings");
 		SettingsEntry.StyleNameOverride = "CalloutToolbar";
 
-		PlaySection.AddEntry(SettingsEntry);
-
+		SettingsSection.AddEntry(SettingsEntry);
 	}
-
 #undef LOCTEXT_NAMESPACE
 }
 
@@ -1413,6 +1416,8 @@ TSharedRef< SWidget > FLevelEditorToolBar::MakeLevelEditorToolBar( const TShared
 	LevelEditorMenuContext->LevelEditor = InLevelEditor;
 	MenuContext.AddObject(LevelEditorMenuContext);
 
+	TWeakPtr<SLevelEditor> LevelEditorWeakPtr(InLevelEditor);
+
 	// Create the tool bar!
 	return
 		SNew(SOverlay)
@@ -1440,21 +1445,35 @@ TSharedRef< SWidget > FLevelEditorToolBar::MakeLevelEditorToolBar( const TShared
 			]
 			
 			+ SHorizontalBox::Slot()
+			.AutoWidth()
 			.HAlign(HAlign_Left)
 			[
-				UToolMenus::Get()->GenerateWidget("LevelEditor.LevelEditorToolBar.ModesToolBar", MenuContext)
+				SNew(SBox)
+				.MinDesiredWidth_Lambda(
+					[LevelEditorWeakPtr]()
+					{
+						if (TSharedPtr<SLevelEditor> LevelEditor = LevelEditorWeakPtr.Pin())
+						{
+							return LevelEditor->GetTickSpaceGeometry().GetLocalSize().X / 4.0f;
+						}
+						return 350.f;
+					}
+				)
+				[
+					UToolMenus::Get()->GenerateWidget("LevelEditor.LevelEditorToolBar.ModesToolBar", MenuContext)
+				]
 			]
 			+ SHorizontalBox::Slot()
-			.HAlign(HAlign_Right)
-			.FillWidth(.65)
+			.HAlign(HAlign_Left)
 			[
 				UToolMenus::Get()->GenerateWidget("LevelEditor.LevelEditorToolBar.PlayToolBar", MenuContext)
 			]
-			/*+ SHorizontalBox::Slot()
+			+ SHorizontalBox::Slot()
 			.HAlign(HAlign_Right)
+			.Padding(0.0f, 0.0f, 7.0f, 0.0f)
 			[
-				UToolMenus::Get()->GenerateWidget("LevelEditor.LevelEditorToolBar.MarketplaceToolBar", MenuContext)
-			]*/
+				UToolMenus::Get()->GenerateWidget("LevelEditor.LevelEditorToolBar.SettingsToolBar", MenuContext)
+			]
 		];
 }
 
