Adding command handling for FToolUIActions. There is now an in-viewport property menu that can be enabled by [REDACTED_PERSON] in the console, and then pressing Tab when a supported actor (currently any light actor) is selected.

diff --git a/Engine/Source/Developer/ToolMenus/Private/ToolMenu.cpp b/Engine/Source/Developer/ToolMenus/Private/ToolMenu.cpp
--- a/Engine/Source/Developer/ToolMenus/Private/ToolMenu.cpp
+++ b/Engine/Source/Developer/ToolMenus/Private/ToolMenu.cpp
@@ -33,6 +33,30 @@ void UToolMenu::InitMenu(const FToolMenuOwner InOwner, FName InName, FName InPar
 	MenuType = InType;
 }
 
+FReply UToolMenu::OnKeyDown(const FGeometry& MyGeometry, const FKeyEvent& InKeyEvent)
+{
+	for (int32 i = 0; i < Sections.Num(); ++i)
+	{
+		for (FToolMenuEntry& Entry : Sections[i].Blocks)
+		{
+			if (Entry.Type == EMultiBlockType::ToolBarButton
+				&& Entry.IsCommandKeybindOnly())
+			{
+				if (Entry.CommandAcceptsInput(InKeyEvent))
+				{
+					if (Entry.TryExecuteToolUIAction(Context))
+					{
+						return FReply::Handled();
+					}
+				}
+			}
+		}
+	}
+	return FReply::Unhandled();
+}
+
+
+
 void UToolMenu::InitGeneratedCopy(const UToolMenu* Source, const FName InMenuName, const FToolMenuContext* InContext)
 {
 	// Skip sections
diff --git a/Engine/Source/Developer/ToolMenus/Private/ToolMenuEntry.cpp b/Engine/Source/Developer/ToolMenus/Private/ToolMenuEntry.cpp
--- a/Engine/Source/Developer/ToolMenus/Private/ToolMenuEntry.cpp
+++ b/Engine/Source/Developer/ToolMenus/Private/ToolMenuEntry.cpp
@@ -9,6 +9,7 @@
 #include "Framework/MultiBox/MultiBox.h"
 #include "Framework/MultiBox/MultiBoxBuilder.h"
 #include "Internationalization/Internationalization.h"
+#include "Input/Events.h"
 
 
 FToolMenuEntry::FToolMenuEntry() :
@@ -17,7 +18,8 @@ FToolMenuEntry::FToolMenuEntry() :
 	bShouldCloseWindowAfterMenuSelection(true),
 	ScriptObject(nullptr),
 	StyleNameOverride(NAME_None),
-	bAddedDuringRegister(false)
+	bAddedDuringRegister(false),
+	bCommandIsKeybindOnly(false)
 {
 }
 
@@ -29,7 +31,8 @@ FToolMenuEntry::FToolMenuEntry(const FToolMenuOwner InOwner, const FName InName,
 	bShouldCloseWindowAfterMenuSelection(true),
 	ScriptObject(nullptr),
 	StyleNameOverride(NAME_None),
-	bAddedDuringRegister(false)
+	bAddedDuringRegister(false),
+	bCommandIsKeybindOnly(false)
 {
 }
 
@@ -69,6 +72,19 @@ void FToolMenuEntry::AddOptionsDropdown(FUIAction InAction, const FOnGetContent
 	ToolBarData.OptionsDropdownData->ToolTip = InToolTip;
 }
 
+void FToolMenuEntry::AddKeybindFromCommand(const TSharedPtr< const FUICommandInfo >& InCommand)
+{
+	if (Type == EMultiBlockType::ToolBarButton)
+	{
+		Command = InCommand;
+		bCommandIsKeybindOnly = true;
+	}
+	else
+	{
+		ensureMsgf(false, TEXT("Keybinds from commands can only be associated with toolbar buttons."));
+	}
+}
+
 void FToolMenuEntry::SetCommand(const TSharedPtr<const FUICommandInfo>& InCommand, TOptional<FName> InName, const TAttribute<FText>& InLabel, const TAttribute<FText>& InToolTip, const TAttribute<FSlateIcon>& InIcon)
 {
 	Command = InCommand;
@@ -184,7 +200,6 @@ FToolMenuEntry FToolMenuEntry::InitToolBarButton(const TSharedPtr< const FUIComm
 	Entry.SetCommand(InCommand, InName, InLabel, InToolTip, InIcon);
 	return Entry;
 }
-
 FToolMenuEntry FToolMenuEntry::InitComboButton(const FName InName, const FToolUIActionChoice& InAction, const FNewToolMenuChoice& InMenuContentGenerator, const TAttribute<FText>& InLabel, const TAttribute<FText>& InToolTip, const TAttribute<FSlateIcon>& InIcon, bool bInSimpleComboBox, FName InTutorialHighlightName)
 {
 	FToolMenuEntry Entry(UToolMenus::Get()->CurrentOwner(), InName, EMultiBlockType::ToolBarComboButton);
@@ -229,6 +244,48 @@ bool FToolMenuEntry::IsNonLegacyDynamicConstruct() const
 	return Construct.IsBound() || IsScriptObjectDynamicConstruct();
 }
 
+bool FToolMenuEntry::IsCommandKeybindOnly() const
+{
+	return bCommandIsKeybindOnly;
+}
+
+bool FToolMenuEntry::CommandAcceptsInput(const FKeyEvent& InKeyEvent) const
+{
+	bool bAccepted = false;
+	for (uint32 i = 0; i < static_cast<uint8>(EMultipleKeyBindingIndex::NumChords); ++i)
+	{
+		// check each bound chord
+		EMultipleKeyBindingIndex ChordIndex = static_cast<EMultipleKeyBindingIndex>(i);
+		const FInputChord& Chord = *Command->GetActiveChord(ChordIndex);
+
+		bAccepted |= Chord.IsValidChord()
+			&& (!Chord.NeedsControl() || InKeyEvent.IsControlDown())
+			&& (!Chord.NeedsAlt() || InKeyEvent.IsAltDown())
+			&& (!Chord.NeedsShift() || InKeyEvent.IsShiftDown())
+			&& (!Chord.NeedsCommand() || InKeyEvent.IsCommandDown())
+			&& Chord.Key == InKeyEvent.GetKey();
+	}
+	return bAccepted;
+}
+
+bool FToolMenuEntry::TryExecuteToolUIAction(const FToolMenuContext& InContext)
+{
+	bool bCanExecute = false;
+	if (Action.GetToolUIAction() && Action.GetToolUIAction()->ExecuteAction.IsBound())
+	{
+		bCanExecute = true;
+		if (Action.GetToolUIAction()->CanExecuteAction.IsBound())
+		{
+			bCanExecute = Action.GetToolUIAction()->CanExecuteAction.Execute(InContext);
+		}
+		if (bCanExecute)
+		{
+			Action.GetToolUIAction()->ExecuteAction.Execute(InContext);
+		}
+	}
+	return bCanExecute;
+}
+
 bool FToolMenuEntry::IsScriptObjectDynamicConstruct() const
 {
 	static const FName ConstructMenuEntryName = GET_FUNCTION_NAME_CHECKED(UToolMenuEntryScript, ConstructMenuEntry);
diff --git a/Engine/Source/Developer/ToolMenus/Private/ToolMenus.cpp b/Engine/Source/Developer/ToolMenus/Private/ToolMenus.cpp
--- a/Engine/Source/Developer/ToolMenus/Private/ToolMenus.cpp
+++ b/Engine/Source/Developer/ToolMenus/Private/ToolMenus.cpp
@@ -203,6 +203,8 @@ public:
 
 	void AddStandardEntryToMenuBuilder()
 	{
+		// First, check for a ToolUIAction, otherwise do the rest of this (have CommandList and Command)
+		// Need another variable to store if we are using a keybind from a command
 		if (Block.Command.IsValid())
 		{
 			bool bPopCommandList = false;
@@ -1266,7 +1268,7 @@ void UToolMenus::PopulateToolBarBuilder(FToolBarBuilder& ToolBarBuilder, UToolMe
 
 			if (Block.Type == EMultiBlockType::ToolBarButton)
 			{
-				if (Block.Command.IsValid())
+				if (Block.Command.IsValid() && !Block.IsCommandKeybindOnly())
 				{
 					bool bPopCommandList = false;
 					TSharedPtr<const FUICommandList> CommandListForAction;
diff --git a/Engine/Source/Developer/ToolMenus/Public/ToolMenu.h b/Engine/Source/Developer/ToolMenus/Public/ToolMenu.h
--- a/Engine/Source/Developer/ToolMenus/Public/ToolMenu.h
+++ b/Engine/Source/Developer/ToolMenus/Public/ToolMenu.h
@@ -89,6 +89,8 @@ public:
 	}
 
 
+	FReply OnKeyDown(const FGeometry& MyGeometry, const FKeyEvent& InKeyEvent);
+
 	friend class UToolMenus;
 
 private:
diff --git a/Engine/Source/Developer/ToolMenus/Public/ToolMenuEntry.h b/Engine/Source/Developer/ToolMenus/Public/ToolMenuEntry.h
--- a/Engine/Source/Developer/ToolMenus/Public/ToolMenuEntry.h
+++ b/Engine/Source/Developer/ToolMenus/Public/ToolMenuEntry.h
@@ -12,10 +12,10 @@
 #include "Framework/MultiBox/MultiBoxDefs.h"
 #include "Framework/Commands/UICommandInfo.h"
 #include "UObject/TextProperty.h"
-
 #include "ToolMenuEntry.generated.h"
 
 class UToolMenuEntryScript;
+struct FKeyEvent;
 
 struct FToolMenuEntrySubMenuData
 {
@@ -121,7 +121,11 @@ struct TOOLMENUS_API FToolMenuEntry
 	void SetCommandList(const TSharedPtr<const FUICommandList>& InCommandList);
 
 	void AddOptionsDropdown(FUIAction InAction, const FOnGetContent InMenuContentGenerator, const TAttribute<FText>& InToolTip = TAttribute<FText>());
+	void AddKeybindFromCommand(const TSharedPtr< const FUICommandInfo >& InCommand);
 
+	bool IsCommandKeybindOnly() const;
+	bool CommandAcceptsInput(const FKeyEvent& InKeyEvent) const;
+	bool TryExecuteToolUIAction(const FToolMenuContext& InContext);
 	friend struct FToolMenuSection;
 	friend class UToolMenuEntryScript;
 
@@ -194,5 +198,8 @@ private:
 	FNewToolMenuDelegateLegacy ConstructLegacy;
 
 	bool bAddedDuringRegister;
+
+	UPROPERTY()
+	bool bCommandIsKeybindOnly;
 };
 
diff --git a/Engine/Source/Developer/ToolMenus/ToolMenus.Build.cs b/Engine/Source/Developer/ToolMenus/ToolMenus.Build.cs
--- a/Engine/Source/Developer/ToolMenus/ToolMenus.Build.cs
+++ b/Engine/Source/Developer/ToolMenus/ToolMenus.Build.cs
@@ -13,6 +13,7 @@ namespace UnrealBuildTool.Rules
 					"ApplicationCore",
 					"Slate",
 					"SlateCore",
+					"InputCore"
 				}
 			);
 		}
diff --git a/Engine/Source/Editor/EditorStyle/Private/StarshipStyle.cpp b/Engine/Source/Editor/EditorStyle/Private/StarshipStyle.cpp
--- a/Engine/Source/Editor/EditorStyle/Private/StarshipStyle.cpp
+++ b/Engine/Source/Editor/EditorStyle/Private/StarshipStyle.cpp
@@ -3105,6 +3105,23 @@ void FStarshipEditorStyle::FStyle::SetupPropertyEditorStyles()
 			.SetForegroundColor( DefaultForeground )
 			);
 
+	
+		
+		FLinearColor TransparentHeader = FStyleColors::Recessed.GetSpecifiedColor();
+		TransparentHeader.A = 0.95f;
+		FLinearColor TransparentToolbar = FStyleColors::Header.GetSpecifiedColor();
+		TransparentToolbar.A = 0.95f;
+		FLinearColor TransparentBackground = FStyleColors::Background.GetSpecifiedColor();
+		TransparentBackground.A = 0.75f;
+		Set("PropertyTable.InViewport.Background", new FSlateColorBrush(FSlateColor(TransparentBackground)));
+		Set("PropertyTable.InViewport.Header", new FSlateColorBrush(FSlateColor(TransparentHeader)));
+		// InViewportToolbar
+		{
+			FToolBarStyle InViewportToolbar = FStarshipCoreStyle::GetCoreStyle().GetWidgetStyle<FToolBarStyle>("SlimToolBar");
+			InViewportToolbar.SetBackground(FSlateColorBrush(FSlateColor(TransparentToolbar)));
+			Set("InViewportToolbar", InViewportToolbar);
+		}
+
 		Set( "PropertyTable.Selection.Active",						new IMAGE_BRUSH( "Common/Selection", Icon8x8, SelectionColor ) );
 
 		Set( "PropertyTable.HeaderRow.Column.PathDelimiter",		new IMAGE_BRUSH( "Common/SmallArrowRight", Icon10x10 ) );
diff --git a/Engine/Source/Editor/LevelEditor/Private/LevelEditorSubsystem.cpp b/Engine/Source/Editor/LevelEditor/Private/LevelEditorSubsystem.cpp
--- a/Engine/Source/Editor/LevelEditor/Private/LevelEditorSubsystem.cpp
+++ b/Engine/Source/Editor/LevelEditor/Private/LevelEditorSubsystem.cpp
@@ -12,10 +12,62 @@
 #include "EditorScriptingHelpers.h"
 #include "LevelEditor.h"
 #include "Modules/ModuleManager.h"
+#include "ToolMenus.h"
+#include "LevelEditorMenuContext.h"
+#include "Elements/Framework/TypedElementSelectionSet.h"
+#include "ToolMenuDelegates.h"
 
 DEFINE_LOG_CATEGORY_STATIC(LevelEditorSubsystem, Log, All);
 
-void ULevelEditorSubsystem::PilotLevelActor(AActor* ActorToPilot)
+#define LOCTEXT_NAMESPACE "LevelEditorSubsystem"
+
+void ULevelEditorSubsystem::Initialize(FSubsystemCollectionBase& Collection)
+{
+	UToolMenus::RegisterStartupCallback(FSimpleMulticastDelegate::FDelegate::CreateUObject(this, &ULevelEditorSubsystem::ExtendQuickActionMenu));
+}
+
+void ULevelEditorSubsystem::Deinitialize()
+{
+	UToolMenus::UnRegisterStartupCallback(this);
+	UToolMenus::UnregisterOwner(this);
+}
+
+void ULevelEditorSubsystem::ExtendQuickActionMenu()
+{
+	FToolMenuOwnerScoped MenuOwner(this);
+	UToolMenu* Menu = UToolMenus::Get()->ExtendMenu("LevelEditor.InViewportPanel");
+	{
+		FToolMenuSection& Section = Menu->FindOrAddSection("QuickActions");
+		FToolMenuEntry& Entry = Section.AddDynamicEntry("LevelActors", FNewToolMenuSectionDelegate::CreateLambda([this](FToolMenuSection& InSection)
+			{
+				UQuickActionMenuContext* Context = InSection.FindContext<UQuickActionMenuContext>();
+				if (Context && Context->CurrentSelection && Context->CurrentSelection->GetElementList()->Num() == 1)
+				{
+				
+					FToolUIAction PilotActorAction;
+					PilotActorAction.ExecuteAction = FToolMenuExecuteAction::CreateUObject(this, &ULevelEditorSubsystem::PilotLevelActor);
+					FToolMenuEntry& PilotActorEntry = InSection.AddEntry(FToolMenuEntry::InitToolBarButton(
+						"PilotActor",
+						PilotActorAction,
+						LOCTEXT("PilotSelectedActor", "Pilot Selected Actor"),
+						LOCTEXT("PilotSelectedActorToolTip", "Move the selected actor around using the viewport controls, and bind the viewport to the actor's location and orientation."),
+						FSlateIcon(FAppStyle::Get().GetStyleSetName(), "LevelViewport.PilotSelectedActor")
+					));
+				//	LevelEditorSubsystem.AddKeybindFromCommand(FLightEditingCommands::Get().SwapLightType);
+				}
+			}));
+	}
+
+}
+
+void ULevelEditorSubsystem::PilotLevelActor(const FToolMenuContext& InContext)
+{
+	UQuickActionMenuContext* QuickMenuContext = InContext.FindContext<UQuickActionMenuContext>();
+	AActor* SelectedActor = QuickMenuContext->CurrentSelection->GetTopSelectedObject<AActor>();
+	PilotLevelActor(SelectedActor);
+}
+
+void ULevelEditorSubsystem::PilotLevelActor(AActor* ActorToPilot) 
 {
 	FLevelEditorModule& LevelEditorModule = FModuleManager::GetModuleChecked<FLevelEditorModule>("LevelEditor");
 
@@ -456,4 +508,6 @@ ULevel* ULevelEditorSubsystem::GetCurrentLevel()
 	}
 
 	return World->GetCurrentLevel();
-}
\ No newline at end of file
+}
+
+#undef LOCTEXT_NAMESPACE
diff --git a/Engine/Source/Editor/LevelEditor/Private/LightEditorSubsystem.cpp b/Engine/Source/Editor/LevelEditor/Private/LightEditorSubsystem.cpp
new file mode 100644
--- /dev/null
+++ b/Engine/Source/Editor/LevelEditor/Private/LightEditorSubsystem.cpp
@@ -0,0 +1,88 @@
+// Copyright Epic Games, Inc. All Rights Reserved.
+
+#include "LightEditorSubsystem.h"
+#include "ToolMenus.h"
+#include "ToolMenuEntry.h"
+#include "Selection.h"
+#include "Engine/SpotLight.h"
+#include "Engine/PointLight.h"
+#include "Elements/Interfaces/TypedElementObjectInterface.h"
+#include "Elements/Framework/TypedElementSelectionSet.h"
+#include "LevelEditorMenuContext.h"
+
+#define LOCTEXT_NAMESPACE "EditorLightSubsystem"
+
+void FLightEditingCommands::RegisterCommands()
+{
+	UI_COMMAND(SwapLightType, "Add Key Frame", "Inserts a new key frame at the current time", EUserInterfaceActionType::Button, FInputChord(EKeys::P, EModifierKey::Shift));
+}
+
+
+void ULightEditorSubsystem::Initialize(FSubsystemCollectionBase& Collection)
+{
+	FLightEditingCommands::Register();
+	UToolMenus::RegisterStartupCallback(FSimpleMulticastDelegate::FDelegate::CreateUObject(this, &ULightEditorSubsystem::ExtendQuickActionMenu));
+}
+
+void ULightEditorSubsystem::Deinitialize()
+{
+	UToolMenus::UnRegisterStartupCallback(this);
+	UToolMenus::UnregisterOwner(this);
+}
+
+void ULightEditorSubsystem::ExtendQuickActionMenu()
+{
+	FToolMenuOwnerScoped MenuOwner(this);
+	UToolMenu* Menu = UToolMenus::Get()->ExtendMenu("LevelEditor.InViewportPanel");
+	{
+		FToolMenuSection& Section = Menu->FindOrAddSection("QuickActions");
+		FToolMenuEntry& Entry = Section.AddDynamicEntry("Lights", FNewToolMenuSectionDelegate::CreateLambda([this](FToolMenuSection& InSection)
+			{
+				UQuickActionMenuContext* Context = InSection.FindContext<UQuickActionMenuContext>();
+				if (Context && Context->CurrentSelection && Context->CurrentSelection->GetElementList()->Num() > 0)
+				{
+					// Only SpotLights Selected
+					if (Context->CurrentSelection->CountSelectedObjects<ASpotLight>() == Context->CurrentSelection->GetNumSelectedElements())
+					{
+						FToolMenuEntry& SwapLightEntry = InSection.AddEntry(FToolMenuEntry::InitToolBarButton(
+							"SwapLightTypes",
+							FToolUIAction(FToolMenuExecuteAction::CreateUObject(this, &ULightEditorSubsystem::SwapLightType, APointLight::StaticClass())),
+							LOCTEXT("SwapToPointLights", "Swap to Point Lights"),
+							LOCTEXT("SwapToPointLightsToolTip", "Changes all selected spot lights to point lights."),
+							FSlateIcon(FAppStyle::Get().GetStyleSetName(), "PlacementBrowser.Icons.Lights")
+						));
+						SwapLightEntry.AddKeybindFromCommand(FLightEditingCommands::Get().SwapLightType);
+					}
+					// Only PointLights selected
+					if (Context->CurrentSelection->CountSelectedObjects<APointLight>() == Context->CurrentSelection->GetNumSelectedElements())
+					{
+						FToolMenuEntry& SwapLightEntry = InSection.AddEntry(FToolMenuEntry::InitToolBarButton(
+							"SwapLightTypes",
+							FToolUIAction(FToolMenuExecuteAction::CreateUObject(this, &ULightEditorSubsystem::SwapLightType, ASpotLight::StaticClass())),
+							LOCTEXT("SwapToSpotLights", "Swap to Spot Lights"),
+							LOCTEXT("SwapToSpotLightsToolTip", "Changes all selected point lights to spot lights."),
+							FSlateIcon(FAppStyle::Get().GetStyleSetName(), "PlacementBrowser.Icons.Lights")
+						));
+						SwapLightEntry.AddKeybindFromCommand(FLightEditingCommands::Get().SwapLightType);
+					}
+				}
+			}));
+	}
+
+}
+
+void ULightEditorSubsystem::SwapLightType(const FToolMenuContext& InContext, UClass* InClass) const
+{
+	const FAssetData NoAssetData;
+	UActorFactory* ActorFactory = GEditor->FindActorFactoryForActorClass(InClass);
+	if (ActorFactory)
+	{
+		GEditor->ReplaceSelectedActors(ActorFactory, NoAssetData);
+	}
+}
+
+
+
+
+
+#undef LOCTEXT_NAMESPACE
diff --git a/Engine/Source/Editor/LevelEditor/Private/LightEditorSubsystem.h b/Engine/Source/Editor/LevelEditor/Private/LightEditorSubsystem.h
new file mode 100644
--- /dev/null
+++ b/Engine/Source/Editor/LevelEditor/Private/LightEditorSubsystem.h
@@ -0,0 +1,54 @@
+// Copyright Epic Games, Inc. All Rights Reserved.
+
+#pragma once
+
+#include "CoreMinimal.h"
+#include "Layout/Visibility.h"
+#include "Widgets/DeclarativeSyntaxSupport.h"
+#include "Widgets/SCompoundWidget.h"
+#include "Framework/Text/SlateHyperlinkRun.h"
+#include "EditorUndoClient.h"
+#include "Widgets/Views/SListView.h"
+#include "IDetailPropertyRow.h"
+#include "EditorSubsystem.h"
+#include "UObject/ObjectMacros.h"
+#include "Framework/Commands/Commands.h"
+#include "LightEditorSubsystem.generated.h"
+
+
+/**
+ * Light manipulation actions
+ */
+class FLightEditingCommands : public TCommands<FLightEditingCommands>
+{
+
+public:
+	FLightEditingCommands() : TCommands<FLightEditingCommands>
+		(
+			"LightActor", // Context name for fast lookup
+			NSLOCTEXT("Contexts", "LightActor", "Light Actor"), // Localized context name for displaying
+			NAME_None, FEditorStyle::GetStyleSetName()
+			)
+	{
+	}
+
+	TSharedPtr< FUICommandInfo > SwapLightType;
+
+	virtual void RegisterCommands() override;
+};
+
+UCLASS()
+class ULightEditorSubsystem : public UEditorSubsystem
+{
+	GENERATED_BODY()
+public:
+	ULightEditorSubsystem() {};
+	virtual void Initialize(FSubsystemCollectionBase& Collection) override;
+
+	virtual void Deinitialize() override;
+	void ExtendQuickActionMenu();
+	void SwapLightType(const FToolMenuContext& InContext, UClass* InClass) const;
+	// The list of commands with bound delegates for the level editor.
+	TSharedPtr<FUICommandList> LightEditingCommands;
+};
+
diff --git a/Engine/Source/Editor/LevelEditor/Private/SInViewportDetails.cpp b/Engine/Source/Editor/LevelEditor/Private/SInViewportDetails.cpp
new file mode 100644
--- /dev/null
+++ b/Engine/Source/Editor/LevelEditor/Private/SInViewportDetails.cpp
@@ -0,0 +1,513 @@
+// Copyright Epic Games, Inc. All Rights Reserved.
+
+#include "SInViewportDetails.h"
+#include "Widgets/SBoxPanel.h"
+#include "Components/ActorComponent.h"
+#include "Components/SceneComponent.h"
+#include "GameFramework/Actor.h"
+#include "Engine/Blueprint.h"
+#include "Editor.h"
+#include "HAL/FileManager.h"
+#include "Modules/ModuleManager.h"
+#include "Widgets/Layout/SBorder.h"
+#include "Widgets/Images/SImage.h"
+#include "Widgets/Text/SRichTextBlock.h"
+#include "Widgets/Layout/SBox.h"
+#include "Widgets/Layout/SSplitter.h"
+#include "EditorStyleSet.h"
+#include "Editor/UnrealEdEngine.h"
+#include "Engine/Selection.h"
+#include "UnrealEdGlobals.h"
+#include "LevelEditor.h"
+#include "Kismet2/KismetEditorUtilities.h"
+#include "SSCSEditor.h"
+#include "PropertyEditorModule.h"
+#include "IDetailsView.h"
+#include "LevelEditorGenericDetails.h"
+#include "ScopedTransaction.h"
+#include "SourceCodeNavigation.h"
+#include "Widgets/Docking/SDockTab.h"
+#include "Subsystems/PanelExtensionSubsystem.h"
+#include "DetailsViewObjectFilter.h"
+#include "IDetailRootObjectCustomization.h"
+#include "IPropertyRowGenerator.h"
+#include "IDetailTreeNode.h"
+#include "IDetailPropertyRow.h"
+#include "Widgets/Layout/SBackgroundBlur.h"
+#include "PropertyCustomizationHelpers.h"
+#include "Input/DragAndDrop.h"
+#include "SEditorViewport.h"
+#include "SResetToDefaultPropertyEditor.h"
+#include "ToolMenus.h"
+#include "Editor/EditorEngine.h"
+#include "LevelEditorMenuContext.h"
+
+#define LOCTEXT_NAMESPACE "InViewportDetails"
+void FInViewportUIDragOperation::OnDrop(bool bDropWasHandled, const FPointerEvent& MouseEvent)
+{
+	check(CursorDecoratorWindow.IsValid());
+
+	// Destroy the CursorDecoratorWindow by calling the base class implementation because we are relocating the content into a more permanent home.
+	FDragDropOperation::OnDrop(bDropWasHandled, MouseEvent);
+
+	UIBeingDragged.Reset();
+}
+
+void FInViewportUIDragOperation::OnDragged(const FDragDropEvent& DragDropEvent)
+{
+	// The tab is being dragged. Move the the decorator window to match the cursor position.
+	FVector2D TargetPosition = DragDropEvent.GetScreenSpacePosition() - GetDecoratorOffsetFromCursor();
+	CursorDecoratorWindow->UpdateMorphTargetShape(FSlateRect(TargetPosition.X, TargetPosition.Y, TargetPosition.X + LastContentSize.X, TargetPosition.Y + LastContentSize.Y));
+	CursorDecoratorWindow->MoveWindowTo(TargetPosition);
+}
+
+TSharedRef<FInViewportUIDragOperation> FInViewportUIDragOperation::New(const TSharedRef<class SInViewportDetails>& InUIToBeDragged, const FVector2D InTabGrabOffset, const FVector2D& OwnerAreaSize)
+{
+	const TSharedRef<FInViewportUIDragOperation> Operation = MakeShareable(new FInViewportUIDragOperation(InUIToBeDragged, InTabGrabOffset, OwnerAreaSize));
+	return Operation;
+}
+
+FVector2D FInViewportUIDragOperation::GetTabGrabOffsetFraction() const
+{
+	return TabGrabOffsetFraction;
+}
+
+FInViewportUIDragOperation::~FInViewportUIDragOperation()
+{
+
+}
+
+FInViewportUIDragOperation::FInViewportUIDragOperation(const TSharedRef<class SInViewportDetails>& InUIToBeDragged, const FVector2D InTabGrabOffsetFraction, const FVector2D& OwnerAreaSize)
+	: UIBeingDragged(InUIToBeDragged)
+	, TabGrabOffsetFraction(InTabGrabOffsetFraction)
+	, LastContentSize(OwnerAreaSize)
+{
+	// Create the decorator window that we will use during this drag and drop to make the user feel like
+	// they are actually dragging a piece of UI.
+
+	// Start the window off hidden.
+	const bool bShowImmediately = true;
+	CursorDecoratorWindow = FSlateApplication::Get().AddWindow(SWindow::MakeCursorDecorator(), bShowImmediately);
+	// Usually cursor decorators figure out their size automatically from content, but we will drive it
+	// here because the window will reshape itself to better reflect what will happen when the user drops the Tab.
+	CursorDecoratorWindow->SetSizingRule(ESizingRule::Autosized);
+	CursorDecoratorWindow->SetOpacity(0.45f);
+	CursorDecoratorWindow->SetContent
+	(
+		SNew(SBorder)
+		.BorderImage(FCoreStyle::Get().GetBrush("Docking.Background"))
+		[
+			InUIToBeDragged
+		]
+	);
+
+}
+
+const FVector2D FInViewportUIDragOperation::GetDecoratorOffsetFromCursor()
+{
+	const FVector2D TabDesiredSize = UIBeingDragged->GetDesiredSize();
+	return TabGrabOffsetFraction * TabDesiredSize;
+}
+
+class SInViewportDetailsRow : public STableRow< TSharedPtr<IDetailTreeNode> >
+{
+public:
+
+	SLATE_BEGIN_ARGS(SInViewportDetailsRow)
+	{}
+
+	/** The item content. */
+	SLATE_ARGUMENT(TSharedPtr<IDetailTreeNode>, InNode)
+	SLATE_ARGUMENT(TSharedPtr<SInViewportDetails>, InDetailsView)
+	SLATE_END_ARGS()
+
+
+	/**
+	* Construct the widget
+	*
+	* @param InArgs   A declaration from which to construct the widget
+	*/
+	void Construct(const FArguments& InArgs, const TSharedRef<STableViewBase>& InOwnerTableView)
+	{
+		ParentDetailsView = InArgs._InDetailsView;
+		FDetailColumnSizeData& ColumnSizeData = InArgs._InDetailsView->GetColumnSizeData();
+
+		TSharedPtr<IDetailPropertyRow> DetailPropertyRow = InArgs._InNode->GetRow();
+		FDetailWidgetRow Row;
+		TSharedPtr< SWidget > NameWidget;
+		TSharedPtr< SWidget > ValueWidget;
+		DetailPropertyRow->GetDefaultWidgets(NameWidget, ValueWidget, Row, true);
+		TSharedPtr< SWidget > ResetWidget = SNew(SResetToDefaultPropertyEditor, InArgs._InNode->CreatePropertyHandle());
+ 		TSharedPtr<SWidget> RowWidget = SNullWidget::NullWidget;
+		{
+ 
+			RowWidget = SNew(SSplitter)
+ 				.Style(FEditorStyle::Get(), "DetailsView.Splitter")
+				.PhysicalSplitterHandleSize(1.0f)
+ 				.HitDetectionSplitterHandleSize(5.0f)
+ 				+ SSplitter::Slot()
+				.SizeRule(SSplitter::ESizeRule::FractionOfParent)
+ 				.Value(ColumnSizeData.NameColumnWidth)
+				.OnSlotResized(ColumnSizeData.OnNameColumnResized)
+ 				[
+					SNew(SBox)
+					.HAlign(HAlign_Right)
+ 					.Padding(2.0f)
+ 					[
+						NameWidget.ToSharedRef()
+					]
+ 				]
+ 				+ SSplitter::Slot()
+				.SizeRule(SSplitter::ESizeRule::FractionOfParent)
+ 				.Value(ColumnSizeData.ValueColumnWidth)
+ 				.OnSlotResized(ColumnSizeData.OnValueColumnResized)
+ 				[
+ 					SNew(SBox)
+ 					.Padding(2.0f)
+ 					[
+ 						ValueWidget.ToSharedRef()
+ 					]
+				]
+				+ SSplitter::Slot()
+					.SizeRule(SSplitter::ESizeRule::SizeToContent)
+					[
+						SNew(SHorizontalBox)
+						+ SHorizontalBox::Slot()
+						.VAlign(VAlign_Center)
+						.HAlign(HAlign_Center)
+						.Padding(0.0f)
+						[
+							ResetWidget.ToSharedRef()
+						]
+					];
+ 		}
+
+		ChildSlot
+			[
+				SNew(SBox)
+				.MinDesiredWidth(300.0f)
+				[
+					RowWidget.ToSharedRef()
+				]
+			];
+
+		STableRow< TSharedPtr<IDetailTreeNode> >::ConstructInternal(
+			STableRow::FArguments()
+			.Style(FAppStyle::Get(), "DetailsView.TreeView.TableRow")
+			.ShowSelection(false),
+			InOwnerTableView
+		);
+	}
+
+	TWeakPtr<SInViewportDetails> ParentDetailsView;
+
+};
+
+void SInViewportDetails::Construct(const FArguments& InArgs)
+{
+	OwningViewport = InArgs._InOwningViewport;
+	ParentLevelEditor = InArgs._InOwningLevelEditor;
+	FPropertyEditorModule& PropertyEditorModule = FModuleManager::GetModuleChecked<FPropertyEditorModule>("PropertyEditor");
+	FPropertyRowGeneratorArgs Args;
+	Args.DefaultsOnlyVisibility = EEditDefaultsOnlyNodeVisibility::Hide;
+	Args.NotifyHook = GUnrealEd;
+	ColumnSizeData.ValueColumnWidth = 0.5f;
+	PropertyRowGenerator = PropertyEditorModule.CreatePropertyRowGenerator(Args);
+
+	USelection::SelectionChangedEvent.AddRaw(this, &SInViewportDetails::OnEditorSelectionChanged);
+
+	GEditor->RegisterForUndo(this);
+
+	GenerateWidget();
+
+}
+
+void SInViewportDetails::GenerateWidget()
+{
+	if (AActor* SelectedActor = GetSelectedActorInEditor())
+	{
+		FString NameString;
+		if (GEditor->GetSelectedActors()->Num() > 1)
+		{
+			NameString = LOCTEXT("SelectedObjects", "Selected Objects").ToString();
+		}
+		else
+		{
+			NameString = SelectedActor->GetHumanReadableName();
+		}
+		
+		ChildSlot
+			[
+				SNew(SBackgroundBlur)
+				.Visibility(this, &SInViewportDetails::GetHeaderVisibility)
+				.BlurStrength(1)
+				.BlurRadius(10)
+				[
+					SNew(SVerticalBox)
+					+ SVerticalBox::Slot()
+					.AutoHeight()
+					[
+						SNew(SInViewportDetailsHeader)
+						.Parent(SharedThis(this))
+						.Content()
+						[					
+							SNew(SBorder)
+							.BorderImage(FAppStyle::Get().GetBrush("PropertyTable.InViewport.Header"))
+							.Padding(5.0f)
+							[
+								SNew(STextBlock)
+								.Text(FText::FromString(NameString))
+								.TextStyle(FAppStyle::Get(), "DetailsView.CategoryTextStyle")
+							]
+						]
+					]
+					+ SVerticalBox::Slot()
+					.AutoHeight()
+					[
+						SNew(SInViewportDetailsToolbar)
+						.Parent(SharedThis(this))
+					]
+					+ SVerticalBox::Slot()
+					.AutoHeight()
+					[
+						MakeDetailsWidget()
+					]
+				]
+			];
+	}
+}
+
+EVisibility SInViewportDetails::GetHeaderVisibility() const
+{
+	return Nodes.Num() ? EVisibility::Visible : EVisibility::Collapsed;
+}
+
+TSharedRef<SWidget> SInViewportDetails::MakeDetailsWidget()
+{
+	TSharedRef<SWidget> DetailWidget = SNullWidget::NullWidget;
+	static const FName Name_ShouldShowInViewport("ShouldShowInViewport");
+	if (PropertyRowGenerator.IsValid())
+	{
+		Nodes.Empty();
+		
+		for (TSharedRef<IDetailTreeNode> RootNode : PropertyRowGenerator->GetRootTreeNodes())
+		{
+			TArray<TSharedRef<IDetailTreeNode>> Children;
+			RootNode->GetChildren(Children);
+			
+			for (TSharedRef<IDetailTreeNode> Child : Children)
+			{
+				bool bShowChild = false;
+				{
+					TSharedPtr<IPropertyHandle> NodePropertyHandle = Child->CreatePropertyHandle();
+					if (NodePropertyHandle.IsValid() && NodePropertyHandle->GetProperty()->HasAllPropertyFlags(CPF_DisableEditOnInstance))
+					{
+						continue;
+					}
+					if (NodePropertyHandle.IsValid() && NodePropertyHandle->GetProperty()->GetBoolMetaData(Name_ShouldShowInViewport))
+					{
+						bShowChild = true;
+					}
+					if (bShowChild)
+					{
+						Nodes.Add(Child);
+					}	
+				}
+
+			}
+			
+		}
+	}
+	if (Nodes.Num())
+	{
+		NodeList = SNew(SListView< TSharedPtr<IDetailTreeNode> >)
+			.ItemHeight(24)
+			.ListItemsSource(&Nodes)
+			.OnGenerateRow(this, &SInViewportDetails::GenerateListRow);
+
+		
+		DetailWidget = SNew(SBorder)
+			.BorderImage(FAppStyle::Get().GetBrush("PropertyTable.InViewport.Background"))
+			[
+				NodeList.ToSharedRef()
+			];
+	}
+	return DetailWidget;
+}
+
+SInViewportDetails::~SInViewportDetails()
+{
+	if (GEditor)
+	{
+		GEditor->UnregisterForUndo(this);
+	}
+	USelection::SelectionChangedEvent.RemoveAll(this);
+}
+
+
+void SInViewportDetails::SetObjects(const TArray<UObject*>& InObjects, bool bForceRefresh)
+{
+	if (PropertyRowGenerator)
+	{
+		PropertyRowGenerator->SetObjects(InObjects);
+		GenerateWidget();
+		if (!Nodes.Num())
+		{
+			// Do this on a delay so that we are are not caught in a loop of creating and hiding the menu
+			GEditor->GetTimerManager()->SetTimerForNextTick([this]()
+				{
+					OwningViewport.Pin()->HideInViewportContextMenu();
+				});
+		}
+	}
+}
+
+void SInViewportDetails::PostUndo(bool bSuccess)
+{
+
+}
+
+void SInViewportDetails::PostRedo(bool bSuccess)
+{
+	PostUndo(bSuccess);
+}
+
+
+
+void SInViewportDetails::OnEditorSelectionChanged(UObject* Object)
+{
+	TArray<UObject*> SelectedActors;
+	for (FSelectionIterator It(GEditor->GetSelectedActorIterator()); It; ++It)
+	{
+		AActor* Actor = static_cast<AActor*>(*It);
+		checkSlow(Actor->IsA(AActor::StaticClass()));
+
+		if (!Actor->IsPendingKill())
+		{
+			SelectedActors.Add(Actor);
+		}
+	}
+	SetObjects(SelectedActors);
+}
+
+AActor* SInViewportDetails::GetSelectedActorInEditor() const
+{
+	//@todo this doesn't work w/ multi-select
+	return GEditor->GetSelectedActors()->GetTop<AActor>();
+}
+
+UToolMenu* SInViewportDetails::GetGeneratedToolbarMenu() const
+{
+	return GeneratedToolbarMenu.IsValid() ? GeneratedToolbarMenu.Get() : nullptr;
+}
+
+AActor* SInViewportDetails::GetActorContext() const
+{
+	AActor* SelectedActorInEditor = GetSelectedActorInEditor();
+	
+	return SelectedActorInEditor;
+}
+
+TSharedRef<ITableRow> SInViewportDetails::GenerateListRow(TSharedPtr<IDetailTreeNode> InItem, const TSharedRef<STableViewBase>& InOwningTable)
+{
+	return SNew(SInViewportDetailsRow, InOwningTable)
+		.InNode(InItem)
+		.InDetailsView(SharedThis(this));
+
+}
+
+FReply SInViewportDetails::StartDraggingDetails(FVector2D InTabGrabOffsetFraction, const FPointerEvent& MouseEvent)
+{
+	// Start dragging.
+	TSharedRef<FInViewportUIDragOperation> DragDropOperation =
+		FInViewportUIDragOperation::New(
+			SharedThis(this),
+			InTabGrabOffsetFraction,
+			GetDesiredSize()
+		);
+	if (OwningViewport.IsValid())
+	{
+		OwningViewport.Pin()->ToggleInViewportContextMenu();
+	}
+	return FReply::Handled().BeginDragDrop(DragDropOperation);
+}
+
+void SInViewportDetailsHeader::Construct(const FArguments& InArgs)
+{
+	ParentPtr = InArgs._Parent;
+	ChildSlot
+		[
+			InArgs._Content.Widget
+		];
+}
+
+FReply SInViewportDetailsHeader::OnDragDetected(const FGeometry& MyGeometry, const FPointerEvent& MouseEvent)
+{
+	// Need to remember where within a tab we grabbed
+	const FVector2D TabGrabOffset = MyGeometry.AbsoluteToLocal(MouseEvent.GetScreenSpacePosition());
+	const FVector2D TabSize = MyGeometry.GetLocalSize();
+	const FVector2D TabGrabOffsetFraction = FVector2D(
+		FMath::Clamp(TabGrabOffset.X / TabSize.X, 0.0f, 1.0f),
+		FMath::Clamp(TabGrabOffset.Y / TabSize.Y, 0.0f, 1.0f));
+
+	TSharedPtr<SInViewportDetails> PinnedParent = ParentPtr.Pin();
+	if (PinnedParent.IsValid())
+	{
+		return PinnedParent->StartDraggingDetails(TabGrabOffsetFraction, MouseEvent);
+	}
+
+	return FReply::Unhandled();
+}
+
+TSharedPtr<class FDragDropOperation> SInViewportDetailsHeader::CreateDragDropOperation()
+{
+	TSharedPtr<FDragDropOperation> Operation = MakeShareable(new FDragDropOperation());
+
+	return Operation;
+}
+
+void SInViewportDetailsToolbar::Construct(const FArguments& InArgs)
+{
+	TSharedPtr<SInViewportDetails> Parent = InArgs._Parent;
+	if (!Parent.IsValid())
+	{
+		return;
+	}
+	const FName ToolBarName = GetQuickActionMenuName(Parent->GetSelectedActorInEditor()->GetClass());
+	UToolMenus* ToolMenus = UToolMenus::Get();
+	UToolMenu* FoundMenu = ToolMenus->FindMenu(ToolBarName);
+
+	if (!FoundMenu || !FoundMenu->IsRegistered())
+	{
+		FoundMenu = ToolMenus->RegisterMenu(ToolBarName, NAME_None, EMultiBoxType::SlimHorizontalToolBar);
+	}
+
+	FToolMenuContext MenuContext;
+
+	UQuickActionMenuContext* ToolbarMenuContext = NewObject<UQuickActionMenuContext>(FoundMenu);
+	TWeakPtr<ILevelEditor> ParentLevelEditor = Parent->ParentLevelEditor;
+	if (ParentLevelEditor.IsValid())
+	{
+		ToolbarMenuContext->CurrentSelection = ParentLevelEditor.Pin()->GetElementSelectionSet();
+	}
+	MenuContext.AddObject(ToolbarMenuContext);
+	Parent->GeneratedToolbarMenu = ToolMenus->GenerateMenu(ToolBarName, MenuContext);
+
+	// Move this to the GenerateMenu API
+	Parent->GeneratedToolbarMenu->StyleName = "InViewportToolbar";
+	Parent->GeneratedToolbarMenu->bToolBarIsFocusable = false;
+	Parent->GeneratedToolbarMenu->bToolBarForceSmallIcons = true;
+	TSharedRef< class SWidget > ToolBarWidget = ToolMenus->GenerateWidget(Parent->GeneratedToolbarMenu.Get());
+
+	ChildSlot
+	[
+		ToolBarWidget
+	];
+}
+
+FName SInViewportDetailsToolbar::GetQuickActionMenuName(UClass* InClass)
+{
+	return FName("LevelEditor.InViewportPanel");
+}
+
+#undef LOCTEXT_NAMESPACE
\ No newline at end of file
diff --git a/Engine/Source/Editor/LevelEditor/Private/SInViewportDetails.h b/Engine/Source/Editor/LevelEditor/Private/SInViewportDetails.h
new file mode 100644
--- /dev/null
+++ b/Engine/Source/Editor/LevelEditor/Private/SInViewportDetails.h
@@ -0,0 +1,175 @@
+// Copyright Epic Games, Inc. All Rights Reserved.
+
+#pragma once
+
+#include "CoreMinimal.h"
+#include "Layout/Visibility.h"
+#include "Widgets/DeclarativeSyntaxSupport.h"
+#include "Widgets/SCompoundWidget.h"
+#include "Framework/Text/SlateHyperlinkRun.h"
+#include "EditorUndoClient.h"
+#include "Widgets/Views/SListView.h"
+#include "IDetailPropertyRow.h"
+#include "ToolMenuContext.h"
+
+
+class AActor;
+class FSCSEditorTreeNode;
+class FTabManager;
+class FUICommandList;
+class IDetailsView;
+class SBox;
+class SSCSEditor;
+class SSplitter;
+class UBlueprint;
+class FDetailsViewObjectFilter;
+class IDetailTreeNode;
+class IPropertyRowGenerator;
+class ILevelEditor;
+class UToolMenu;
+
+/** A Sample implementation of IDragDropOperation */
+class FInViewportUIDragOperation : public FDragDropOperation
+{
+public:
+
+	DRAG_DROP_OPERATOR_TYPE(FInViewportUIDragOperation, FDragDropOperation)
+
+	/**
+	* Invoked when the drag and drop operation has ended.
+	*
+	* @param bDropWasHandled   true when the drop was handled by some widget; false otherwise
+	*/
+	virtual void OnDrop(bool bDropWasHandled, const FPointerEvent& MouseEvent) override;
+
+	/**
+	 * Called when the mouse was moved during a drag and drop operation
+	 *
+	 * @param DragDropEvent    The event that describes this drag drop operation.
+	 */
+	virtual void OnDragged(const FDragDropEvent& DragDropEvent) override;
+
+	/**
+	 * Create this Drag and Drop Content
+	 *
+	 * @param InUIToBeDragged	  The UI being dragged 
+	 * @param InTabGrabOffset     Where within the tab we grabbed, so we're not dragging by the upper left of the tab.
+	 * @param OwnerAreaSize       Size of the DockArea at the time when we start dragging.
+	 *
+	 * @return a new FDockingDragOperation
+	 */
+	static TSharedRef<FInViewportUIDragOperation> New(const TSharedRef<class SInViewportDetails>& InUIToBeDragged, const FVector2D InTabGrabOffset, const FVector2D& OwnerAreaSize);
+
+	/** @return location where the user grabbed within the tab as a fraction of the tab's size */
+	FVector2D GetTabGrabOffsetFraction() const;
+
+	virtual ~FInViewportUIDragOperation();
+
+	/** @return The offset into the tab where the user grabbed in Slate Units. */
+	const FVector2D GetDecoratorOffsetFromCursor();
+
+protected:
+	/** The constructor is protected, so that this class can only be instantiated as a shared pointer. */
+	FInViewportUIDragOperation(const TSharedRef<class SInViewportDetails>& InUIToBeDragged, const FVector2D InTabGrabOffsetFraction, const FVector2D& OwnerAreaSize);
+
+	/** What is actually being dragged in this operation */
+	TSharedPtr<class SInViewportDetails> UIBeingDragged;
+
+	/** Where the user grabbed the tab as a fraction of the tab's size */
+	FVector2D TabGrabOffsetFraction;
+
+	/** Decorator widget where we add temp doc tabs to */
+	TSharedPtr<SDockingTabStack> CursorDecoratorStackNode;
+
+	/** What the size of the content was when it was last shown. The user drags splitters to set this size; it is legitimate state. */
+	FVector2D LastContentSize;
+};
+
+/**
+ * Wraps a details panel customized for viewing actors
+ */
+class SInViewportDetails : public SCompoundWidget, public FEditorUndoClient
+{
+public:
+	SLATE_BEGIN_ARGS(SInViewportDetails) {}
+	SLATE_ARGUMENT(TSharedPtr<class SEditorViewport>, InOwningViewport)
+	SLATE_ARGUMENT(TSharedPtr<ILevelEditor>, InOwningLevelEditor)
+	SLATE_END_ARGS()
+
+	void Construct(const FArguments& InArgs);
+
+	void GenerateWidget();
+
+	EVisibility GetHeaderVisibility() const;
+	TSharedRef<SWidget> MakeDetailsWidget();
+	~SInViewportDetails();
+
+	/**
+	 * Sets the objects to be viewed by the details panel
+	 *
+	 * @param InObjects	The objects to set
+	 */
+	void SetObjects(const TArray<UObject*>& InObjects, bool bForceRefresh = false);
+
+	/** FEditorUndoClient Interface */
+	virtual void PostUndo(bool bSuccess) override;
+	virtual void PostRedo(bool bSuccess) override;
+
+	FReply StartDraggingDetails(FVector2D InTabGrabOffsetFraction, const FPointerEvent& MouseEvent);
+	FDetailColumnSizeData& GetColumnSizeData() { return ColumnSizeData; }
+	AActor* GetSelectedActorInEditor() const;
+	UToolMenu* GetGeneratedToolbarMenu() const;
+
+	friend class SInViewportDetailsToolbar;
+
+private:
+	AActor* GetActorContext() const;
+	TSharedRef<ITableRow> GenerateListRow(TSharedPtr<IDetailTreeNode> InItem, const TSharedRef<STableViewBase>& InOwningTable);
+	void OnEditorSelectionChanged(UObject* Object);
+
+private:
+	TSharedPtr<SSplitter> DetailsSplitter;
+	TSharedPtr<SListView<TSharedPtr<class IDetailTreeNode>>> NodeList;
+	TArray<TSharedPtr<class IDetailTreeNode>> Nodes;
+	TSharedPtr<IPropertyRowGenerator> PropertyRowGenerator;
+	FDetailColumnSizeData ColumnSizeData;
+	TWeakPtr<class SEditorViewport> OwningViewport;
+	TWeakPtr<ILevelEditor> ParentLevelEditor;
+	TWeakObjectPtr<UToolMenu> GeneratedToolbarMenu;
+};
+
+class SInViewportDetailsHeader : public SCompoundWidget
+{
+
+public:
+	SLATE_BEGIN_ARGS(SInViewportDetailsHeader) {}
+	SLATE_DEFAULT_SLOT(FArguments, Content)
+	SLATE_ARGUMENT(TSharedPtr<SInViewportDetails>, Parent)
+	SLATE_END_ARGS()
+
+	void Construct(const FArguments& InArgs);
+
+	FReply OnMouseButtonDown(const FGeometry& MyGeometry, const FPointerEvent& MouseEvent) override
+	{
+		return FReply::Handled().DetectDrag(SharedThis(this), EKeys::LeftMouseButton);
+	};
+
+
+	FReply OnDragDetected(const FGeometry& MyGeometry, const FPointerEvent& MouseEvent) override;
+	TSharedPtr<class FDragDropOperation> CreateDragDropOperation();
+
+	/** The parent in-viewport details */
+	TWeakPtr<SInViewportDetails> ParentPtr;
+};
+
+class SInViewportDetailsToolbar : public SCompoundWidget
+{
+public:
+	SLATE_BEGIN_ARGS(SInViewportDetailsToolbar) {}
+	SLATE_ARGUMENT(TSharedPtr<SInViewportDetails>, Parent)
+	SLATE_END_ARGS()
+
+	void Construct(const FArguments& InArgs);
+	FName GetQuickActionMenuName(UClass* InClass);
+
+};
\ No newline at end of file
diff --git a/Engine/Source/Editor/LevelEditor/Private/SLevelViewport.cpp b/Engine/Source/Editor/LevelEditor/Private/SLevelViewport.cpp
--- a/Engine/Source/Editor/LevelEditor/Private/SLevelViewport.cpp
+++ b/Engine/Source/Editor/LevelEditor/Private/SLevelViewport.cpp
@@ -82,9 +82,12 @@
 #include "WorldPartition/DataLayer/WorldDataLayers.h"
 #include "WorldPartition/DataLayer/DataLayer.h"
 #include "DataLayer/DataLayerEditorSubsystem.h"
+#include "SInViewportDetails.h"
 
-static const FName LevelEditorName("LevelEditor");
 
+static const FName LevelEditorName("LevelEditor");
+static FAutoConsoleCommand EnableInViewportMenu(TEXT("Editor.EnableInViewportMenu"), TEXT("Enables the new in-viewport property menu"), FConsoleCommandDelegate::CreateStatic(&SLevelViewport::EnableInViewportMenu));
+bool SLevelViewport::bInViewportMenuEnabled = false;
 #define LOCTEXT_NAMESPACE "LevelViewport"
 
 // @todo Slate Hack: Disallow game UI to be used in play in viewport until GWorld problem is fixed
@@ -239,6 +242,9 @@ void SLevelViewport::Construct(const FArguments& InArgs, const FAssetEditorViewp
 
 	FEditorDelegates::PostPIEStarted.AddSP(this, &SLevelViewport::TransitionToPIE);
 	FEditorDelegates::PrePIEEnded.AddSP(this, &SLevelViewport::TransitionFromPIE);
+
+	bIsInViewportMenuShowing = false;
+	bIsInViewportMenuInitialized = false;
 }
 
 void SLevelViewport::ConstructViewportOverlayContent()
@@ -618,6 +624,12 @@ FReply SLevelViewport::OnKeyDown( const FGeometry& MyGeometry, const FKeyEvent&
 	{
 		Reply = SEditorViewport::OnKeyDown(MyGeometry,InKeyEvent);
 
+		// Otherwise, give the in-viewport context menu a chance to handle the keypress.
+		if (!Reply.IsEventHandled() && InViewportMenu.IsValid())
+		{
+			Reply = InViewportMenu->GetGeneratedToolbarMenu()->OnKeyDown(MyGeometry, InKeyEvent);
+		}
+
 
 		// If we are in immersive mode and the event was not handled, we will check to see if the the 
 		//  optional parent level editor is set.  If it is, we give it a chance to handle the key event.
@@ -718,6 +730,10 @@ bool SLevelViewport::HandleDragObjects(const FGeometry& MyGeometry, const FDragD
 
 		new(SelectedAssetDatas) FAssetData(BrushOperation->GetBrushBuilder().Get());
 	}
+	else if (Operation->IsOfType<FInViewportUIDragOperation>())
+	{
+		bValidDrag = true;
+	}
 	else
 	{
 		SelectedAssetDatas = AssetUtil::ExtractAssetDataFromDrag( DragDropEvent );
@@ -974,19 +990,29 @@ bool SLevelViewport::HandlePlaceDraggedObjects(const FGeometry& MyGeometry, cons
 
 FReply SLevelViewport::OnDrop( const FGeometry& MyGeometry, const FDragDropEvent& DragDropEvent )
 {
-	ULevel* CurrentLevel = (GetWorld()) ? GetWorld()->GetCurrentLevel() : nullptr;
-
-	if (CurrentLevel && !FLevelUtils::IsLevelLocked(CurrentLevel))
+	if (DragDropEvent.GetOperation()->IsOfType<FInViewportUIDragOperation>())
 	{
-		return HandlePlaceDraggedObjects(MyGeometry, DragDropEvent, /*bCreateDropPreview=*/false) ? FReply::Handled() : FReply::Unhandled();
+		FVector2D ScreenSpaceDropLocation = DragDropEvent.GetScreenSpacePosition() - DragDropEvent.GetOperationAs<FInViewportUIDragOperation>()->GetDecoratorOffsetFromCursor();
+		UpdateInViewportMenuLocation(MyGeometry.AbsoluteToLocal(ScreenSpaceDropLocation));
+		ToggleInViewportContextMenu();
 	}
 	else
 	{
-		FNotificationInfo Info(LOCTEXT("Error_OperationDisallowedOnLockedLevel", "The requested operation could not be completed because the level is locked."));
-		Info.ExpireDuration = 3.0f;
-		FSlateNotificationManager::Get().AddNotification(Info);
-		return FReply::Handled();
+		ULevel* CurrentLevel = (GetWorld()) ? GetWorld()->GetCurrentLevel() : nullptr;
+
+		if (CurrentLevel && !FLevelUtils::IsLevelLocked(CurrentLevel))
+		{
+			return HandlePlaceDraggedObjects(MyGeometry, DragDropEvent, /*bCreateDropPreview=*/false) ? FReply::Handled() : FReply::Unhandled();
+		}
+		else
+		{
+			FNotificationInfo Info(LOCTEXT("Error_OperationDisallowedOnLockedLevel", "The requested operation could not be completed because the level is locked."));
+			Info.ExpireDuration = 3.0f;
+			FSlateNotificationManager::Get().AddNotification(Info);
+			return FReply::Handled();
+		}
 	}
+	return FReply::Unhandled();
 }
 
 
@@ -4502,6 +4528,104 @@ UWorld* SLevelViewport::GetWorld() const
 	return ParentLevelEditor.IsValid() ? ParentLevelEditor.Pin()->GetWorld() : nullptr;
 }
 
+void SLevelViewport::ToggleInViewportContextMenu()
+{
+	USelection* ActorSelection = GEditor->GetSelectedActors();
+	if (ActorSelection->Num())
+	{
+		if (!bIsInViewportMenuShowing)
+		{
+			// Set up the correct menu location first
+			if (!bIsInViewportMenuInitialized)
+			{
+				FVector2D NewViewportContextMenuLocation = GetDefault<ULevelEditorViewportSettings>()->LastInViewportMenuLocation;
+				if (!NewViewportContextMenuLocation.IsZero())
+				{
+					UpdateInViewportMenuLocation(NewViewportContextMenuLocation);
+				}
+				else
+				{
+					AActor* SelectedActor = ActorSelection->GetTop<AActor>();
+					FSceneViewFamilyContext ViewFamily(FSceneViewFamily::ConstructionValues(
+						GetActiveViewport(),
+						LevelViewportClient->GetScene(),
+						LevelViewportClient->EngineShowFlags)
+						.SetRealtimeUpdate(IsRealtime()));
+					// SceneView is deleted with the ViewFamily
+					FSceneView* SceneView = LevelViewportClient->CalcSceneView(&ViewFamily);
+					FVector2D ScreenPos;
+					SceneView->WorldToPixel(SelectedActor->GetTransform().GetLocation(), ScreenPos);
+					UpdateInViewportMenuLocation(ScreenPos);
+				}
+				bIsInViewportMenuInitialized = true;
+			}
+
+			bIsInViewportMenuShowing = true;
+			InViewportMenu = SNew(SInViewportDetails)
+				.InOwningViewport(SharedThis(this))
+				.InOwningLevelEditor(ParentLevelEditor.Pin());
+			InViewportMenuWrapper = SNew(SHorizontalBox)
+				+ SHorizontalBox::Slot()
+				.FillWidth(1.0f)
+				.VAlign(VAlign_Top)
+				.HAlign(HAlign_Left)
+				.Padding(TAttribute<FMargin>(this, &SLevelViewport::GetContextMenuPadding))
+				[
+					InViewportMenu.ToSharedRef()
+				];
+
+			// Immediately update it (otherwise it will appear empty)
+			{
+				TArray<UObject*> SelectedActors;
+				for (FSelectionIterator It(GEditor->GetSelectedActorIterator()); It; ++It)
+				{
+					AActor* Actor = static_cast<AActor*>(*It);
+					checkSlow(Actor->IsA(AActor::StaticClass()));
+
+					if (!Actor->IsPendingKill())
+					{
+						SelectedActors.Add(Actor);
+					}
+				}
+
+				const bool bForceRefresh = true;
+				InViewportMenu->SetObjects(SelectedActors, bForceRefresh);
+			}
+			AddOverlayWidget(InViewportMenuWrapper.ToSharedRef());
+		}
+		else
+		{
+			HideInViewportContextMenu();
+		}
+	
+	}
+}
+
+void SLevelViewport::HideInViewportContextMenu()
+{
+	if (InViewportMenuWrapper.IsValid())
+	{
+		RemoveOverlayWidget(InViewportMenuWrapper.ToSharedRef());
+	}
+	bIsInViewportMenuShowing = false;
+	InViewportMenu.Reset();
+}
+
+bool SLevelViewport::CanToggleInViewportContextMenu()
+{
+	return SLevelViewport::bInViewportMenuEnabled;
+}
+
+void SLevelViewport::EnableInViewportMenu()
+{
+	SLevelViewport::bInViewportMenuEnabled = !SLevelViewport::bInViewportMenuEnabled;
+}
+
+FMargin SLevelViewport::GetContextMenuPadding() const
+{
+	return FMargin(InViewportContextMenuLocation.X, InViewportContextMenuLocation.Y, 0, 0);
+}
+
 void SLevelViewport::RemoveActorPreview( int32 PreviewIndex, AActor* Actor, const bool bRemoveFromDesktopViewport /*=true */ )
 {
 	IVREditorModule& VREditorModule = IVREditorModule::Get();
diff --git a/Engine/Source/Editor/LevelEditor/Public/LevelEditorMenuContext.h b/Engine/Source/Editor/LevelEditor/Public/LevelEditorMenuContext.h
--- a/Engine/Source/Editor/LevelEditor/Public/LevelEditorMenuContext.h
+++ b/Engine/Source/Editor/LevelEditor/Public/LevelEditorMenuContext.h
@@ -13,6 +13,7 @@ class UActorComponent;
 class SLevelViewport;
 class SLevelViewportToolBar;
 class FLevelEditorViewportClient;
+class UTypedElementSelectionSet;
 
 UCLASS()
 class LEVELEDITOR_API ULevelEditorMenuContext : public UObject
@@ -59,3 +60,13 @@ public:
 
 	FLevelEditorViewportClient* GetLevelViewportClient();
 };
+
+
+UCLASS()
+class LEVELEDITOR_API UQuickActionMenuContext : public UObject
+{
+	GENERATED_BODY()
+public:
+	UPROPERTY(Transient)
+	const UTypedElementSelectionSet* CurrentSelection;
+};
diff --git a/Engine/Source/Editor/LevelEditor/Public/LevelEditorSubsystem.h b/Engine/Source/Editor/LevelEditor/Public/LevelEditorSubsystem.h
--- a/Engine/Source/Editor/LevelEditor/Public/LevelEditorSubsystem.h
+++ b/Engine/Source/Editor/LevelEditor/Public/LevelEditorSubsystem.h
@@ -17,11 +17,15 @@ class LEVELEDITOR_API ULevelEditorSubsystem : public UEditorSubsystem
 	GENERATED_BODY()
 
 public:
+	virtual void Initialize(FSubsystemCollectionBase& Collection) override;
 
-	UFUNCTION(BlueprintCallable, Category = "Editor Scripting | Level Utility", meta=(DevelopmentOnly))
-	void PilotLevelActor(AActor* ActorToPilot);
+	virtual void Deinitialize() override;
+	void ExtendQuickActionMenu();
 
 	UFUNCTION(BlueprintCallable, Category = "Editor Scripting | Level Utility", meta=(DevelopmentOnly))
+	void PilotLevelActor(AActor* ActorToPilot);
+	void PilotLevelActor(const FToolMenuContext& InContext);
+	UFUNCTION(BlueprintCallable, Category = "Editor Scripting | Level Utility", meta = (DevelopmentOnly))
 	void EjectPilotLevelActor();
 
 	UFUNCTION(BlueprintCallable, Category = "Editor Scripting | Level Utility", meta = (DevelopmentOnly))
diff --git a/Engine/Source/Editor/LevelEditor/Public/SLevelViewport.h b/Engine/Source/Editor/LevelEditor/Public/SLevelViewport.h
--- a/Engine/Source/Editor/LevelEditor/Public/SLevelViewport.h
+++ b/Engine/Source/Editor/LevelEditor/Public/SLevelViewport.h
@@ -127,7 +127,13 @@ public:
 	virtual void OnFocusViewportToSelection() override;
 	virtual EVisibility GetTransformToolbarVisibility() const override;
 	virtual UWorld* GetWorld() const override;
+	virtual void ToggleInViewportContextMenu() override;
 
+	virtual void HideInViewportContextMenu() override;
+	virtual bool CanToggleInViewportContextMenu() override;
+
+	static void EnableInViewportMenu();
+	FMargin GetContextMenuPadding() const;
 	/**
 	 * Called when the maximize command is executed                   
 	 */
@@ -909,6 +915,11 @@ private:
 
 	/** Whether to show a full toolbar, or a compact one */
 	bool bShowFullToolbar;
+	TSharedPtr<class SWidget> InViewportMenuWrapper;
+	bool bIsInViewportMenuShowing;
+	bool bIsInViewportMenuInitialized;
+	TSharedPtr<class SInViewportDetails> InViewportMenu;
+	static bool bInViewportMenuEnabled;
 protected:
 	void LockActorInternal(AActor* NewActorToLock);
 
diff --git a/Engine/Source/Editor/PropertyEditor/Private/PropertyNode.cpp b/Engine/Source/Editor/PropertyEditor/Private/PropertyNode.cpp
--- a/Engine/Source/Editor/PropertyEditor/Private/PropertyNode.cpp
+++ b/Engine/Source/Editor/PropertyEditor/Private/PropertyNode.cpp
@@ -122,6 +122,10 @@ void FPropertyNode::InitNode(const FPropertyNodeInitParams& InitParams)
 	PropertyNodeFlags = EPropertyNodeFlags::NoFlags;
 	SetNodeFlags(EPropertyNodeFlags::IsSparseClassData, bIsSparse);
 
+	static const FName Name_ShouldShowInViewport("ShouldShowInViewport");
+	bool bShouldShowInViewport = Property.IsValid() ? Property->GetBoolMetaData(Name_ShouldShowInViewport) : false;
+	SetNodeFlags(EPropertyNodeFlags::ShouldShowInViewport, bShouldShowInViewport);
+
 	//default to copying from the parent
 	if (ParentNode)
 	{
diff --git a/Engine/Source/Editor/PropertyEditor/Private/PropertyNode.h b/Engine/Source/Editor/PropertyEditor/Private/PropertyNode.h
--- a/Engine/Source/Editor/PropertyEditor/Private/PropertyNode.h
+++ b/Engine/Source/Editor/PropertyEditor/Private/PropertyNode.h
@@ -65,6 +65,8 @@ namespace EPropertyNodeFlags
 
 	const Type	IsSparseClassData				= 1 << 24;	/** true if the property on this node is part of a sparse class data structure */
 
+	const Type	ShouldShowInViewport			= 1 << 25; /** true if the property should be shown in the viewport context menu */
+
 	const Type 	NoFlags							= 0;
 
 };
diff --git a/Engine/Source/Editor/UnrealEd/Classes/Settings/LevelEditorViewportSettings.h b/Engine/Source/Editor/UnrealEd/Classes/Settings/LevelEditorViewportSettings.h
--- a/Engine/Source/Editor/UnrealEd/Classes/Settings/LevelEditorViewportSettings.h
+++ b/Engine/Source/Editor/UnrealEd/Classes/Settings/LevelEditorViewportSettings.h
@@ -514,6 +514,9 @@ public:
 	/** The scale to apply to spline tangent lengths */
 	UPROPERTY(EditAnywhere, config, Category = LookAndFeel, AdvancedDisplay, meta = (ClampMin = "0.00"))
 	float SplineTangentScale;
+	
+	UPROPERTY(config)
+	FVector2D LastInViewportMenuLocation;
 
 private:
 
diff --git a/Engine/Source/Editor/UnrealEd/Private/EditorViewportCommands.cpp b/Engine/Source/Editor/UnrealEd/Private/EditorViewportCommands.cpp
--- a/Engine/Source/Editor/UnrealEd/Private/EditorViewportCommands.cpp
+++ b/Engine/Source/Editor/UnrealEd/Private/EditorViewportCommands.cpp
@@ -160,6 +160,8 @@ void FEditorViewportCommands::RegisterCommands()
 
 	UI_COMMAND( ToggleAutoExposure, "Auto", "If enabled, enables automatic exposure", EUserInterfaceActionType::ToggleButton, FInputChord() );
 	UI_COMMAND( ToggleInGameExposure, "Game Settings", "If enabled, uses game settings", EUserInterfaceActionType::ToggleButton, FInputChord() );
+
+	UI_COMMAND(ToggleInViewportContextMenu, "In-Viewport Context Menu", "Shows a contextual menu of key properties and actions for the selected items in the viewport.", EUserInterfaceActionType::ToggleButton, FInputChord(EKeys::Tab));
 }
 
 #undef LOCTEXT_NAMESPACE
diff --git a/Engine/Source/Editor/UnrealEd/Private/SEditorViewport.cpp b/Engine/Source/Editor/UnrealEd/Private/SEditorViewport.cpp
--- a/Engine/Source/Editor/UnrealEd/Private/SEditorViewport.cpp
+++ b/Engine/Source/Editor/UnrealEd/Private/SEditorViewport.cpp
@@ -360,6 +360,12 @@ void SEditorViewport::BindCommands()
 		FCanExecuteAction(),
 		FIsActionChecked::CreateSP( this, &SEditorViewport::IsExposureSettingSelected ) );
 
+	CommandListRef.MapAction(
+		Commands.ToggleInViewportContextMenu,
+		FExecuteAction::CreateSP(this, &SEditorViewport::ToggleInViewportContextMenu),
+		FCanExecuteAction::CreateSP(this, &SEditorViewport::CanToggleInViewportContextMenu)
+	);
+
 	// Simple macro for binding many view mode UI commands
 
 #define MAP_VIEWMODEPARAM_ACTION( ViewModeCommand, ViewModeParam ) \
@@ -610,6 +616,14 @@ TSharedRef<SWidget> SEditorViewport::BuildFixedEV100Menu()  const
 };
 
 				
+void SEditorViewport::UpdateInViewportMenuLocation(const FVector2D InLocation)
+{
+	InViewportContextMenuLocation = InLocation;
+	ULevelEditorViewportSettings* LevelEditorViewportSettings = GetMutableDefault<ULevelEditorViewportSettings>();
+	LevelEditorViewportSettings->LastInViewportMenuLocation = InLocation;
+	LevelEditorViewportSettings->SaveConfig();
+}
+
 float SEditorViewport::OnGetFixedEV100Value() const
 {
 	if( Client.IsValid() )
diff --git a/Engine/Source/Editor/UnrealEd/Private/Settings/SettingsClasses.cpp b/Engine/Source/Editor/UnrealEd/Private/Settings/SettingsClasses.cpp
--- a/Engine/Source/Editor/UnrealEd/Private/Settings/SettingsClasses.cpp
+++ b/Engine/Source/Editor/UnrealEd/Private/Settings/SettingsClasses.cpp
@@ -902,6 +902,7 @@ ULevelEditorViewportSettings::ULevelEditorViewportSettings( const FObjectInitial
 	bAllowScreenRotate = false;
 	// Set a default preview mesh
 	PreviewMeshes.Add(FSoftObjectPath("/Engine/EditorMeshes/ColorCalibrator/SM_ColorCalibrator.SM_ColorCalibrator"));
+	LastInViewportMenuLocation = FVector2D(EForceInit::ForceInitToZero);
 }
 
 void ULevelEditorViewportSettings::PostInitProperties()
diff --git a/Engine/Source/Editor/UnrealEd/Public/EditorViewportCommands.h b/Engine/Source/Editor/UnrealEd/Public/EditorViewportCommands.h
--- a/Engine/Source/Editor/UnrealEd/Public/EditorViewportCommands.h
+++ b/Engine/Source/Editor/UnrealEd/Public/EditorViewportCommands.h
@@ -195,7 +195,8 @@ public:
 	TSharedPtr< FUICommandInfo > ToggleAutoExposure;
 	/** Toggle game settings exposure */
 	TSharedPtr< FUICommandInfo > ToggleInGameExposure;
-
+	/** Toggle the in-viewport context menu */
+	TSharedPtr< FUICommandInfo > ToggleInViewportContextMenu;
 	/**
 	 * Grid commands
 	 */
diff --git a/Engine/Source/Editor/UnrealEd/Public/SEditorViewport.h b/Engine/Source/Editor/UnrealEd/Public/SEditorViewport.h
--- a/Engine/Source/Editor/UnrealEd/Public/SEditorViewport.h
+++ b/Engine/Source/Editor/UnrealEd/Public/SEditorViewport.h
@@ -16,6 +16,7 @@ class FSceneViewport;
 class FUICommandList;
 class SViewport;
 struct FSlateBrush;
+struct FToolMenuContext;
 
 class UNREALED_API SEditorViewport
 	: public SCompoundWidget
@@ -91,6 +92,14 @@ public:
 	/** Build the exposure menu using EV100 settings */
 	TSharedRef<SWidget> BuildFixedEV100Menu()  const;
 
+	/**
+ * Called when the user wants to show the in-viewport context menu
+ */
+	virtual void ToggleInViewportContextMenu() {}
+	virtual void HideInViewportContextMenu() {}
+	virtual void UpdateInViewportMenuLocation(const FVector2D InLocation);
+	virtual bool CanToggleInViewportContextMenu() { return false; }
+
 ///////////////////////////////////////////////////////////////////////////////
 // begin feature level control functions block
 ///////////////////////////////////////////////////////////////////////////////
@@ -257,6 +266,8 @@ protected:
 	/** The last time the viewport was ticked (for visibility determination) */
 	double LastTickTime;
 
+	FVector2D InViewportContextMenuLocation;
+
 private:
 	/** Ensures a Slate tick/paint pass when the viewport is realtime or was invalidated this frame */
 	EActiveTimerReturnType EnsureTick( double InCurrentTime, float InDeltaTime );
diff --git a/Engine/Source/Runtime/Engine/Classes/Components/DirectionalLightComponent.h b/Engine/Source/Runtime/Engine/Classes/Components/DirectionalLightComponent.h
--- a/Engine/Source/Runtime/Engine/Classes/Components/DirectionalLightComponent.h
+++ b/Engine/Source/Runtime/Engine/Classes/Components/DirectionalLightComponent.h
@@ -150,7 +150,7 @@ class ENGINE_API UDirectionalLightComponent : public ULightComponent
 	/**
 	 * Whether the directional light can interact with the atmosphere, cloud and generate a visual disk. All of which compose the visual sky.
 	 */
-	UPROPERTY(EditAnywhere, BlueprintReadOnly, Category= AtmosphereAndCloud, meta=(DisplayName = "Atmosphere Sun Light"))
+	UPROPERTY(EditAnywhere, BlueprintReadOnly, Category= AtmosphereAndCloud, meta=(DisplayName = "Atmosphere Sun Light", ShouldShowInViewport = true))
 	uint32 bUsedAsAtmosphereSunLight : 1;
 
 	/**
diff --git a/Engine/Source/Runtime/Engine/Classes/Components/LightComponent.h b/Engine/Source/Runtime/Engine/Classes/Components/LightComponent.h
--- a/Engine/Source/Runtime/Engine/Classes/Components/LightComponent.h
+++ b/Engine/Source/Runtime/Engine/Classes/Components/LightComponent.h
@@ -45,7 +45,7 @@ class ENGINE_API ULightComponent : public ULightComponentBase
 	* Color temperature in Kelvin of the blackbody illuminant.
 	* White (D65) is 6500K.
 	*/
-	UPROPERTY(EditAnywhere, BlueprintReadOnly, interp, Category = Light, meta = (UIMin = "1700.0", UIMax = "12000.0"))
+	UPROPERTY(EditAnywhere, BlueprintReadOnly, interp, Category = Light, meta = (UIMin = "1700.0", UIMax = "12000.0", ShouldShowInViewport = true, DisplayAfter ="bUseTemperature"))
 	float Temperature;
 	
 	UPROPERTY(EditAnywhere, Category = Performance)
@@ -55,7 +55,7 @@ class ENGINE_API ULightComponent : public ULightComponentBase
 	float MaxDistanceFadeRange;
 
 	/** false: use white (D65) as illuminant. */
-	UPROPERTY(EditAnywhere, BlueprintReadOnly, Category=Light, meta=(DisplayName = "Use Temperature"))
+	UPROPERTY(EditAnywhere, BlueprintReadOnly, Category=Light, meta=(DisplayName = "Use Temperature", ShouldShowInViewport = true))
 	uint32 bUseTemperature : 1;
 
 	/** 
diff --git a/Engine/Source/Runtime/Engine/Classes/Components/LightComponentBase.h b/Engine/Source/Runtime/Engine/Classes/Components/LightComponentBase.h
--- a/Engine/Source/Runtime/Engine/Classes/Components/LightComponentBase.h
+++ b/Engine/Source/Runtime/Engine/Classes/Components/LightComponentBase.h
@@ -29,14 +29,14 @@ class ENGINE_API ULightComponentBase : public USceneComponent
 	/** 
 	 * Total energy that the light emits.  
 	 */
-	UPROPERTY(BlueprintReadOnly, interp, Category=Light, meta=(DisplayName = "Intensity", UIMin = "0.0", UIMax = "20.0"))
+	UPROPERTY(BlueprintReadOnly, interp, Category=Light, meta=(DisplayName = "Intensity", UIMin = "0.0", UIMax = "20.0", ShouldShowInViewport = true))
 	float Intensity;
 
 	/** 
 	 * Filter color of the light.
 	 * Note that this can change the light's effective intensity.
 	 */
-	UPROPERTY(BlueprintReadOnly, interp, Category=Light, meta=(HideAlphaChannel))
+	UPROPERTY(BlueprintReadOnly, interp, Category=Light, meta=(HideAlphaChannel, ShouldShowInViewport = true))
 	FColor LightColor;
 
 	/** 
@@ -44,7 +44,7 @@ class ENGINE_API ULightComponentBase : public USceneComponent
 	 * A disabled light will not contribute to the scene in any way.  This setting cannot be changed at runtime and unbuilds lighting when changed.
 	 * Setting this to false has the same effect as deleting the light, so it is useful for non-destructive experiments.
 	 */
-	UPROPERTY(EditAnywhere, BlueprintReadOnly, Category=Light)
+	UPROPERTY(EditAnywhere, BlueprintReadOnly, Category=Light, meta = (ShouldShowInViewport = true))
 	uint32 bAffectsWorld:1;
 
 	/**
diff --git a/Engine/Source/Runtime/Engine/Classes/Components/LocalLightComponent.h b/Engine/Source/Runtime/Engine/Classes/Components/LocalLightComponent.h
--- a/Engine/Source/Runtime/Engine/Classes/Components/LocalLightComponent.h
+++ b/Engine/Source/Runtime/Engine/Classes/Components/LocalLightComponent.h
@@ -33,7 +33,7 @@ class ENGINE_API ULocalLightComponent : public ULightComponent
 	 * Bounds the light's visible influence.  
 	 * This clamping of the light's influence is not physically correct but very important for performance, larger lights cost more.
 	 */
-	UPROPERTY(interp, BlueprintReadOnly, Category=Light, meta=(UIMin = "8.0", UIMax = "16384.0", SliderExponent = "5.0"))
+	UPROPERTY(interp, BlueprintReadOnly, Category=Light, meta=(UIMin = "8.0", UIMax = "16384.0", SliderExponent = "5.0", ShouldShowInViewport = true))
 	float AttenuationRadius;
 
 	/** The Lightmass settings for this object. */
diff --git a/Engine/Source/Runtime/Engine/Classes/Components/PointLightComponent.h b/Engine/Source/Runtime/Engine/Classes/Components/PointLightComponent.h
--- a/Engine/Source/Runtime/Engine/Classes/Components/PointLightComponent.h
+++ b/Engine/Source/Runtime/Engine/Classes/Components/PointLightComponent.h
@@ -40,7 +40,7 @@ class ENGINE_API UPointLightComponent : public ULocalLightComponent
 	 * Radius of light source shape.
 	 * Note that light sources shapes which intersect shadow casting geometry can cause shadowing artifacts.
 	 */
-	UPROPERTY(EditAnywhere, BlueprintReadOnly, Category=Light)
+	UPROPERTY(EditAnywhere, BlueprintReadOnly, Category=Light, meta=(ShouldShowInViewport = true))
 	float SourceRadius;
 
 	/**
diff --git a/Engine/Source/Runtime/Engine/Classes/Components/SpotLightComponent.h b/Engine/Source/Runtime/Engine/Classes/Components/SpotLightComponent.h
--- a/Engine/Source/Runtime/Engine/Classes/Components/SpotLightComponent.h
+++ b/Engine/Source/Runtime/Engine/Classes/Components/SpotLightComponent.h
@@ -18,11 +18,11 @@ class ENGINE_API USpotLightComponent : public UPointLightComponent
 	GENERATED_UCLASS_BODY()
 
 	/** Degrees. */
-	UPROPERTY(EditAnywhere, BlueprintReadOnly, Category=Light, meta=(UIMin = "1.0", UIMax = "80.0"))
+	UPROPERTY(EditAnywhere, BlueprintReadOnly, Category=Light, meta=(UIMin = "1.0", UIMax = "80.0", ShouldShowInViewport = true))
 	float InnerConeAngle;
 
 	/** Degrees. */
-	UPROPERTY(EditAnywhere, BlueprintReadOnly, Category=Light, meta=(UIMin = "1.0", UIMax = "80.0"))
+	UPROPERTY(EditAnywhere, BlueprintReadOnly, Category=Light, meta=(UIMin = "1.0", UIMax = "80.0", ShouldShowInViewport = true))
 	float OuterConeAngle;
 
 	UFUNCTION(BlueprintCallable, Category="Rendering|Lighting")
