Updated actor property editing to be typed element based

SActorDetails now operates on a given UTypedElementSelectionSet (removing its dependence on the global selection), and also operates based on the list of selected top-level elements rather than the selected actors.

UTypedElementDetailsInterface provides the interface needed for an element to be used with the details panel, with ITypedElementDetailsObject being the underlying glue object.
The ITypedElementDetailsObject instance will exist as long as the details panel is using it, so gives a lifetime to potentially synthesized UObject instances that are created purely for editing (eg, on instances).

This change removes almost all direct use of UpdateFloatingPropertyWindowsFromActorList, with only a single use-case in the reference viewer remaining.
Most code was already using UpdateFloatingPropertyWindows anyway, which now emits an extra OnElementSelectionChanged event in addition to its previous OnActorSelectionChanged event.
UpdateFloatingPropertyWindowsFromActorList now only emits an OnOverridePropertyEditorSelection event and no longer emits an OnActorSelectionChanged event.
SActorDetails is updated from both OnElementSelectionChanged and OnOverridePropertyEditorSelection, with the OnActorSelectionChanged event remaining for backwards compatability with existing code.

diff --git a/Engine/Source/Editor/EditorFramework/Public/Elements/Interfaces/TypedElementDetailsInterface.h b/Engine/Source/Editor/EditorFramework/Public/Elements/Interfaces/TypedElementDetailsInterface.h
new file mode 100644
--- /dev/null
+++ b/Engine/Source/Editor/EditorFramework/Public/Elements/Interfaces/TypedElementDetailsInterface.h
@@ -0,0 +1,53 @@
+// Copyright Epic Games, Inc. All Rights Reserved.
+
+#pragma once
+
+#include "Elements/Framework/TypedElementHandle.h"
+#include "TypedElementDetailsInterface.generated.h"
+
+/**
+ * Proxy instance to provide a UObject for editing by a details panel.
+ * This instance will exist as long as the details panel is using it, so gives a lifetime to potentially 
+ * synthesized UObject instances that are created purely for editing (eg, on instances).
+ */
+class ITypedElementDetailsObject
+{
+public:
+	virtual ~ITypedElementDetailsObject() = default;
+
+	/**
+	 * Get the underlying UObject that should be edited by the details panel, if any.
+	 */
+	virtual UObject* GetObject() = 0;
+
+	/**
+	 * Called during GC to collect references.
+	 */
+	virtual void AddReferencedObjects(FReferenceCollector& Collector) {}
+};
+
+UCLASS(Abstract)
+class EDITORFRAMEWORK_API UTypedElementDetailsInterface : public UTypedElementInterface
+{
+	GENERATED_BODY()
+
+public:
+	/**
+	 * Is the given element considered "top-level" for editing?
+	 * ie) Should it be edited simply from being part of the main selection set for a level?
+	 */
+	virtual bool IsTopLevelElement(const FTypedElementHandle& InElementHandle) { return true; }
+
+	/**
+	 * Get the proxy instance for the given element, if any.
+	 * @see ITypedElementDetailsObject for more details about this proxy instance.
+	 */
+	virtual TUniquePtr<ITypedElementDetailsObject> GetDetailsObject(const FTypedElementHandle& InElementHandle) { return nullptr; }
+};
+
+template <>
+struct TTypedElement<UTypedElementDetailsInterface> : public TTypedElementBase<UTypedElementDetailsInterface>
+{
+	bool IsTopLevelElement() const { return InterfacePtr->IsTopLevelElement(*this); }
+	TUniquePtr<ITypedElementDetailsObject> GetDetailsObject() const { return InterfacePtr->GetDetailsObject(*this); }
+};
diff --git a/Engine/Source/Editor/LevelEditor/Private/LevelEditor.cpp b/Engine/Source/Editor/LevelEditor/Private/LevelEditor.cpp
--- a/Engine/Source/Editor/LevelEditor/Private/LevelEditor.cpp
+++ b/Engine/Source/Editor/LevelEditor/Private/LevelEditor.cpp
@@ -418,6 +418,16 @@ void FLevelEditorModule::BroadcastActorSelectionChanged(const TArray<UObject*>&
 	ActorSelectionChangedEvent.Broadcast(NewSelection, bForceRefresh);
 }
 
+void FLevelEditorModule::BroadcastElementSelectionChanged(const UTypedElementSelectionSet* SelectionSet, bool bForceRefresh)
+{
+	ElementSelectionChangedEvent.Broadcast(SelectionSet, bForceRefresh);
+}
+
+void FLevelEditorModule::BroadcastOverridePropertyEditorSelection(const TArray<AActor*>& NewSelection, bool bForceRefresh)
+{
+	OverridePropertyEditorSelectionEvent.Broadcast(NewSelection, bForceRefresh);
+}
+
 void FLevelEditorModule::BroadcastRedrawViewports( bool bInvalidateHitProxies )
 {
 	RedrawLevelEditingViewportsEvent.Broadcast( bInvalidateHitProxies );
diff --git a/Engine/Source/Editor/LevelEditor/Private/SActorDetails.cpp b/Engine/Source/Editor/LevelEditor/Private/SActorDetails.cpp
--- a/Engine/Source/Editor/LevelEditor/Private/SActorDetails.cpp
+++ b/Engine/Source/Editor/LevelEditor/Private/SActorDetails.cpp
@@ -31,7 +31,9 @@
 #include "Widgets/Docking/SDockTab.h"
 #include "DetailsViewObjectFilter.h"
 #include "IDetailRootObjectCustomization.h"
-
+#include "Elements/Framework/TypedElementRegistry.h"
+#include "Elements/Framework/TypedElementSelectionSet.h"
+#include "Elements/Framework/EngineElementsLibrary.h"
 
 class SActorDetailsUneditableComponentWarning : public SCompoundWidget
 {
@@ -84,9 +86,11 @@ public:
 	}
 };
 
-void SActorDetails::Construct(const FArguments& InArgs, const FName TabIdentifier, TSharedPtr<FUICommandList> InCommandList, TSharedPtr<FTabManager> InTabManager)
+void SActorDetails::Construct(const FArguments& InArgs, UTypedElementSelectionSet* InSelectionSet, const FName TabIdentifier, TSharedPtr<FUICommandList> InCommandList, TSharedPtr<FTabManager> InTabManager)
 {
-	USelection::SelectionChangedEvent.AddRaw(this, &SActorDetails::OnEditorSelectionChanged);
+	SelectionSet = InSelectionSet;
+	checkf(SelectionSet, TEXT("SActorDetails must be constructed with a valid selection set!"));
+	SelectionSet->OnChanged().AddRaw(this, &SActorDetails::OnElementSelectionChanged);
 	
 	FLevelEditorModule& LevelEditor = FModuleManager::GetModuleChecked<FLevelEditorModule>("LevelEditor");
 	LevelEditor.OnComponentsEdited().AddRaw(this, &SActorDetails::OnComponentsEditedInWorld);
@@ -119,7 +123,6 @@ void SActorDetails::Construct(const FArguments& InArgs, const FName TabIdentifie
 	DetailsView->SetIsPropertyVisibleDelegate(FIsPropertyVisible::CreateLambda(IsPropertyVisible));
 	DetailsView->SetIsPropertyReadOnlyDelegate(FIsPropertyReadOnly::CreateSP(this, &SActorDetails::IsPropertyReadOnly));
 	DetailsView->SetIsPropertyEditingEnabledDelegate(FIsPropertyEditingEnabled::CreateSP(this, &SActorDetails::IsPropertyEditingEnabled));
-	DetailsView->SetOnObjectArrayChanged(FOnObjectArrayChanged::CreateSP(this, &SActorDetails::OnDetailsViewObjectArrayChanged));
 
 	// Set up a delegate to call to add generic details to the view
 	DetailsView->SetGenericLayoutDetailsDelegate(FOnGetDetailCustomizationInstance::CreateStatic(&FLevelEditorGenericDetails::MakeInstance));
@@ -206,6 +209,9 @@ void SActorDetails::Construct(const FArguments& InArgs, const FName TabIdentifie
 	[
 		ComponentsBox.ToSharedRef()
 	];
+
+	// Immediately update (otherwise we will appear empty)
+	RefreshSelection(/*bForceRefresh*/true);
 }
 
 SActorDetails::~SActorDetails()
@@ -214,7 +220,9 @@ SActorDetails::~SActorDetails()
 	{
 		GEditor->UnregisterForUndo(this);
 	}
-	USelection::SelectionChangedEvent.RemoveAll(this);
+
+	SelectionSet->OnChanged().RemoveAll(this);
+
 	RemoveBPComponentCompileEventDelegate();
 
 	FLevelEditorModule* LevelEditor = FModuleManager::GetModulePtr<FLevelEditorModule>("LevelEditor");
@@ -224,85 +232,192 @@ SActorDetails::~SActorDetails()
 	}
 }
 
-void SActorDetails::OnDetailsViewObjectArrayChanged(const FString& InTitle, const TArray<UObject*>& InObjects)
+bool SActorDetails::IsObservingSelectionSet(const UTypedElementSelectionSet* InSelectionSet) const
+{
+	return SelectionSet == InSelectionSet;
+}
+
+void SActorDetails::RefreshSelection(const bool bForceRefresh)
+{
+	TArray<TTypedElement<UTypedElementDetailsInterface>> DetailsElements;
+	DetailsElements.Reserve(SelectionSet->GetNumSelectedElements());
+	SelectionSet->ForEachSelectedElement<UTypedElementDetailsInterface>([&DetailsElements](const TTypedElement<UTypedElementDetailsInterface>& InDetailsElement)
+	{
+		DetailsElements.Add(InDetailsElement);
+		return true;
+	});
+
+	RefreshTopLevelElements(DetailsElements, bForceRefresh, /*bOverrideLock*/false);
+}
+
+void SActorDetails::OverrideSelection(const TArray<AActor*>& InActors, const bool bForceRefresh)
 {
-	// The DetailsView will already check validity every tick and hide itself when invalid, so this piggy-backs on that code instead of needing a second tick function.
-	if (InObjects.Num() == 0 && !LockedActorSelection.IsValid())
+	UTypedElementRegistry* Registry = UTypedElementRegistry::GetInstance();
+
+	TArray<TTypedElement<UTypedElementDetailsInterface>> DetailsElements;
+	DetailsElements.Reserve(SelectionSet->GetNumSelectedElements());
+	for (AActor* Actor : InActors)
 	{
-		bHasComponentsToShow = false;
+		if (FTypedElementHandle ActorElementHandle = UEngineElementsLibrary::AcquireEditorActorElementHandle(Actor))
+		{
+			if (TTypedElement<UTypedElementDetailsInterface> ActorDetailsHandle = Registry->GetElement<UTypedElementDetailsInterface>(ActorElementHandle))
+			{
+				DetailsElements.Add(MoveTemp(ActorDetailsHandle));
+			}
+		}
 	}
+
+	RefreshTopLevelElements(DetailsElements, bForceRefresh, /*bOverrideLock*/false);
 }
 
-void SActorDetails::SetObjects(const TArray<UObject*>& InObjects, bool bForceRefresh)
+void SActorDetails::RefreshTopLevelElements(TArrayView<const TTypedElement<UTypedElementDetailsInterface>> InDetailsElements, const bool bForceRefresh, const bool bOverrideLock)
 {
-	if (!DetailsView->IsLocked())
+	// Nothing to do if this view is locked!
+	if (DetailsView->IsLocked() && !bOverrideLock)
+	{
+		return;
+	}
+
+	// Build the array of top-level elements to edit
+	TopLevelElements.Reset(InDetailsElements.Num());
+	for (const TTypedElement<UTypedElementDetailsInterface>& DetailsElement : InDetailsElements)
+	{
+		if (DetailsElement.IsTopLevelElement())
+		{
+			if (TUniquePtr<ITypedElementDetailsObject> ElementDetailsObject = DetailsElement.GetDetailsObject())
+			{
+				TopLevelElements.Add(MoveTemp(ElementDetailsObject));
+			}
+		}
+	}
+
+	// Update the underlying details view
+	SetElementDetailsObjects(TopLevelElements, bForceRefresh, bOverrideLock);
+
+	// Update the SCS tree if we were asked to edit a single actor
+	if (InDetailsElements.Num() == 1)
+	{
+		if (AActor* Actor = GetActorContext())
+		{
+			// Enable the selection guard to prevent OnTreeSelectionChanged() from altering the editor's component selection
+			TGuardValue<bool> SelectionGuard(bSelectionGuard, true);
+			SCSEditor->UpdateTree();
+		}
+	}
+
+	// Draw attention to this tab if needed
+	if (TSharedPtr<FTabManager> TabManager = DetailsView->GetHostTabManager())
 	{
-		DetailsView->SetObjects(InObjects, bForceRefresh);
+		TSharedPtr<SDockTab> Tab = TabManager->FindExistingLiveTab(DetailsView->GetIdentifier());
+		if (Tab.IsValid() && !Tab->IsForeground())
+		{
+			Tab->FlashTab();
+		}
+	}
+}
 
-		bHasComponentsToShow = false;
+void SActorDetails::RefreshSCSTreeElements(TArrayView<const TSharedPtr<class FSCSEditorTreeNode>> InSelectedNodes, const bool bForceRefresh, const bool bOverrideLock)
+{
+	// Nothing to do if this view is locked!
+	if (DetailsView->IsLocked() && !bOverrideLock)
+	{
+		return;
+	}
 
-		if (InObjects.Num() == 1 && FKismetEditorUtilities::CanCreateBlueprintOfClass(InObjects[0]->GetClass()))
+	// Does the SCS tree have components selected?
+	TArray<UActorComponent*> Components;
+	if (AActor* Actor = GetActorContext())
+	{
+		for (const FSCSEditorTreeNodePtrType& SelectedNode : InSelectedNodes)
 		{
-			AActor* Actor = GetSelectedActorInEditor();
-			if (Actor)
+			if (SelectedNode)
 			{
-				LockedActorSelection = Actor;
-				bHasComponentsToShow = true;
+				if (SelectedNode->GetNodeType() == FSCSEditorTreeNode::RootActorNode)
+				{
+					// If the actor node is selected then we ignore the component selection
+					Components.Reset();
+					break;
+				}
 
-				// Update the tree if a new actor is selected
-				if (GEditor->GetSelectedComponentCount() == 0)
+				if (SelectedNode->GetNodeType() == FSCSEditorTreeNode::ComponentNode)
 				{
-					// Enable the selection guard to prevent OnTreeSelectionChanged() from altering the editor's component selection
-					TGuardValue<bool> SelectionGuard(bSelectionGuard, true);
-					SCSEditor->UpdateTree();
+					if (UActorComponent* Component = SelectedNode->FindComponentInstanceInActor(Actor))
+					{
+						Components.Add(Component);
+					}
 				}
 			}
 		}
+	}
 
-		if (DetailsView->GetHostTabManager().IsValid())
+	SCSTreeElements.Reset(Components.Num());
+	if (Components.Num() > 0)
+	{
+		UTypedElementRegistry* Registry = UTypedElementRegistry::GetInstance();
+		for (UActorComponent* Component : Components)
 		{
-			TSharedPtr<SDockTab> Tab = DetailsView->GetHostTabManager()->FindExistingLiveTab(DetailsView->GetIdentifier());
-			if (Tab.IsValid() && !Tab->IsForeground() )
+			if (FTypedElementHandle ComponentElementHandle = UEngineElementsLibrary::AcquireEditorComponentElementHandle(Component))
 			{
-				Tab->FlashTab();
+				if (TTypedElement<UTypedElementDetailsInterface> ComponentDetailsHandle = Registry->GetElement<UTypedElementDetailsInterface>(ComponentElementHandle))
+				{
+					if (TUniquePtr<ITypedElementDetailsObject> ElementDetailsObject = ComponentDetailsHandle.GetDetailsObject())
+					{
+						SCSTreeElements.Add(MoveTemp(ElementDetailsObject));
+					}
+				}
 			}
 		}
+
+		// Use the component elements
+		SetElementDetailsObjects(SCSTreeElements, bForceRefresh, bOverrideLock);
+	}
+	else
+	{
+		// Use the top-level elements
+		SetElementDetailsObjects(TopLevelElements, bForceRefresh, bOverrideLock);
 	}
 }
 
-void SActorDetails::PostUndo(bool bSuccess)
+void SActorDetails::SetElementDetailsObjects(TArrayView<const TUniquePtr<ITypedElementDetailsObject>> InElementDetailsObjects, const bool bForceRefresh, const bool bOverrideLock)
 {
-	// Enable the selection guard to prevent OnTreeSelectionChanged() from altering the editor's component selection
-	TGuardValue<bool> SelectionGuard(bSelectionGuard, true);
-
-	if (!DetailsView->IsLocked())
+	TArray<UObject*> DetailsObjects;
+	DetailsObjects.Reserve(InElementDetailsObjects.Num());
+	for (const TUniquePtr<ITypedElementDetailsObject>& ElementDetailsObject : InElementDetailsObjects)
 	{
-		// Make sure the locked actor selection matches the editor selection
-		AActor* SelectedActor = GetSelectedActorInEditor();
-		if (SelectedActor && SelectedActor != LockedActorSelection.Get())
+		if (UObject* DetailsObject = ElementDetailsObject->GetObject())
 		{
-			LockedActorSelection = SelectedActor;
+			DetailsObjects.Add(DetailsObject);
 		}
 	}
-	
+	DetailsView->SetObjects(DetailsObjects, bForceRefresh, bOverrideLock);
+}
 
+void SActorDetails::PostUndo(bool bSuccess)
+{
+	// Enable the selection guard to prevent OnTreeSelectionChanged() from altering the editor's component selection
+	TGuardValue<bool> SelectionGuard(bSelectionGuard, true);
+	
 	// Refresh the tree and update the selection to match the world
 	SCSEditor->UpdateTree();
 	UpdateComponentTreeFromEditorSelection();
+}
 
-	AActor* SelectedActor = GetSelectedActorInEditor();
-	if (SelectedActor)
-	{
-		SelectedActor->PushSelectionToProxies();
+void SActorDetails::PostRedo(bool bSuccess)
+{
+	PostUndo(bSuccess);
+}
 
-		// Update the pivot (widget) as the current selection may be a component within the Actor instance
-		GUnrealEd->UpdatePivotLocationForSelection();
+void SActorDetails::AddReferencedObjects(FReferenceCollector& Collector)
+{
+	for (const TUniquePtr<ITypedElementDetailsObject>& TopLevelElement : TopLevelElements)
+	{
+		TopLevelElement->AddReferencedObjects(Collector);
 	}
 }
 
-void SActorDetails::PostRedo(bool bSuccess)
+FString SActorDetails::GetReferencerName() const
 {
-	PostUndo(bSuccess);
+	return TEXT("SActorDetails");
 }
 
 void SActorDetails::SetActorDetailsRootCustomization(TSharedPtr<FDetailsViewObjectFilter> ActorDetailsObjectFilter, TSharedPtr<IDetailRootObjectCustomization> ActorDetailsRootCustomization)
@@ -322,48 +437,30 @@ void SActorDetails::SetSCSEditorUICustomization(TSharedPtr<ISCSEditorUICustomiza
 
 void SActorDetails::OnComponentsEditedInWorld()
 {
-	if (GetSelectedActorInEditor() == GetActorContext())
+	if (AActor* Actor = GetActorContext())
 	{
-		// The component composition of the observed actor has changed, so rebuild the node tree
-		TGuardValue<bool> SelectionGuard(bSelectionGuard, true);
-
+		if (SelectionSet->IsElementSelected(UEngineElementsLibrary::AcquireEditorActorElementHandle(Actor), FTypedElementIsSelectedOptions()))
+		{
+			// The component composition of the observed actor has changed, so rebuild the node tree
+			TGuardValue<bool> SelectionGuard(bSelectionGuard, true);
 
-		// Refresh the tree and update the selection to match the world
-		SCSEditor->UpdateTree();
+			// Refresh the tree and update the selection to match the world
+			SCSEditor->UpdateTree();
 
-		DetailsView->ForceRefresh();
+			DetailsView->ForceRefresh();
+		}
 	}
 }
 
-void SActorDetails::OnEditorSelectionChanged(UObject* Object)
+void SActorDetails::OnElementSelectionChanged(const UTypedElementSelectionSet* InSelectionSet)
 {
-	if(!bSelectionGuard && SCSEditor.IsValid())
+	check(IsObservingSelectionSet(InSelectionSet));
+	if (!bSelectionGuard && SCSEditor)
 	{
-		// Make sure the selection set that changed is relevant to us
-		USelection* Selection = Cast<USelection>(Object);
-		if(Selection == GEditor->GetSelectedComponents() || Selection == GEditor->GetSelectedActors())
-		{
-			UpdateComponentTreeFromEditorSelection();
-
-			if(GEditor->GetSelectedComponentCount() == 0) // An actor was selected
-			{
-				// Ensure the selection flags are up to date for the components in the selected actor
-				for(FSelectionIterator It(GEditor->GetSelectedActorIterator()); It; ++It)
-				{
-					AActor* Actor = CastChecked<AActor>(*It);
-					Actor->PushSelectionToProxies();
-				}
-			}
-		}
+		UpdateComponentTreeFromEditorSelection();
 	}
 }
 
-AActor* SActorDetails::GetSelectedActorInEditor() const
-{
-	//@todo this doesn't work w/ multi-select
-	return GEditor->GetSelectedActors()->GetTop<AActor>();
-}
-
 bool SActorDetails::GetAllowComponentTreeEditing() const
 {
 	return GEditor->PlayWorld == nullptr;
@@ -371,252 +468,141 @@ bool SActorDetails::GetAllowComponentTreeEditing() const
 
 AActor* SActorDetails::GetActorContext() const
 {
-	AActor* SelectedActorInEditor = GetSelectedActorInEditor();
-	const bool bDetailsLocked = DetailsView->IsLocked();
-	
-	// If the details is locked or we have a valid locked selection that doesn't match the editor's selected actor, use the locked selection
-	if (bDetailsLocked || ( LockedActorSelection.IsValid() && LockedActorSelection.Get() != SelectedActorInEditor ))
-	{
-		return LockedActorSelection.Get();
-	}
-	else
-	{
-		return SelectedActorInEditor;
-	}
+	return TopLevelElements.Num() == 1
+		? Cast<AActor>(TopLevelElements[0]->GetObject())
+		: nullptr;
 }
 
 void SActorDetails::OnSCSEditorTreeViewSelectionChanged(const TArray<FSCSEditorTreeNodePtrType>& SelectedNodes)
 {
-	if (!bSelectionGuard && SelectedNodes.Num() > 0)
+	if (bSelectionGuard)
 	{
-		if( SelectedNodes.Num() > 1 && SelectedBPComponentBlueprint.IsValid() )
-		{
-			// Remove the compilation delegate if we are no longer displaying the full details for a single blueprint component.
-			RemoveBPComponentCompileEventDelegate();
-		}
+		// Preventing selection changes from having an effect...
+		return;
+	}
 
-		AActor* Actor = GetActorContext();
-		if (Actor)
-		{
-			TArray<UObject*> DetailsObjects;
+	if (SelectedNodes.Num() == 0)
+	{
+		// Don't respond to de-selecting everything...
+		return;
+	}
+	
+	AActor* Actor = GetActorContext();
+	if (!Actor)
+	{
+		// The SCS editor requires an actor context...
+		return;
+	}
 
-			// Determine if the root actor node is among the selected nodes and Count number of components selected
-			bool bActorNodeSelected = false;
-			int NumSelectedComponentNodes = 0;
-			for (const FSCSEditorTreeNodePtrType& SelectedNode : SelectedNodes)
+	if (SelectedNodes.Num() > 1 && SelectedBPComponentBlueprint.IsValid())
+	{
+		// Remove the compilation delegate if we are no longer displaying the full details for a single blueprint component.
+		RemoveBPComponentCompileEventDelegate();
+	}
+	else if (SelectedNodes.Num() == 1 && SelectedNodes[0]->GetNodeType() == FSCSEditorTreeNode::ComponentNode)
+	{
+		// Add delegate to monitor blueprint component compilation if we have a full details view ( i.e. single selection )
+		if (UActorComponent* Component = SelectedNodes[0]->FindComponentInstanceInActor(Actor))
+		{
+			if (UBlueprintGeneratedClass* ComponentBPGC = Cast<UBlueprintGeneratedClass>(Component->GetClass()))
 			{
-				if (SelectedNode.IsValid())
+				if (UBlueprint* ComponentBlueprint = Cast<UBlueprint>(ComponentBPGC->ClassGeneratedBy))
 				{
-					if (SelectedNode->GetNodeType() == FSCSEditorTreeNode::RootActorNode)
-					{
-						bActorNodeSelected = true;
-					}
-					else if (SelectedNode->GetNodeType() == FSCSEditorTreeNode::ComponentNode)
-					{
-						++NumSelectedComponentNodes;
-					}
+					AddBPComponentCompileEventDelegate(ComponentBlueprint);
 				}
 			}
+		}
+	}
 
-			if (DetailsView->IsLocked())
-			{
-				// When the details panel is locked, we don't want to touch the editor's component selection
-				// We do want to force the locked panel to update to match the selected components, though, since they are part of the actor selection we're locked on
-
-				if (bActorNodeSelected)
-				{
-					// If the actor root is selected, then the editor component selection should remain empty and we only show the Actor's details
-					DetailsObjects.Add(Actor);
-				}
-				else
-				{
-					const bool bSingleComponentSelection = SelectedNodes.Num() == 1;
-
-					for (const FSCSEditorTreeNodePtrType& SelectedNode : SelectedNodes)
-					{
-						UActorComponent* ComponentInstance = SelectedNode->FindComponentInstanceInActor(Actor);
-						if (ComponentInstance)
-						{
-							DetailsObjects.Add(ComponentInstance);
-
-							if(bSingleComponentSelection)
-							{
-								// Add delegate to monitor blueprint component compilation if we have a full details view ( i.e. single selection )
-								if(UBlueprintGeneratedClass* ComponentBPGC = Cast<UBlueprintGeneratedClass>(ComponentInstance->GetClass()))
-								{
-									if(UBlueprint* ComponentBlueprint = Cast<UBlueprint>(ComponentBPGC->ClassGeneratedBy))
-									{
-										AddBPComponentCompileEventDelegate(ComponentBlueprint);
-									}
-								}
-							}
-						}
-					}
-				}
+	// We only actually update the editor selection state if we're not locked
+	if (!DetailsView->IsLocked())
+	{
+		TArray<FTypedElementHandle> NewEditorSelection;
+		NewEditorSelection.Add(UEngineElementsLibrary::AcquireEditorActorElementHandle(Actor));
 
-				const bool bOverrideDetailsLock = true;
-				DetailsView->SetObjects(DetailsObjects, false, bOverrideDetailsLock);
-			}
-			else
+		for (const FSCSEditorTreeNodePtrType& SelectedNode : SelectedNodes)
+		{
+			if (SelectedNode)
 			{
-				// Enable the selection guard to prevent OnEditorSelectionChanged() from altering the contents of the SCSTreeWidget
-				TGuardValue<bool> SelectionGuard(bSelectionGuard, true);
-
-				// Make sure the actor is selected in the editor (possible if the panel was just unlocked, but still assigned to the locked actor)
-				if (!GEditor->GetSelectedActors()->IsSelected(Actor))
+				if (SelectedNode->GetNodeType() == FSCSEditorTreeNode::RootActorNode)
 				{
-					GEditor->SelectNone(false, true, false);
-					GEditor->SelectActor(Actor, true, true, true);
-				}
-
-				USelection* SelectedComponents = GEditor->GetSelectedComponents();
-
-				// Determine if the selected non-root actor nodes differ from the editor component selection
-				bool bComponentSelectionChanged = GEditor->GetSelectedComponentCount() != NumSelectedComponentNodes;
-				if (!bComponentSelectionChanged)
-				{
-					// Check to see if any of the selected nodes aren't already selected in the world
-					for (const FSCSEditorTreeNodePtrType& SelectedNode : SelectedNodes)
-					{
-						if (SelectedNode.IsValid() && SelectedNode->GetNodeType() == FSCSEditorTreeNode::ComponentNode)
-						{
-							UActorComponent* ComponentInstance = SelectedNode->FindComponentInstanceInActor(Actor);
-							if (ComponentInstance && !SelectedComponents->IsSelected(ComponentInstance))
-							{
-								bComponentSelectionChanged = true;
-								break;
-							}
-						}
-					}
+					// If the actor node is selected then we ignore the component selection
+					NewEditorSelection.Reset();
+					NewEditorSelection.Add(UEngineElementsLibrary::AcquireEditorActorElementHandle(Actor));
+					break;
 				}
 
-				// Does the actor selection differ from our previous state?
-				const bool bActorSelectionChanged = bShowingRootActorNodeSelected != bActorNodeSelected;
-
-				// If necessary, update the editor component selection
-				if (bActorSelectionChanged || ( bComponentSelectionChanged && !bActorNodeSelected ))
+				if (SelectedNode->GetNodeType() == FSCSEditorTreeNode::ComponentNode)
 				{
-					// Store whether we're now showing the actor root as selected
-					bShowingRootActorNodeSelected = bActorNodeSelected;
-
-					// Note: this transaction should not take place if we are in the middle of executing an undo or redo because it would clear the top of the transaction stack.
-					const bool bShouldActuallyTransact = !GIsTransacting;
-					const FScopedTransaction Transaction(NSLOCTEXT("UnrealEd", "ClickingOnComponentInTree", "Clicking on Component (tree view)"), bShouldActuallyTransact);
-
-					// Dirty the actor selection so it stays in sync with the component selection
-					GEditor->GetSelectedActors()->Modify();
-					// Update the editor's component selection to match the node selection
-					SelectedComponents->Modify();
-					SelectedComponents->BeginBatchSelectOperation();
-					SelectedComponents->DeselectAll();
-
-					if (bShowingRootActorNodeSelected)
+					if (UActorComponent* Component = SelectedNode->FindComponentInstanceInActor(Actor))
 					{
-						// If the actor root is selected, then the editor component selection should remain empty and we only show the Actor's details
-						DetailsObjects.Add(Actor);
+						NewEditorSelection.Add(UEngineElementsLibrary::AcquireEditorComponentElementHandle(Component));
 					}
-					else
-					{
-						const bool bSingleComponentSelection = SelectedNodes.Num() == 1;
-
-						for (const FSCSEditorTreeNodePtrType& SelectedNode : SelectedNodes)
-						{
-							if (SelectedNode.IsValid())
-							{
-								UActorComponent* ComponentInstance = SelectedNode->FindComponentInstanceInActor(Actor);
-								if (ComponentInstance)
-								{
-									DetailsObjects.Add(ComponentInstance);
-									SelectedComponents->Select(ComponentInstance);
-
-									if(bSingleComponentSelection)
-									{
-										// Add delegate to monitor blueprint component compilation if we have a full details view ( i.e. single selection )
-										if(UBlueprintGeneratedClass* ComponentBPGC = Cast<UBlueprintGeneratedClass>(ComponentInstance->GetClass()))
-										{
-											if(UBlueprint* ComponentBlueprint = Cast<UBlueprint>(ComponentBPGC->ClassGeneratedBy))
-											{
-												AddBPComponentCompileEventDelegate(ComponentBlueprint);
-											}
-										}
-									}
-									// Ensure the selection override is bound for this component (including any attached editor-only children)
-									USceneComponent* SceneComponent = Cast<USceneComponent>(ComponentInstance);
-									if (SceneComponent)
-									{
-										FComponentEditorUtils::BindComponentSelectionOverride(SceneComponent, true);
-									}
-								}
-							}
-						}
-					}
-
-					SelectedComponents->EndBatchSelectOperation();
-
-					DetailsView->SetObjects(DetailsObjects);
-
-					Actor->PushSelectionToProxies();
-					GUnrealEd->UpdatePivotLocationForSelection(true);
-					GEditor->RedrawLevelEditingViewports();
 				}
 			}
 		}
+
+		// Note: this transaction should not take place if we are in the middle of executing an undo or redo because it would clear the top of the transaction stack.
+		const bool bShouldActuallyTransact = !GIsTransacting;
+		const FScopedTransaction Transaction(NSLOCTEXT("UnrealEd", "ClickingOnComponentInTree", "Clicking on Component (tree view)"), bShouldActuallyTransact);
+
+		// Enable the selection guard to prevent OnEditorSelectionChanged() from altering the contents of the SCSTreeWidget
+		TGuardValue<bool> SelectionGuard(bSelectionGuard, true);
+		SelectionSet->SetSelection(NewEditorSelection, FTypedElementSelectionOptions());
+		SelectionSet->NotifyPendingChanges(); // Fire while still under the selection guard
 	}
+
+	// Update the underlying details view
+	RefreshSCSTreeElements(SelectedNodes, /*bForceRefresh*/false, DetailsView->IsLocked());
 }
 
 void SActorDetails::OnSCSEditorTreeViewItemDoubleClicked(const TSharedPtr<class FSCSEditorTreeNode> ClickedNode)
 {
-	if (ClickedNode.IsValid())
+	if (ClickedNode && ClickedNode->GetNodeType() == FSCSEditorTreeNode::ComponentNode)
 	{
-		if (ClickedNode->GetNodeType() == FSCSEditorTreeNode::ComponentNode)
+		if (USceneComponent* SceneComponent = Cast<USceneComponent>(ClickedNode->GetComponentTemplate()))
 		{
-			USceneComponent* SceneComponent = Cast<USceneComponent>(ClickedNode->GetComponentTemplate());
-			if (SceneComponent != nullptr)
-			{
-				const bool bActiveViewportOnly = false;
-				GEditor->MoveViewportCamerasToComponent(SceneComponent, bActiveViewportOnly);
-			}
+			const bool bActiveViewportOnly = false;
+			GEditor->MoveViewportCamerasToComponent(SceneComponent, bActiveViewportOnly);
 		}
 	}
 }
 
 void SActorDetails::UpdateComponentTreeFromEditorSelection()
 {
-	if (!DetailsView->IsLocked())
+	if (DetailsView->IsLocked())
 	{
-		// Enable the selection guard to prevent OnTreeSelectionChanged() from altering the editor's component selection
-		TGuardValue<bool> SelectionGuard(bSelectionGuard, true);
-
-		TSharedPtr<SSCSTreeType>& SCSTreeWidget = SCSEditor->SCSTreeWidget;
-		TArray<UObject*> DetailsObjects;
-
-		// Update the tree selection to match the level editor component selection
-		SCSTreeWidget->ClearSelection();
-		for (FSelectionIterator It(GEditor->GetSelectedComponentIterator()); It; ++It)
-		{
-			UActorComponent* Component = CastChecked<UActorComponent>(*It);
+		return;
+	}
 
-			FSCSEditorTreeNodePtrType SCSTreeNode = SCSEditor->GetNodeFromActorComponent(Component, false);
-			if (SCSTreeNode.IsValid() && SCSTreeNode->GetComponentTemplate())
-			{
-				SCSTreeWidget->RequestScrollIntoView(SCSTreeNode);
-				SCSTreeWidget->SetItemSelection(SCSTreeNode, true);
+	// Enable the selection guard to prevent OnTreeSelectionChanged() from altering the editor's component selection
+	TGuardValue<bool> SelectionGuard(bSelectionGuard, true);
 
-				UActorComponent* ComponentTemplate = SCSTreeNode->GetComponentTemplate();
-				check(Component == ComponentTemplate);
-				DetailsObjects.Add(Component);
-			}
-		}
+	TSharedPtr<SSCSTreeType>& SCSTreeWidget = SCSEditor->SCSTreeWidget;
 
-		if (DetailsObjects.Num() > 0)
-		{
-			DetailsView->SetObjects(DetailsObjects, bSelectedComponentRecompiled);
-		}
-		else
+	// Update the tree selection to match the level editor component selection
+	SCSTreeWidget->ClearSelection();
+	SelectionSet->ForEachSelectedObject<UActorComponent>([this, &SCSTreeWidget](UActorComponent* InComponent)
+	{
+		FSCSEditorTreeNodePtrType SCSTreeNode = SCSEditor->GetNodeFromActorComponent(InComponent, false);
+		if (SCSTreeNode && SCSTreeNode->GetComponentTemplate())
 		{
-			SCSEditor->SelectRoot();
+			SCSTreeWidget->RequestScrollIntoView(SCSTreeNode);
+			SCSTreeWidget->SetItemSelection(SCSTreeNode, true);
+			check(InComponent == SCSTreeNode->GetComponentTemplate());
 		}
+		return true;
+	});
+
+	TArray<TSharedPtr<FSCSEditorTreeNode>> SelectedNodes = SCSEditor->GetSelectedNodes();
+	if (SelectedNodes.Num() == 0)
+	{
+		SCSEditor->SelectRoot();
+		SelectedNodes = SCSEditor->GetSelectedNodes();
 	}
+
+	// Update the underlying details view
+	RefreshSCSTreeElements(SelectedNodes, bSelectedComponentRecompiled, /*bOverrideLock*/false);
 }
 
 bool SActorDetails::IsPropertyReadOnly(const FPropertyAndParent& PropertyAndParent) const
@@ -695,7 +681,7 @@ void SActorDetails::OnNativeComponentWarningHyperlinkClicked(const FSlateHyperli
 
 EVisibility SActorDetails::GetComponentsBoxVisibility() const
 {
-	return bHasComponentsToShow ? EVisibility::Visible : EVisibility::Collapsed;
+	return GetActorContext() ? EVisibility::Visible : EVisibility::Collapsed;
 }
 
 EVisibility SActorDetails::GetUCSComponentWarningVisibility() const
@@ -814,7 +800,6 @@ void SActorDetails::RemoveBPComponentCompileEventDelegate()
 
 void SActorDetails::OnBlueprintComponentCompiled(UBlueprint* ComponentBlueprint)
 {
-	bSelectedComponentRecompiled = true;
+	TGuardValue<bool> SelectedComponentRecompiledGuard(bSelectedComponentRecompiled, true);
 	UpdateComponentTreeFromEditorSelection();
-	bSelectedComponentRecompiled = false;
 }
diff --git a/Engine/Source/Editor/LevelEditor/Private/SActorDetails.h b/Engine/Source/Editor/LevelEditor/Private/SActorDetails.h
--- a/Engine/Source/Editor/LevelEditor/Private/SActorDetails.h
+++ b/Engine/Source/Editor/LevelEditor/Private/SActorDetails.h
@@ -7,8 +7,9 @@
 #include "Widgets/DeclarativeSyntaxSupport.h"
 #include "Widgets/SCompoundWidget.h"
 #include "Framework/Text/SlateHyperlinkRun.h"
+#include "Containers/ArrayView.h"
 #include "EditorUndoClient.h"
-
+#include "Elements/Interfaces/TypedElementDetailsInterface.h"
 
 class AActor;
 class FSCSEditorTreeNode;
@@ -20,30 +21,43 @@ class SSCSEditor;
 class SSplitter;
 class UBlueprint;
 class FDetailsViewObjectFilter;
+class UTypedElementSelectionSet;
 
 /**
  * Wraps a details panel customized for viewing actors
  */
-class SActorDetails : public SCompoundWidget, public FEditorUndoClient
+class SActorDetails : public SCompoundWidget, public FEditorUndoClient, public FGCObject
 {
 public:
 	SLATE_BEGIN_ARGS(SActorDetails) {}
 	SLATE_END_ARGS()
 
-	void Construct(const FArguments& InArgs, const FName TabIdentifier, TSharedPtr<FUICommandList> InCommandList, TSharedPtr<FTabManager> InTabManager);
+	void Construct(const FArguments& InArgs, UTypedElementSelectionSet* InSelectionSet, const FName TabIdentifier, TSharedPtr<FUICommandList> InCommandList, TSharedPtr<FTabManager> InTabManager);
 	~SActorDetails();
 
 	/**
-	 * Sets the objects to be viewed by the details panel
-	 *
-	 * @param InObjects	The objects to set
+	 * Return true if this details panel is observing the given selection set.
+	 */
+	bool IsObservingSelectionSet(const UTypedElementSelectionSet* InSelectionSet) const;
+
+	/**
+	 * Update the view based on our observed selection set.
+	 */
+	void RefreshSelection(const bool bForceRefresh = false);
+
+	/**
+	 * Update the view based on the given set of actors.
 	 */
-	void SetObjects(const TArray<UObject*>& InObjects, bool bForceRefresh = false);
+	void OverrideSelection(const TArray<AActor*>& InActors, const bool bForceRefresh = false);
 
 	/** FEditorUndoClient Interface */
 	virtual void PostUndo(bool bSuccess) override;
 	virtual void PostRedo(bool bSuccess) override;
 	
+	//~ FGCObject
+	virtual void AddReferencedObjects(FReferenceCollector& Collector) override;
+	virtual FString GetReferencerName() const override;
+
 	/**
 	 * Sets the filter that should be used to filter incoming actors in or out of the details panel
 	 *
@@ -55,16 +69,18 @@ public:
 	void SetSCSEditorUICustomization(TSharedPtr<class ISCSEditorUICustomization> ActorDetailsSCSEditorUICustomization);
 
 private:
-	AActor* GetSelectedActorInEditor() const;
+	void RefreshTopLevelElements(TArrayView<const TTypedElement<UTypedElementDetailsInterface>> InDetailsElements, const bool bForceRefresh, const bool bOverrideLock);
+	void RefreshSCSTreeElements(TArrayView<const TSharedPtr<class FSCSEditorTreeNode>> InSelectedNodes, const bool bForceRefresh, const bool bOverrideLock);
+	void SetElementDetailsObjects(TArrayView<const TUniquePtr<ITypedElementDetailsObject>> InElementDetailsObjects, const bool bForceRefresh, const bool bOverrideLock);
+
 	AActor* GetActorContext() const;
 	bool GetAllowComponentTreeEditing() const;
 
 	void OnComponentsEditedInWorld();
-	void OnEditorSelectionChanged(UObject* Object);
+	void OnElementSelectionChanged(const UTypedElementSelectionSet* InSelectionSet);
 	void OnSCSEditorTreeViewSelectionChanged(const TArray<TSharedPtr<class FSCSEditorTreeNode> >& SelectedNodes);
 	void OnSCSEditorTreeViewItemDoubleClicked(const TSharedPtr<class FSCSEditorTreeNode> ClickedNode);
 	void UpdateComponentTreeFromEditorSelection();
-	void OnDetailsViewObjectArrayChanged(const FString& InTitle, const TArray<UObject*>& InObjects);
 
 	bool IsPropertyReadOnly(const struct FPropertyAndParent& PropertyAndParent) const;
 	bool IsPropertyEditingEnabled() const;
@@ -85,8 +101,14 @@ private:
 	TSharedPtr<SBox> ComponentsBox;
 	TSharedPtr<class SSCSEditor> SCSEditor;
 
-	// The actor selected when the details panel was locked
-	TWeakObjectPtr<AActor> LockedActorSelection;
+	// The selection set this details panel is observing
+	UTypedElementSelectionSet* SelectionSet = nullptr;
+
+	// Array of top-level elements that are currently being edited
+	TArray<TUniquePtr<ITypedElementDetailsObject>> TopLevelElements;
+
+	// Array of component elements that are being edited from the SCS tree selection
+	TArray<TUniquePtr<ITypedElementDetailsObject>> SCSTreeElements;
 
 	// The current component blueprint selection
 	TWeakObjectPtr<UBlueprint> SelectedBPComponentBlueprint;
@@ -94,10 +116,4 @@ private:
 
 	// Used to prevent reentrant changes
 	bool bSelectionGuard = false;
-
-	// True if the actor details has any component to show.
-	bool bHasComponentsToShow = false;
-
-	// True if the actor "root" node in the SCS editor is currently shown as selected
-	bool bShowingRootActorNodeSelected = false;
 };
diff --git a/Engine/Source/Editor/LevelEditor/Private/SLevelEditor.cpp b/Engine/Source/Editor/LevelEditor/Private/SLevelEditor.cpp
--- a/Engine/Source/Editor/LevelEditor/Private/SLevelEditor.cpp
+++ b/Engine/Source/Editor/LevelEditor/Private/SLevelEditor.cpp
@@ -251,7 +251,8 @@ void SLevelEditor::Initialize( const TSharedRef<SDockTab>& OwnerTab, const TShar
 
 	FModuleManager::Get().LoadModuleChecked("StatusBar");
 
-	LevelEditorModule.OnActorSelectionChanged().AddSP(this, &SLevelEditor::OnActorSelectionChanged);
+	LevelEditorModule.OnElementSelectionChanged().AddSP(this, &SLevelEditor::OnElementSelectionChanged);
+	LevelEditorModule.OnOverridePropertyEditorSelection().AddSP(this, &SLevelEditor::OnOverridePropertyEditorSelection);
 
 	TSharedRef<SWidget> ContentArea = RestoreContentArea( OwnerTab, OwnerWindow );
 	TSharedRef<SWidget> MenuBar =
@@ -1807,11 +1808,24 @@ void SLevelEditor::HandleAssetsDeleted(const TArray<UClass*>& DeletedClasses)
 	}
 }
 
-void SLevelEditor::OnActorSelectionChanged(const TArray<UObject*>& NewSelection, bool bForceRefresh)
+void SLevelEditor::OnElementSelectionChanged(const UTypedElementSelectionSet* SelectionSet, bool bForceRefresh)
 {
 	for (TSharedRef<SActorDetails> ActorDetails : GetAllActorDetails())
 	{
-		ActorDetails->SetObjects(NewSelection, bForceRefresh || bNeedsRefresh);
+		if (ActorDetails->IsObservingSelectionSet(SelectionSet))
+		{
+			ActorDetails->RefreshSelection(bForceRefresh || bNeedsRefresh);
+		}
+	}
+
+	bNeedsRefresh = false;
+}
+
+void SLevelEditor::OnOverridePropertyEditorSelection(const TArray<AActor*>& NewSelection, bool bForceRefresh)
+{
+	for (TSharedRef<SActorDetails> ActorDetails : GetAllActorDetails())
+	{
+		ActorDetails->OverrideSelection(NewSelection, bForceRefresh || bNeedsRefresh);
 	}
 
 	bNeedsRefresh = false;
@@ -1893,25 +1907,7 @@ void SLevelEditor::AddStandaloneLevelViewport( const TSharedRef<SLevelViewport>&
 
 TSharedRef<SWidget> SLevelEditor::CreateActorDetails( const FName TabIdentifier )
 {
-	TSharedRef<SActorDetails> ActorDetails = SNew( SActorDetails, TabIdentifier, LevelEditorCommands, GetTabManager() );
-
-	// Immediately update it (otherwise it will appear empty)
-	{
-		TArray<UObject*> SelectedActors;
-		for( FSelectionIterator It( GEditor->GetSelectedActorIterator() ); It; ++It )
-		{
-			AActor* Actor = static_cast<AActor*>( *It );
-			checkSlow( Actor->IsA( AActor::StaticClass() ) );
-
-			if( !Actor->IsPendingKill() )
-			{
-				SelectedActors.Add( Actor );
-			}
-		}
-
-		const bool bForceRefresh = true;
-		ActorDetails->SetObjects( SelectedActors, bForceRefresh );
-	}
+	TSharedRef<SActorDetails> ActorDetails = SNew( SActorDetails, GetMutableElementSelectionSet(), TabIdentifier, LevelEditorCommands, GetTabManager() );
 
 	ActorDetails->SetActorDetailsRootCustomization(ActorDetailsObjectFilter, ActorDetailsRootCustomization);
 	ActorDetails->SetSCSEditorUICustomization(ActorDetailsSCSEditorUICustomization);
diff --git a/Engine/Source/Editor/LevelEditor/Private/SLevelEditor.h b/Engine/Source/Editor/LevelEditor/Private/SLevelEditor.h
--- a/Engine/Source/Editor/LevelEditor/Private/SLevelEditor.h
+++ b/Engine/Source/Editor/LevelEditor/Private/SLevelEditor.h
@@ -229,8 +229,11 @@ private:
 	/** Handles deletion of assets */
 	void HandleAssetsDeleted(const TArray<UClass*>& DeletedClasses);
 
-	/** Called when actors are selected or unselected */
-	void OnActorSelectionChanged(const TArray<UObject*>& NewSelection, bool bForceRefresh = false);
+	/** Called when element selection changes */
+	void OnElementSelectionChanged(const UTypedElementSelectionSet* SelectionSet, bool bForceRefresh = false);
+
+	/** Called to set property editors to show the given actors, even if those actors aren't in the current selection set */
+	void OnOverridePropertyEditorSelection(const TArray<AActor*>& NewSelection, bool bForceRefresh = false);
 
 	/** Called when an actor changes outer */
 	void OnLevelActorOuterChanged(AActor* InActor = nullptr, UObject* InOldOuter = nullptr);
diff --git a/Engine/Source/Editor/LevelEditor/Private/SLevelViewport.cpp b/Engine/Source/Editor/LevelEditor/Private/SLevelViewport.cpp
--- a/Engine/Source/Editor/LevelEditor/Private/SLevelViewport.cpp
+++ b/Engine/Source/Editor/LevelEditor/Private/SLevelViewport.cpp
@@ -1962,23 +1962,18 @@ void SLevelViewport::OnCreateCameraActor(UClass* InClass)
 	pNewCamera->GetCameraComponent()->SetFieldOfView( ViewportClient->ViewFOV );
 
 	// Deselect any currently selected actors
-	GUnrealEd->SelectNone( true, true );
+	GUnrealEd->SelectNone( false, true );
 	GEditor->GetSelectedActors()->DeselectAll();
 	GEditor->GetSelectedObjects()->DeselectAll();
 
 	// Select newly created Camera
-	TArray<UObject *> SelectedActors;
-	GEditor->SelectActor( pNewCamera, true, false );
-	SelectedActors.Add( pNewCamera );
+	GEditor->SelectActor( pNewCamera, true, true );
 
 	// Send notification about actors that may have changed
 	ULevel::LevelDirtiedEvent.Broadcast();
-	
-	// Update the details window with the actors we have just selected
-	GUnrealEd->UpdateFloatingPropertyWindowsFromActorList( SelectedActors );
 
 	// Redraw viewports to show new camera
-	GEditor->RedrawAllViewports();	
+	GEditor->RedrawAllViewports();
 }
 
 bool SLevelViewport::IsPerspectiveViewport() const
diff --git a/Engine/Source/Editor/LevelEditor/Public/LevelEditor.h b/Engine/Source/Editor/LevelEditor/Public/LevelEditor.h
--- a/Engine/Source/Editor/LevelEditor/Public/LevelEditor.h
+++ b/Engine/Source/Editor/LevelEditor/Public/LevelEditor.h
@@ -18,6 +18,7 @@ class IAssetViewport;
 class SLevelEditor;
 class UAnimSequence;
 class USkeletalMeshComponent;
+class UTypedElementSelectionSet;
 
 enum class EMapChangeType : uint8;
 
@@ -204,6 +205,7 @@ public:
 	/** Called when a level editor widget has been created */
 	DECLARE_EVENT_OneParam(FLevelEditorModule, FOnLevelEditorCreated, TSharedPtr<ILevelEditor>);
 	virtual FOnLevelEditorCreated& OnLevelEditorCreated() { return LevelEditorCreatedEvent; }
+	
 	/**
 	 * Called when actor selection changes
 	 * 
@@ -211,6 +213,20 @@ public:
 	 */
 	virtual void BroadcastActorSelectionChanged(const TArray<UObject*>& NewSelection, bool bForceRefresh=false);
 	
+	/**
+	 * Called when element selection changes
+	 *
+	 * @param SelectionSet	Set of selected elements
+	 */
+	virtual void BroadcastElementSelectionChanged(const UTypedElementSelectionSet* SelectionSet, bool bForceRefresh=false);
+
+	/**
+	 * Called to set property editors to show the given actors, even if those actors aren't in the current selection set
+	 * 
+	 * @param NewSelection	List of actors to edit the properties of
+	 */
+	virtual void BroadcastOverridePropertyEditorSelection(const TArray<AActor*>& NewSelection, bool bForceRefresh=false);
+
 	/**
 	 * Called by the engine when level editing viewports need to be redrawn
 	 * 
@@ -238,6 +254,14 @@ public:
 	DECLARE_EVENT_TwoParams(FLevelEditorModule, FActorSelectionChangedEvent, const TArray<UObject*>&, bool);
 	virtual FActorSelectionChangedEvent& OnActorSelectionChanged() { return ActorSelectionChangedEvent; }
 
+	/** Called when element selection changes */
+	DECLARE_EVENT_TwoParams(FLevelEditorModule, FElementSelectionChangedEvent, const UTypedElementSelectionSet*, bool);
+	virtual FElementSelectionChangedEvent& OnElementSelectionChanged() { return ElementSelectionChangedEvent; }
+
+	/** Called to set property editors to show the given actors, even if those actors aren't in the current selection set */
+	DECLARE_EVENT_TwoParams(FLevelEditorModule, FOverridePropertyEditorSelectionEvent, const TArray<AActor*>&, bool);
+	virtual FOverridePropertyEditorSelectionEvent& OnOverridePropertyEditorSelection() { return OverridePropertyEditorSelectionEvent; }
+
 	/** Called when level editor viewports should be redrawn */
 	DECLARE_EVENT_OneParam( FLevelEditorModule, FRedrawLevelEditingViewportsEvent, bool );
 	virtual FRedrawLevelEditingViewportsEvent& OnRedrawLevelEditingViewports() { return RedrawLevelEditingViewportsEvent; }
@@ -411,6 +435,12 @@ private:
 	/** Multicast delegate executed when actor selection changes */
 	FActorSelectionChangedEvent ActorSelectionChangedEvent;
 
+	/** Multicast delegate executed when element selection changes */
+	FElementSelectionChangedEvent ElementSelectionChangedEvent;
+
+	/** Multicast delegate executed to set property editors to show the given actors, even if those actors aren't in the current selection set */
+	FOverridePropertyEditorSelectionEvent OverridePropertyEditorSelectionEvent;
+
 	/** Multicast delegate executed when viewports should be redrawn */
 	FRedrawLevelEditingViewportsEvent RedrawLevelEditingViewportsEvent;
 
diff --git a/Engine/Source/Editor/UnrealEd/Classes/Editor/UnrealEdEngine.h b/Engine/Source/Editor/UnrealEd/Classes/Editor/UnrealEdEngine.h
--- a/Engine/Source/Editor/UnrealEd/Classes/Editor/UnrealEdEngine.h
+++ b/Engine/Source/Editor/UnrealEd/Classes/Editor/UnrealEdEngine.h
@@ -263,12 +263,12 @@ public:
 	virtual void UpdateFloatingPropertyWindows(bool bForceRefresh=false);
 
 	/**
-	*	Updates the property windows of the actors in the supplied ActorList
+	*	Updates the property windows to show the data of the supplied ActorList
 	*
-	*	@param	ActorList	The list of actors whose property windows should be updated
+	*	@param	ActorList	The list of actors to show the properties for
 	*
 	*/
-	virtual void UpdateFloatingPropertyWindowsFromActorList(const TArray< UObject *>& ActorList, bool bForceRefresh=false);
+	virtual void UpdateFloatingPropertyWindowsFromActorList(const TArray<AActor*>& ActorList, bool bForceRefresh=false);
 
 	/**
 	 * Called whenever the actor selection has changed to invalidate any cached state
diff --git a/Engine/Source/Editor/UnrealEd/Private/Dialogs/DlgReferenceTree.cpp b/Engine/Source/Editor/UnrealEd/Private/Dialogs/DlgReferenceTree.cpp
--- a/Engine/Source/Editor/UnrealEd/Private/Dialogs/DlgReferenceTree.cpp
+++ b/Engine/Source/Editor/UnrealEd/Private/Dialogs/DlgReferenceTree.cpp
@@ -447,7 +447,7 @@ TSharedPtr< SWidget > SReferenceTree::BuildMenuWidget()
 				FUIAction SelectActorAction( FExecuteAction::CreateStatic( ReferenceTreeView::Helpers::SelectObjectInEditor, SelectedObject ) );
 				MenuBuilder.AddMenuEntry(NSLOCTEXT("ReferenceTreeView", "SelectActor", "Select Actor"), NSLOCTEXT("ReferenceTreeView", "SelectActor_Tooltip", "Select the actor in the viewport."), FSlateIcon(), SelectActorAction);
 
-				FUIAction ViewPropertiesAction( FExecuteAction::CreateRaw( this, &SReferenceTree::OnMenuViewProperties, SelectedObject ) );
+				FUIAction ViewPropertiesAction( FExecuteAction::CreateRaw( this, &SReferenceTree::OnMenuViewProperties, Actor ) );
 				MenuBuilder.AddMenuEntry(NSLOCTEXT("ReferenceTreeView", "ViewProperties", "View Properties"), NSLOCTEXT("ReferenceTreeView", "ViewProperties_Tooltip", "View the actor's properties."), FSlateIcon(), ViewPropertiesAction);	
 			}
 			else
@@ -525,15 +525,15 @@ void SReferenceTree::DestroyGraphAndTree()
 	ReferenceGraph.Empty();
 }
 
-void SReferenceTree::OnMenuViewProperties( UObject* InObject )
+void SReferenceTree::OnMenuViewProperties( AActor* InActor )
 {
 	// Show the property windows and create one if necessary
 	GUnrealEd->ShowActorProperties();
 
 	// Show the property window for the actor
-	TArray<UObject*> Objects;
-	Objects.Add( InObject );
-	GUnrealEd->UpdateFloatingPropertyWindowsFromActorList( Objects );
+	TArray<AActor*> Actors;
+	Actors.Add(InActor);
+	GUnrealEd->UpdateFloatingPropertyWindowsFromActorList(Actors);
 }
 
 void SReferenceTree::OnMenuShowEditor( UObject* InObject )
diff --git a/Engine/Source/Editor/UnrealEd/Private/EditorHook.cpp b/Engine/Source/Editor/UnrealEd/Private/EditorHook.cpp
--- a/Engine/Source/Editor/UnrealEd/Private/EditorHook.cpp
+++ b/Engine/Source/Editor/UnrealEd/Private/EditorHook.cpp
@@ -1,8 +1,5 @@
 // Copyright Epic Games, Inc. All Rights Reserved.
 
-
-// Includes.
-
 #include "CoreMinimal.h"
 #include "Modules/ModuleManager.h"
 #include "GameFramework/Actor.h"
@@ -11,25 +8,12 @@
 #include "EditorModeManager.h"
 #include "EdMode.h"
 #include "LevelEditor.h"
-
-
-uint32			EngineThreadId;
-const TCHAR*	GItem;
-const TCHAR*	GValue;
-TCHAR*			GCommand;
-
-extern int GLastScroll;
-
-// Misc.
-UEngine* Engine;
-
-/*-----------------------------------------------------------------------------
-	Editor hook exec.
------------------------------------------------------------------------------*/
+#include "Elements/Framework/TypedElementSelectionSet.h"
 
 void UUnrealEdEngine::NotifyPreChange(FProperty* PropertyAboutToChange)
 {
 }
+
 void UUnrealEdEngine::NotifyPostChange(const FPropertyChangedEvent& PropertyChangedEvent, FProperty* PropertyThatChanged)
 {
 	// Notify all active modes of actor property changes.
@@ -38,27 +22,27 @@ void UUnrealEdEngine::NotifyPostChange(const FPropertyChangedEvent& PropertyChan
 
 void UUnrealEdEngine::UpdateFloatingPropertyWindows(bool bForceRefresh)
 {
-	TArray<UObject*> SelectedObjects;
-
-	// Assemble a set of valid selected actors.
-	for (FSelectionIterator It = GetSelectedActorIterator(); It; ++It)
+	if (const UTypedElementSelectionSet* SelectionSet = GetSelectedActors()->GetElementSelectionSet())
 	{
-		AActor* Actor = static_cast<AActor*>(*It);
-		checkSlow(Actor->IsA(AActor::StaticClass()));
+		FLevelEditorModule& LevelEditor = FModuleManager::LoadModuleChecked<FLevelEditorModule>(TEXT("LevelEditor"));
+		LevelEditor.BroadcastElementSelectionChanged(SelectionSet, bForceRefresh);
 
-		if (!Actor->IsPendingKill())
+		// Assemble a set of valid selected actors.
+		TArray<UObject*> SelectedActors;
+		SelectionSet->ForEachSelectedObject<AActor>([&SelectedActors](AActor* InActor)
 		{
-			SelectedObjects.Add(Actor);
-		}
+			if (!InActor->IsPendingKill())
+			{
+				SelectedActors.Add(InActor);
+			}
+			return true;
+		});
+		LevelEditor.BroadcastActorSelectionChanged(SelectedActors, bForceRefresh);
 	}
-
-	UpdateFloatingPropertyWindowsFromActorList(SelectedObjects, bForceRefresh);
 }
 
-
-void UUnrealEdEngine::UpdateFloatingPropertyWindowsFromActorList(const TArray<UObject*>& ActorList, bool bForceRefresh)
+void UUnrealEdEngine::UpdateFloatingPropertyWindowsFromActorList(const TArray<AActor*>& ActorList, bool bForceRefresh)
 {
 	FLevelEditorModule& LevelEditor = FModuleManager::LoadModuleChecked<FLevelEditorModule>(TEXT("LevelEditor"));
-
-	LevelEditor.BroadcastActorSelectionChanged(ActorList, bForceRefresh);
+	LevelEditor.BroadcastOverridePropertyEditorSelection(ActorList, bForceRefresh);
 }
diff --git a/Engine/Source/Editor/UnrealEd/Private/Elements/Actor/ActorElementDetailsInterface.cpp b/Engine/Source/Editor/UnrealEd/Private/Elements/Actor/ActorElementDetailsInterface.cpp
new file mode 100644
--- /dev/null
+++ b/Engine/Source/Editor/UnrealEd/Private/Elements/Actor/ActorElementDetailsInterface.cpp
@@ -0,0 +1,31 @@
+// Copyright Epic Games, Inc. All Rights Reserved.
+
+#include "Elements/Actor/ActorElementDetailsInterface.h"
+#include "Elements/Actor/ActorElementData.h"
+#include "GameFramework/Actor.h"
+
+class FActorTypedElementDetailsObject : public ITypedElementDetailsObject
+{
+public:
+	explicit FActorTypedElementDetailsObject(AActor* InActor)
+		: ActorPtr(InActor)
+	{
+	}
+
+	virtual UObject* GetObject() override
+	{
+		return ActorPtr.Get();
+	}
+
+private:
+	TWeakObjectPtr<AActor> ActorPtr;
+};
+
+TUniquePtr<ITypedElementDetailsObject> UActorElementDetailsInterface::GetDetailsObject(const FTypedElementHandle& InElementHandle)
+{
+	if (AActor* Actor = ActorElementDataUtil::GetActorFromHandle(InElementHandle))
+	{
+		return MakeUnique<FActorTypedElementDetailsObject>(Actor);
+	}
+	return nullptr;
+}
diff --git a/Engine/Source/Editor/UnrealEd/Private/Elements/Actor/ActorElementDetailsInterface.h b/Engine/Source/Editor/UnrealEd/Private/Elements/Actor/ActorElementDetailsInterface.h
new file mode 100644
--- /dev/null
+++ b/Engine/Source/Editor/UnrealEd/Private/Elements/Actor/ActorElementDetailsInterface.h
@@ -0,0 +1,15 @@
+// Copyright Epic Games, Inc. All Rights Reserved.
+
+#pragma once
+
+#include "Elements/Interfaces/TypedElementDetailsInterface.h"
+#include "ActorElementDetailsInterface.generated.h"
+
+UCLASS()
+class UActorElementDetailsInterface : public UTypedElementDetailsInterface
+{
+	GENERATED_BODY()
+
+public:
+	virtual TUniquePtr<ITypedElementDetailsObject> GetDetailsObject(const FTypedElementHandle& InElementHandle) override;
+};
diff --git a/Engine/Source/Editor/UnrealEd/Private/Elements/Component/ComponentElementDetailsInterface.cpp b/Engine/Source/Editor/UnrealEd/Private/Elements/Component/ComponentElementDetailsInterface.cpp
new file mode 100644
--- /dev/null
+++ b/Engine/Source/Editor/UnrealEd/Private/Elements/Component/ComponentElementDetailsInterface.cpp
@@ -0,0 +1,31 @@
+// Copyright Epic Games, Inc. All Rights Reserved.
+
+#include "Elements/Component/ComponentElementDetailsInterface.h"
+#include "Elements/Component/ComponentElementData.h"
+#include "Components/ActorComponent.h"
+
+class FComponentTypedElementDetailsObject : public ITypedElementDetailsObject
+{
+public:
+	explicit FComponentTypedElementDetailsObject(UActorComponent* InComponent)
+		: ComponentPtr(InComponent)
+	{
+	}
+
+	virtual UObject* GetObject() override
+	{
+		return ComponentPtr.Get();
+	}
+
+private:
+	TWeakObjectPtr<UActorComponent> ComponentPtr;
+};
+
+TUniquePtr<ITypedElementDetailsObject> UComponentElementDetailsInterface::GetDetailsObject(const FTypedElementHandle& InElementHandle)
+{
+	if (UActorComponent* Component = ComponentElementDataUtil::GetComponentFromHandle(InElementHandle))
+	{
+		return MakeUnique<FComponentTypedElementDetailsObject>(Component);
+	}
+	return nullptr;
+}
\ No newline at end of file
diff --git a/Engine/Source/Editor/UnrealEd/Private/Elements/Component/ComponentElementDetailsInterface.h b/Engine/Source/Editor/UnrealEd/Private/Elements/Component/ComponentElementDetailsInterface.h
new file mode 100644
--- /dev/null
+++ b/Engine/Source/Editor/UnrealEd/Private/Elements/Component/ComponentElementDetailsInterface.h
@@ -0,0 +1,16 @@
+// Copyright Epic Games, Inc. All Rights Reserved.
+
+#pragma once
+
+#include "Elements/Interfaces/TypedElementDetailsInterface.h"
+#include "ComponentElementDetailsInterface.generated.h"
+
+UCLASS()
+class UComponentElementDetailsInterface : public UTypedElementDetailsInterface
+{
+	GENERATED_BODY()
+
+public:
+	virtual bool IsTopLevelElement(const FTypedElementHandle& InElementHandle) override { return false; }
+	virtual TUniquePtr<ITypedElementDetailsObject> GetDetailsObject(const FTypedElementHandle& InElementHandle) override;
+};
diff --git a/Engine/Source/Editor/UnrealEd/Private/Elements/Framework/EditorElements.cpp b/Engine/Source/Editor/UnrealEd/Private/Elements/Framework/EditorElements.cpp
--- a/Engine/Source/Editor/UnrealEd/Private/Elements/Framework/EditorElements.cpp
+++ b/Engine/Source/Editor/UnrealEd/Private/Elements/Framework/EditorElements.cpp
@@ -3,11 +3,14 @@
 #include "Elements/Framework/EditorElements.h"
 #include "Elements/Framework/TypedElementRegistry.h"
 
+#include "Elements/Object/ObjectElementDetailsInterface.h"
 #include "Elements/Object/ObjectElementEditorSelectionInterface.h"
 
+#include "Elements/Actor/ActorElementDetailsInterface.h"
 #include "Elements/Actor/ActorElementEditorWorldInterface.h"
 #include "Elements/Actor/ActorElementEditorSelectionInterface.h"
 
+#include "Elements/Component/ComponentElementDetailsInterface.h"
 #include "Elements/Component/ComponentElementEditorWorldInterface.h"
 #include "Elements/Component/ComponentElementEditorSelectionInterface.h"
 
@@ -17,6 +20,7 @@ void RegisterEditorObjectElements()
 {
 	UTypedElementRegistry* Registry = UTypedElementRegistry::GetInstance();
 
+	Registry->RegisterElementInterface<UTypedElementDetailsInterface>(NAME_Object, NewObject<UObjectElementDetailsInterface>());
 	Registry->RegisterElementInterface<UTypedElementSelectionInterface>(NAME_Object, NewObject<UObjectElementEditorSelectionInterface>(), /*bAllowOverride*/true);
 }
 
@@ -24,6 +28,7 @@ void RegisterEditorActorElements()
 {
 	UTypedElementRegistry* Registry = UTypedElementRegistry::GetInstance();
 
+	Registry->RegisterElementInterface<UTypedElementDetailsInterface>(NAME_Actor, NewObject<UActorElementDetailsInterface>());
 	Registry->RegisterElementInterface<UTypedElementWorldInterface>(NAME_Actor, NewObject<UActorElementEditorWorldInterface>(), /*bAllowOverride*/true);
 	Registry->RegisterElementInterface<UTypedElementSelectionInterface>(NAME_Actor, NewObject<UActorElementEditorSelectionInterface>(), /*bAllowOverride*/true);
 }
@@ -32,6 +37,7 @@ void RegisterEditorComponentElements()
 {
 	UTypedElementRegistry* Registry = UTypedElementRegistry::GetInstance();
 
+	Registry->RegisterElementInterface<UTypedElementDetailsInterface>(NAME_Components, NewObject<UComponentElementDetailsInterface>());
 	Registry->RegisterElementInterface<UTypedElementWorldInterface>(NAME_Components, NewObject<UComponentElementEditorWorldInterface>(), /*bAllowOverride*/true);
 	Registry->RegisterElementInterface<UTypedElementSelectionInterface>(NAME_Components, NewObject<UComponentElementEditorSelectionInterface>(), /*bAllowOverride*/true);
 }
diff --git a/Engine/Source/Editor/UnrealEd/Private/Elements/Object/ObjectElementDetailsInterface.cpp b/Engine/Source/Editor/UnrealEd/Private/Elements/Object/ObjectElementDetailsInterface.cpp
new file mode 100644
--- /dev/null
+++ b/Engine/Source/Editor/UnrealEd/Private/Elements/Object/ObjectElementDetailsInterface.cpp
@@ -0,0 +1,30 @@
+// Copyright Epic Games, Inc. All Rights Reserved.
+
+#include "Elements/Object/ObjectElementDetailsInterface.h"
+#include "Elements/Object/ObjectElementData.h"
+
+class FObjectTypedElementDetailsObject : public ITypedElementDetailsObject
+{
+public:
+	explicit FObjectTypedElementDetailsObject(UObject* InObject)
+		: ObjectPtr(InObject)
+	{
+	}
+
+	virtual UObject* GetObject() override
+	{
+		return ObjectPtr.Get();
+	}
+
+private:
+	TWeakObjectPtr<UObject> ObjectPtr;
+};
+
+TUniquePtr<ITypedElementDetailsObject> UObjectElementDetailsInterface::GetDetailsObject(const FTypedElementHandle& InElementHandle)
+{
+	if (UObject* Object = ObjectElementDataUtil::GetObjectFromHandle(InElementHandle))
+	{
+		return MakeUnique<FObjectTypedElementDetailsObject>(Object);
+	}
+	return nullptr;
+}
diff --git a/Engine/Source/Editor/UnrealEd/Private/Elements/Object/ObjectElementDetailsInterface.h b/Engine/Source/Editor/UnrealEd/Private/Elements/Object/ObjectElementDetailsInterface.h
new file mode 100644
--- /dev/null
+++ b/Engine/Source/Editor/UnrealEd/Private/Elements/Object/ObjectElementDetailsInterface.h
@@ -0,0 +1,15 @@
+// Copyright Epic Games, Inc. All Rights Reserved.
+
+#pragma once
+
+#include "Elements/Interfaces/TypedElementDetailsInterface.h"
+#include "ObjectElementDetailsInterface.generated.h"
+
+UCLASS()
+class UObjectElementDetailsInterface : public UTypedElementDetailsInterface
+{
+	GENERATED_BODY()
+
+public:
+	virtual TUniquePtr<ITypedElementDetailsObject> GetDetailsObject(const FTypedElementHandle& InElementHandle) override;
+};
diff --git a/Engine/Source/Editor/UnrealEd/Public/Dialogs/DlgReferenceTree.h b/Engine/Source/Editor/UnrealEd/Public/Dialogs/DlgReferenceTree.h
--- a/Engine/Source/Editor/UnrealEd/Public/Dialogs/DlgReferenceTree.h
+++ b/Engine/Source/Editor/UnrealEd/Public/Dialogs/DlgReferenceTree.h
@@ -184,7 +184,7 @@ private:
 	void DestroyGraphAndTree();
 
 	/** Called when the view properties menu option is chosen. */
-	void OnMenuViewProperties( UObject* InObject );
+	void OnMenuViewProperties( AActor* InActor );
 
 	/** Called when the show object in editor menu option is chosen. */
 	void OnMenuShowEditor( UObject* InObject );
