Adding forward declaration to fix non-unity compiles

diff --git a/Engine/Source/Editor/LevelEditor/Private/LightEditorSubsystem.h b/Engine/Source/Editor/LevelEditor/Private/LightEditorSubsystem.h
--- a/Engine/Source/Editor/LevelEditor/Private/LightEditorSubsystem.h
+++ b/Engine/Source/Editor/LevelEditor/Private/LightEditorSubsystem.h
@@ -15,7 +15,7 @@
 #include "Framework/Commands/Commands.h"
 #include "LightEditorSubsystem.generated.h"
 
-
+struct FToolMenuContext;
 /**
  * Light manipulation actions
  */
diff --git a/Engine/Source/Editor/LevelEditor/Public/LevelEditorSubsystem.h b/Engine/Source/Editor/LevelEditor/Public/LevelEditorSubsystem.h
--- a/Engine/Source/Editor/LevelEditor/Public/LevelEditorSubsystem.h
+++ b/Engine/Source/Editor/LevelEditor/Public/LevelEditorSubsystem.h
@@ -7,6 +7,8 @@
 
 #include "LevelEditorSubsystem.generated.h"
 
+
+struct FToolMenuContext;
 /**
 * ULevelEditorSubsystem
 * Subsystem for exposing Level Editor related functionality to scripts
