Adding includes for nonunity compile fixes

diff --git a/Engine/Source/Editor/LevelEditor/Private/LightEditorSubsystem.cpp b/Engine/Source/Editor/LevelEditor/Private/LightEditorSubsystem.cpp
--- a/Engine/Source/Editor/LevelEditor/Private/LightEditorSubsystem.cpp
+++ b/Engine/Source/Editor/LevelEditor/Private/LightEditorSubsystem.cpp
@@ -9,6 +9,8 @@
 #include "Elements/Interfaces/TypedElementObjectInterface.h"
 #include "Elements/Framework/TypedElementSelectionSet.h"
 #include "LevelEditorMenuContext.h"
+#include "ActorFactories/ActorFactory.h"
+#include "Editor.h"
 
 #define LOCTEXT_NAMESPACE "EditorLightSubsystem"
 
diff --git a/Engine/Source/Editor/LevelEditor/Private/SInViewportDetails.h b/Engine/Source/Editor/LevelEditor/Private/SInViewportDetails.h
--- a/Engine/Source/Editor/LevelEditor/Private/SInViewportDetails.h
+++ b/Engine/Source/Editor/LevelEditor/Private/SInViewportDetails.h
@@ -27,6 +27,7 @@ class IDetailTreeNode;
 class IPropertyRowGenerator;
 class ILevelEditor;
 class UToolMenu;
+class SDockingTabStack;
 
 /** A Sample implementation of IDragDropOperation */
 class FInViewportUIDragOperation : public FDragDropOperation
