WorldPartition: Avoid Conversion Prompt on New map

diff --git a/Engine/Source/Editor/LevelEditor/Private/LevelEditorActions.cpp b/Engine/Source/Editor/LevelEditor/Private/LevelEditorActions.cpp
--- a/Engine/Source/Editor/LevelEditor/Private/LevelEditorActions.cpp
+++ b/Engine/Source/Editor/LevelEditor/Private/LevelEditorActions.cpp
@@ -40,6 +40,7 @@
 #include "WorldPartition/WorldPartition.h"
 #include "WorldPartition/WorldPartitionEditorHash.h"
 #include "WorldPartition/WorldPartitionRuntimeHash.h"
+#include "WorldPartition/IWorldPartitionEditorModule.h"
 #include "Editor/WorldBrowser/Public/WorldBrowserModule.h"
 
 #include "Elements/Framework/TypedElementSelectionSet.h"
@@ -100,6 +101,7 @@
 #include "ISettingsModule.h"
 #include "PlatformInfo.h"
 #include "Misc/CoreMisc.h"
+#include "Misc/ScopeExit.h"
 
 #if WITH_LIVE_CODING
 #include "ILiveCodingModule.h"
@@ -288,6 +290,14 @@ void FLevelEditorActionCallbacks::NewLevel()
 
 					if (bCanConvert)
 					{
+						IWorldPartitionEditorModule& WorldPartitionEditorModule = FModuleManager::LoadModuleChecked<IWorldPartitionEditorModule>("WorldPartitionEditor");
+						bool bPreviousConversionPromptEnabled = WorldPartitionEditorModule.IsConversionPromptEnabled();
+						WorldPartitionEditorModule.SetConversionPromptEnabled(false);
+						ON_SCOPE_EXIT
+						{
+							WorldPartitionEditorModule.SetConversionPromptEnabled(bPreviousConversionPromptEnabled);
+						};
+
 						FEditorFileUtils::LoadMap(SavedFilename);
 
 						// Update the world pointer as we reloaded the map
diff --git a/Engine/Source/Editor/WorldPartitionEditor/Private/WorldPartitionEditorModule.cpp b/Engine/Source/Editor/WorldPartitionEditor/Private/WorldPartitionEditorModule.cpp
--- a/Engine/Source/Editor/WorldPartitionEditor/Private/WorldPartitionEditorModule.cpp
+++ b/Engine/Source/Editor/WorldPartitionEditor/Private/WorldPartitionEditorModule.cpp
@@ -131,6 +131,16 @@ bool FWorldPartitionEditorModule::IsWorldPartitionEnabled()
 	return GetDefault<UWorldPartitionEditorSettings>()->bEnableWorldPartition;
 }
 
+bool FWorldPartitionEditorModule::IsConversionPromptEnabled() const
+{
+	return GetDefault<UWorldPartitionEditorSettings>()->bEnableConversionPrompt;
+}
+
+void FWorldPartitionEditorModule::SetConversionPromptEnabled(bool bEnabled)
+{
+	GetMutableDefault<UWorldPartitionEditorSettings>()->bEnableConversionPrompt = bEnabled;
+}
+
 bool FWorldPartitionEditorModule::ConvertMap(const FString& InLongPackageName)
 {
 	if (!GetDefault<UWorldPartitionEditorSettings>()->bEnableConversionPrompt || ULevel::GetIsLevelPartitionedFromPackage(FName(*InLongPackageName)))
diff --git a/Engine/Source/Editor/WorldPartitionEditor/Public/WorldPartitionEditorModule.h b/Engine/Source/Editor/WorldPartitionEditor/Public/WorldPartitionEditorModule.h
--- a/Engine/Source/Editor/WorldPartitionEditor/Public/WorldPartitionEditorModule.h
+++ b/Engine/Source/Editor/WorldPartitionEditor/Public/WorldPartitionEditorModule.h
@@ -29,7 +29,17 @@ public:
 	 *
 	 */
 	virtual bool IsWorldPartitionEnabled() override;
-	
+
+	/**
+	 *
+	 */
+	virtual bool IsConversionPromptEnabled() const override;
+
+	/**
+	 *
+	 */
+	virtual void SetConversionPromptEnabled(bool bEnabled) override;
+
 	/**
 	 * 
 	 */
diff --git a/Engine/Source/Runtime/Engine/Public/WorldPartition/IWorldPartitionEditorModule.h b/Engine/Source/Runtime/Engine/Public/WorldPartition/IWorldPartitionEditorModule.h
--- a/Engine/Source/Runtime/Engine/Public/WorldPartition/IWorldPartitionEditorModule.h
+++ b/Engine/Source/Runtime/Engine/Public/WorldPartition/IWorldPartitionEditorModule.h
@@ -13,5 +13,7 @@ public:
 	virtual ~IWorldPartitionEditorModule() {}
 
 	virtual bool IsWorldPartitionEnabled() = 0;
+	virtual bool IsConversionPromptEnabled() const = 0;
+	virtual void SetConversionPromptEnabled(bool bEnabled) = 0;
 	virtual bool ConvertMap(const FString& InLongPackageName) = 0;
 };
\ No newline at end of file
