Fixed typo in SLevelEditor::CanSpawnWorldPartitionBrowser.

diff --git a/Engine/Source/Editor/LevelEditor/Private/SLevelEditor.cpp b/Engine/Source/Editor/LevelEditor/Private/SLevelEditor.cpp
--- a/Engine/Source/Editor/LevelEditor/Private/SLevelEditor.cpp
+++ b/Engine/Source/Editor/LevelEditor/Private/SLevelEditor.cpp
@@ -984,7 +984,7 @@ bool SLevelEditor::CanSpawnWorldBrowser(const FSpawnTabArgs& Args) const
 	return !UWorld::HasSubsystem<UWorldPartitionSubsystem>(GetWorld());
 }
 
-bool SLevelEditor::SLevelEditor::CanSpawnWorldPartitionBrowser(const FSpawnTabArgs& Args) const
+bool SLevelEditor::CanSpawnWorldPartitionBrowser(const FSpawnTabArgs& Args) const
 {
 	return UWorld::HasSubsystem<UWorldPartitionSubsystem>(GetWorld());
 }
