Allow opening data layers/world partition editors even when the world doesn't support it by making sure they disable their dialogs.
That way, we can have the editors open at editor startup and avoid having to reopen them when needed.

diff --git a/Engine/Source/Editor/DataLayerEditor/Private/DataLayer/SDataLayerBrowser.cpp b/Engine/Source/Editor/DataLayerEditor/Private/DataLayer/SDataLayerBrowser.cpp
--- a/Engine/Source/Editor/DataLayerEditor/Private/DataLayer/SDataLayerBrowser.cpp
+++ b/Engine/Source/Editor/DataLayerEditor/Private/DataLayer/SDataLayerBrowser.cpp
@@ -12,6 +12,7 @@
 #include "DataLayerOutlinerDeleteButtonColumn.h"
 #include "DataLayer/DataLayerEditorSubsystem.h"
 #include "WorldPartition/DataLayer/DataLayer.h"
+#include "WorldPartition/WorldPartitionSubsystem.h"
 #include "ScopedTransaction.h"
 #include "Editor.h"
 
@@ -122,6 +123,7 @@ void SDataLayerBrowser::Construct(const FArguments& InArgs)
 	ChildSlot
 	[
 		SAssignNew(ContentAreaBox, SVerticalBox)
+		.IsEnabled_Lambda([]() { return GWorld ? UWorld::HasSubsystem<UWorldPartitionSubsystem>(GWorld) : false; })
 	];
 
 	SetupDataLayerMode(Mode);
diff --git a/Engine/Source/Editor/Layers/Private/SLayerBrowser.cpp b/Engine/Source/Editor/Layers/Private/SLayerBrowser.cpp
--- a/Engine/Source/Editor/Layers/Private/SLayerBrowser.cpp
+++ b/Engine/Source/Editor/Layers/Private/SLayerBrowser.cpp
@@ -3,6 +3,7 @@
 #include "SLayerBrowser.h"
 #include "Modules/ModuleManager.h"
 #include "Framework/Application/SlateApplication.h"
+#include "WorldPartition/WorldPartitionSubsystem.h"
 #include "Editor.h"
 #include "SLayerStats.h"
 #include "SceneOutlinerModule.h"
@@ -171,6 +172,7 @@ void SLayerBrowser::Construct(const FArguments& InArgs)
 	ChildSlot
 		[
 			SAssignNew(ContentAreaBox, SVerticalBox)
+			.IsEnabled_Lambda([]() { return GWorld ? !UWorld::HasSubsystem<UWorldPartitionSubsystem>(GWorld) : false; })
 		];
 
 	SetupLayersMode();
diff --git a/Engine/Source/Editor/LevelEditor/Private/SLevelEditor.cpp b/Engine/Source/Editor/LevelEditor/Private/SLevelEditor.cpp
--- a/Engine/Source/Editor/LevelEditor/Private/SLevelEditor.cpp
+++ b/Engine/Source/Editor/LevelEditor/Private/SLevelEditor.cpp
@@ -979,26 +979,6 @@ bool SLevelEditor::HasAnyHostedEditorModeToolkit() const
 	return false;
 }
 
-bool SLevelEditor::CanSpawnWorldBrowser(const FSpawnTabArgs& Args) const
-{
-	return !UWorld::HasSubsystem<UWorldPartitionSubsystem>(GetWorld());
-}
-
-bool SLevelEditor::CanSpawnWorldPartitionBrowser(const FSpawnTabArgs& Args) const
-{
-	return UWorld::HasSubsystem<UWorldPartitionSubsystem>(GetWorld());
-}
-
-bool SLevelEditor::CanSpawnLayerBrowser(const FSpawnTabArgs& Args) const
-{
-	return !UWorld::HasSubsystem<UWorldPartitionSubsystem>(GetWorld());
-}
-
-bool SLevelEditor::CanSpawnDataLayerBrowser(const FSpawnTabArgs& Args) const
-{
-	return UWorld::HasSubsystem<UWorldPartitionSubsystem>(GetWorld());
-}
-
 TSharedPtr<SDockTab> SLevelEditor::TryInvokeTab( FName TabID )
 {
 	TSharedPtr<FTabManager> LevelEditorTabManager = GetTabManager();
@@ -1291,8 +1271,7 @@ TSharedRef<SWidget> SLevelEditor::RestoreContentArea( const TSharedRef<SDockTab>
 
 		{
 			const FSlateIcon LayersIcon(FEditorStyle::GetStyleSetName(), "LevelEditor.Tabs.Layers");
-			LevelEditorTabManager->RegisterTabSpawner(LevelEditorTabIds::LevelEditorLayerBrowser, FOnSpawnTab::CreateSP<SLevelEditor, FName, FString>(this, &SLevelEditor::SpawnLevelEditorTab, LevelEditorTabIds::LevelEditorLayerBrowser, FString()),
-																								  FCanSpawnTab::CreateSP(this, &SLevelEditor::CanSpawnLayerBrowser))
+			LevelEditorTabManager->RegisterTabSpawner(LevelEditorTabIds::LevelEditorLayerBrowser, FOnSpawnTab::CreateSP<SLevelEditor, FName, FString>(this, &SLevelEditor::SpawnLevelEditorTab, LevelEditorTabIds::LevelEditorLayerBrowser, FString()) )
 				.SetDisplayName(NSLOCTEXT("LevelEditorTabs", "LevelEditorLayerBrowser", "Layers"))
 				.SetTooltipText(NSLOCTEXT("LevelEditorTabs", "LevelEditorLayerBrowserTooltipText", "Open the Layers tab. Use this to manage which actors in the world belong to which layers."))
 				.SetGroup( MenuStructure.GetLevelEditorCategory() )
@@ -1301,8 +1280,7 @@ TSharedRef<SWidget> SLevelEditor::RestoreContentArea( const TSharedRef<SDockTab>
 
 		{
 			const FSlateIcon DataLayersIcon(FEditorStyle::GetStyleSetName(), "LevelEditor.Tabs.DataLayers");
-			LevelEditorTabManager->RegisterTabSpawner(LevelEditorTabIds::LevelEditorDataLayerBrowser, FOnSpawnTab::CreateSP<SLevelEditor, FName, FString>(this, &SLevelEditor::SpawnLevelEditorTab, LevelEditorTabIds::LevelEditorDataLayerBrowser, FString()),
-																									  FCanSpawnTab::CreateSP(this, &SLevelEditor::CanSpawnDataLayerBrowser))
+			LevelEditorTabManager->RegisterTabSpawner(LevelEditorTabIds::LevelEditorDataLayerBrowser, FOnSpawnTab::CreateSP<SLevelEditor, FName, FString>(this, &SLevelEditor::SpawnLevelEditorTab, LevelEditorTabIds::LevelEditorDataLayerBrowser, FString()) )
 				.SetDisplayName(NSLOCTEXT("LevelEditorTabs", "LevelEditorDataLayerBrowser", "Data Layers"))
 				.SetTooltipText(NSLOCTEXT("LevelEditorTabs", "LevelEditorDataLayerBrowserTooltipText", "Open the Data Layers tab. Use this to manage which actors in the world belong to which Data Layers."))
 				.SetGroup(MenuStructure.GetLevelEditorCategory())
@@ -1319,8 +1297,7 @@ TSharedRef<SWidget> SLevelEditor::RestoreContentArea( const TSharedRef<SDockTab>
 		}
 		
 		{
-			LevelEditorTabManager->RegisterTabSpawner(LevelEditorTabIds::WorldBrowserHierarchy, FOnSpawnTab::CreateSP<SLevelEditor, FName, FString>(this, &SLevelEditor::SpawnLevelEditorTab, LevelEditorTabIds::WorldBrowserHierarchy, FString()),
-																								FCanSpawnTab::CreateSP(this, &SLevelEditor::CanSpawnWorldBrowser))
+			LevelEditorTabManager->RegisterTabSpawner(LevelEditorTabIds::WorldBrowserHierarchy, FOnSpawnTab::CreateSP<SLevelEditor, FName, FString>(this, &SLevelEditor::SpawnLevelEditorTab, LevelEditorTabIds::WorldBrowserHierarchy, FString()) )
 				.SetDisplayName(NSLOCTEXT("LevelEditorTabs", "WorldBrowserHierarchy", "Levels"))
 				.SetTooltipText(NSLOCTEXT("LevelEditorTabs", "WorldBrowserHierarchyTooltipText", "Open the Levels tab. Use this to manage the levels in the current project."))
 				.SetGroup( WorkspaceMenu::GetMenuStructure().GetLevelEditorCategory() )
@@ -1338,8 +1315,7 @@ TSharedRef<SWidget> SLevelEditor::RestoreContentArea( const TSharedRef<SDockTab>
 				.SetGroup( WorkspaceMenu::GetMenuStructure().GetLevelEditorCategory() )
 				.SetIcon( FSlateIcon(FEditorStyle::GetStyleSetName(), "LevelEditor.Tabs.WorldBrowserComposition") );
 
-			LevelEditorTabManager->RegisterTabSpawner( WorldPartitionEditorTab, FOnSpawnTab::CreateSP<SLevelEditor, FName, FString>(this, &SLevelEditor::SpawnLevelEditorTab, LevelEditorTabIds::WorldBrowserPartitionEditor, FString()),
-																				FCanSpawnTab::CreateSP(this, &SLevelEditor::CanSpawnWorldPartitionBrowser))
+			LevelEditorTabManager->RegisterTabSpawner( WorldPartitionEditorTab, FOnSpawnTab::CreateSP<SLevelEditor, FName, FString>(this, &SLevelEditor::SpawnLevelEditorTab, LevelEditorTabIds::WorldBrowserPartitionEditor, FString()) )
 				.SetDisplayName(NSLOCTEXT("LevelEditorTabs", "WorldPartitionEditor", "World Partition"))
 				.SetTooltipText(NSLOCTEXT("LevelEditorTabs", "WorldPartitionEditorTooltipText", "Open World Partition Editor."))
 				.SetGroup( WorkspaceMenu::GetMenuStructure().GetLevelEditorCategory() )
diff --git a/Engine/Source/Editor/LevelEditor/Private/SLevelEditor.h b/Engine/Source/Editor/LevelEditor/Private/SLevelEditor.h
--- a/Engine/Source/Editor/LevelEditor/Private/SLevelEditor.h
+++ b/Engine/Source/Editor/LevelEditor/Private/SLevelEditor.h
@@ -146,10 +146,6 @@ private:
 	TSharedRef<SDockTab> SpawnLevelEditorTab(const FSpawnTabArgs& Args, FName TabIdentifier, FString InitializationPayload);
 	bool CanSpawnEditorModeToolbarTab(const FSpawnTabArgs& Args) const;
 	bool CanSpawnEditorModeToolboxTab(const FSpawnTabArgs& Args) const;
-	bool CanSpawnWorldBrowser(const FSpawnTabArgs& Args) const;
-	bool CanSpawnWorldPartitionBrowser(const FSpawnTabArgs& Args) const;
-	bool CanSpawnLayerBrowser(const FSpawnTabArgs& Args) const;	
-	bool CanSpawnDataLayerBrowser(const FSpawnTabArgs& Args) const;
 	bool HasAnyHostedEditorModeToolkit() const;
 
 	//TSharedRef<SDockTab> SpawnLevelEditorModeTab(const FSpawnTabArgs& Args, FEdMode* EditorMode);
