LevelEditorSubsystem: Expose Request End Play and IsPlayInEditor

diff --git a/Engine/Source/Editor/LevelEditor/Private/LevelEditorSubsystem.cpp b/Engine/Source/Editor/LevelEditor/Private/LevelEditorSubsystem.cpp
--- a/Engine/Source/Editor/LevelEditor/Private/LevelEditorSubsystem.cpp
+++ b/Engine/Source/Editor/LevelEditor/Private/LevelEditorSubsystem.cpp
@@ -151,6 +151,16 @@ void ULevelEditorSubsystem::EditorSetGameView(bool bGameView)
 	}
 }
 
+void ULevelEditorSubsystem::EditorRequestEndPlay()
+{
+	GUnrealEd->RequestEndPlayMap();
+}
+
+bool ULevelEditorSubsystem::IsInPlayInEditor() const
+{
+	return GUnrealEd->IsPlayingSessionInEditor();
+}
+
 /**
  *
  * Editor Scripting | Level
diff --git a/Engine/Source/Editor/LevelEditor/Public/LevelEditorSubsystem.h b/Engine/Source/Editor/LevelEditor/Public/LevelEditorSubsystem.h
--- a/Engine/Source/Editor/LevelEditor/Public/LevelEditorSubsystem.h
+++ b/Engine/Source/Editor/LevelEditor/Public/LevelEditorSubsystem.h
@@ -39,6 +39,12 @@ public:
 	UFUNCTION(BlueprintCallable, Category = "Editor Scripting | Level Utility", meta = (DevelopmentOnly))
 	void EditorSetGameView(bool bGameView);
 
+	UFUNCTION(BlueprintCallable, Category = "Editor Scripting | Level Utility", meta = (DevelopmentOnly))
+	void EditorRequestEndPlay();
+
+	UFUNCTION(BlueprintCallable, Category = "Editor Scripting | Level Utility", meta = (DevelopmentOnly))
+	bool IsInPlayInEditor() const;
+
 	/**
 	 * Close the current Persistent Level (without saving it). Create a new blank Level and save it. Load the new created level.
 	 * @param	AssetPath		Asset Path of where the level will be saved.
