Fixed typo in "Rough(t) Convex Blocking Volume From Mesh" command name.

diff --git a/Engine/Source/Editor/LevelEditor/Private/LevelEditorActions.cpp b/Engine/Source/Editor/LevelEditor/Private/LevelEditorActions.cpp
--- a/Engine/Source/Editor/LevelEditor/Private/LevelEditorActions.cpp
+++ b/Engine/Source/Editor/LevelEditor/Private/LevelEditorActions.cpp
@@ -3412,7 +3412,7 @@ void FLevelEditorCommands::RegisterCommands()
 	UI_COMMAND( CreateHeavyConvexVolume, "Heavy Convex Blocking Volume From Mesh", "Creates a heavy convex blocking volume from the static mesh", EUserInterfaceActionType::Button, FInputChord() );
 	UI_COMMAND( CreateNormalConvexVolume, "Normal Convex Blocking Volume From Mesh", "Creates a normal convex blocking volume from the static mesh", EUserInterfaceActionType::Button, FInputChord() ); 
 	UI_COMMAND( CreateLightConvexVolume, "Light Convex Blocking Volume From Mesh", "Creates a light convex blocking volume from the static mesh", EUserInterfaceActionType::Button, FInputChord() );
-	UI_COMMAND( CreateRoughConvexVolume, "Rought Convex Blocking Volume From Mesh", "Creates a rough convex blocking volume from the static mesh", EUserInterfaceActionType::Button, FInputChord() );
+	UI_COMMAND( CreateRoughConvexVolume, "Rough Convex Blocking Volume From Mesh", "Creates a rough convex blocking volume from the static mesh", EUserInterfaceActionType::Button, FInputChord() );
 
 	UI_COMMAND( KeepSimulationChanges, "Keep Simulation Changes", "Saves the changes made to this actor in Simulate mode to the actor's default state.", EUserInterfaceActionType::Button, FInputChord( EKeys::K ) );
 
