Fix warnings and missing include

diff --git a/Engine/Plugins/Editor/AssetSearch/Source/Private/Widgets/SSearchBrowser.cpp b/Engine/Plugins/Editor/AssetSearch/Source/Private/Widgets/SSearchBrowser.cpp
--- a/Engine/Plugins/Editor/AssetSearch/Source/Private/Widgets/SSearchBrowser.cpp
+++ b/Engine/Plugins/Editor/AssetSearch/Source/Private/Widgets/SSearchBrowser.cpp
@@ -38,6 +38,7 @@
 #include "SearchStyle.h"
 #include "Widgets/SToolTip.h"
 #include "Widgets/Images/SImage.h"
+#include "ThumbnailRendering/ThumbnailManager.h"
 
 #define LOCTEXT_NAMESPACE "SObjectBrowser"
 
diff --git a/Engine/Source/Editor/LevelEditor/Private/SLevelEditor.h b/Engine/Source/Editor/LevelEditor/Private/SLevelEditor.h
--- a/Engine/Source/Editor/LevelEditor/Private/SLevelEditor.h
+++ b/Engine/Source/Editor/LevelEditor/Private/SLevelEditor.h
@@ -72,7 +72,9 @@ public:
 	virtual const TArray< TSharedPtr< class IToolkit > >& GetHostedToolkits() const override;
 	virtual TArray< TSharedPtr< IAssetViewport > > GetViewports() const override;
 	virtual TSharedPtr<IAssetViewport> GetActiveViewportInterface() override;
+	PRAGMA_DISABLE_DEPRECATION_WARNINGS
 	virtual TSharedPtr< class FAssetThumbnailPool > GetThumbnailPool() const override;
+	PRAGMA_ENABLE_DEPRECATION_WARNINGS
 	virtual void AppendCommands( const TSharedRef<FUICommandList>& InCommandsToAppend ) override;
 	virtual void AddStandaloneLevelViewport( const TSharedRef<SLevelViewport>& LevelViewport ) override;
 
