Favorites system disabled by default for details views. Actor details opt-in added.

diff --git a/Engine/Source/Editor/LevelEditor/Private/SActorDetails.cpp b/Engine/Source/Editor/LevelEditor/Private/SActorDetails.cpp
--- a/Engine/Source/Editor/LevelEditor/Private/SActorDetails.cpp
+++ b/Engine/Source/Editor/LevelEditor/Private/SActorDetails.cpp
@@ -99,6 +99,7 @@ void SActorDetails::Construct(const FArguments& InArgs, UTypedElementSelectionSe
 	FDetailsViewArgs DetailsViewArgs;
 	DetailsViewArgs.bUpdatesFromSelection = true;
 	DetailsViewArgs.bLockable = true;
+	DetailsViewArgs.bAllowFavoriteSystem = true;
 	DetailsViewArgs.NameAreaSettings = FDetailsViewArgs::ComponentsAndActorsUseNameArea;
 	DetailsViewArgs.NotifyHook = GUnrealEd;
 	DetailsViewArgs.ViewIdentifier = TabIdentifier;
diff --git a/Engine/Source/Editor/PropertyEditor/Public/IDetailsView.h b/Engine/Source/Editor/PropertyEditor/Public/IDetailsView.h
--- a/Engine/Source/Editor/PropertyEditor/Public/IDetailsView.h
+++ b/Engine/Source/Editor/PropertyEditor/Public/IDetailsView.h
@@ -120,7 +120,7 @@ public:
 		, bShowDifferingPropertiesOption(false)
 		, bCustomNameAreaLocation(false)
 		, bCustomFilterAreaLocation(false)
-		, bAllowFavoriteSystem(true)	
+		, bAllowFavoriteSystem(false)	
 		, bAllowMultipleTopLevelObjects(false)
 		, bShowScrollBar(true)
 		, bForceHiddenPropertyVisibility(false)
