Adding text to the Create/Content/Blueprints/Cinematics toolbar buttons

diff --git a/Engine/Source/Editor/LevelEditor/Private/LevelEditorToolBar.cpp b/Engine/Source/Editor/LevelEditor/Private/LevelEditorToolBar.cpp
--- a/Engine/Source/Editor/LevelEditor/Private/LevelEditorToolBar.cpp
+++ b/Engine/Source/Editor/LevelEditor/Private/LevelEditorToolBar.cpp
@@ -1202,41 +1202,49 @@ void FLevelEditorToolBar::RegisterLevelEditorToolBar( const TSharedRef<FUIComman
 		{
 			FToolMenuSection& Section = AssetsToolBar->AddSection("Content");
 
-			Section.AddEntry(FToolMenuEntry::InitComboButton(
+			FToolMenuEntry CreateEntry = FToolMenuEntry::InitComboButton(
 				"Create",
 				FUIAction(),
 				FOnGetContent::CreateStatic(&FLevelEditorToolBar::GenerateCreateMenuWidget, InCommandList, TWeakPtr<SLevelEditor>(InLevelEditor)),
-				LOCTEXT("OpenCreate_Label", "Content"),
+				LOCTEXT("OpenCreate_Label", "Create"),
 				LOCTEXT("OpenCreate_Tooltip", "Quickly add to the project."),
 				FSlateIcon(FEditorStyle::GetStyleSetName(), "LevelEditor.OpenPlaceActors")
-			));
+			);
+			CreateEntry.StyleNameOverride = "CalloutToolbar";
+			Section.AddEntry(CreateEntry);
 
-			Section.AddEntry(FToolMenuEntry::InitComboButton(
+			FToolMenuEntry ContentEntry = FToolMenuEntry::InitComboButton(
 				"Content",
 				FUIAction(),
 				FOnGetContent::CreateStatic(&FLevelEditorToolBar::GenerateContentMenuWidget, InCommandList, TWeakPtr<SLevelEditor>(InLevelEditor)),
 				LOCTEXT("OpenContent_Label", "Content"),
 				LOCTEXT("OpenContent_Tooltip", "Quickly find content available to add to the project."),
 				FSlateIcon(FEditorStyle::GetStyleSetName(), "LevelEditor.OpenContentBrowser")
-			));
+			);
+			ContentEntry.StyleNameOverride = "CalloutToolbar";
+			Section.AddEntry(ContentEntry);
 
-			Section.AddEntry(FToolMenuEntry::InitComboButton(
+			FToolMenuEntry BlueprintEntry = FToolMenuEntry::InitComboButton(
 				"OpenBlueprint",
 				FUIAction(),
 				FOnGetContent::CreateStatic(&FLevelEditorToolBar::GenerateOpenBlueprintMenuContent, InCommandList, TWeakPtr<SLevelEditor>(InLevelEditor)),
 				LOCTEXT("OpenBlueprint_Label", "Blueprints"),
 				LOCTEXT("OpenBlueprint_ToolTip", "List of world Blueprints available to the user for editing or creation."),
 				FSlateIcon(FEditorStyle::GetStyleSetName(), "LevelEditor.OpenLevelBlueprint")
-			));
+			);
+			BlueprintEntry.StyleNameOverride = "CalloutToolbar";
+			Section.AddEntry(BlueprintEntry);
 
-			Section.AddEntry(FToolMenuEntry::InitComboButton(
+			FToolMenuEntry CinematicsEntry = FToolMenuEntry::InitComboButton(
 				"EditCinematics",
 				FUIAction(),
 				FOnGetContent::CreateStatic(&FLevelEditorToolBar::GenerateCinematicsMenuContent, InCommandList, TWeakPtr<SLevelEditor>(InLevelEditor)),
 				LOCTEXT("EditCinematics_Label", "Cinematics"),
 				LOCTEXT("EditCinematics_Tooltip", "Displays a list of Level Sequence objects to open in their respective editors"),
 				FSlateIcon(FEditorStyle::GetStyleSetName(), "LevelEditor.EditMatinee")
-			));
+			);
+			CinematicsEntry.StyleNameOverride = "CalloutToolbar";
+			Section.AddEntry(CinematicsEntry);
 
 			Section.AddEntry(FToolMenuEntry::InitToolBarButton(FLevelEditorCommands::Get().ToggleVR, LOCTEXT("ToggleVR", "VR Mode")));
 		}
