Modes can be toggled in the level editor toolbox at runtime.

diff --git a/Engine/Source/Editor/AssetPlacementEdMode/Private/AssetPlacementEdMode.cpp b/Engine/Source/Editor/AssetPlacementEdMode/Private/AssetPlacementEdMode.cpp
--- a/Engine/Source/Editor/AssetPlacementEdMode/Private/AssetPlacementEdMode.cpp
+++ b/Engine/Source/Editor/AssetPlacementEdMode/Private/AssetPlacementEdMode.cpp
@@ -28,10 +28,11 @@ constexpr TCHAR UAssetPlacementEdMode::AssetPlacementEdModeID[];
 
 UAssetPlacementEdMode::UAssetPlacementEdMode()
 {
+	TAttribute<bool> IsEnabledAttr = TAttribute<bool>::Create(TAttribute<bool>::FGetter::CreateLambda([]() -> bool { return GetDefault<ULevelEditorMiscSettings>()->bEnableAssetPlacementMode; }));
 	Info = FEditorModeInfo(UAssetPlacementEdMode::AssetPlacementEdModeID,
 		LOCTEXT("AssetPlacementEdModeName", "Placement"),
 		FSlateIcon(FAssetPlacementEdModeStyle::GetStyleSetName(), "PlacementBrowser.ShowAllContent"),
-		GetDefault<ULevelEditorMiscSettings>()->bEnableAssetPlacementMode);
+		MoveTemp(IsEnabledAttr));
 
 	SettingsClass = UAssetPlacementSettings::StaticClass();
 }
@@ -134,6 +135,11 @@ bool UAssetPlacementEdMode::ShouldDrawWidget() const
 	return Super::ShouldDrawWidget();
 }
 
+bool UAssetPlacementEdMode::IsEnabled()
+{
+	return GetDefault<ULevelEditorMiscSettings>()->bEnableAssetPlacementMode;
+}
+
 bool UAssetPlacementEdMode::DoesPaletteSupportElement(const FTypedElementHandle& InElementToCheck, const TArray<FPaletteItem>& InPaletteToCheck)
 {
 	if (!InElementToCheck.IsSet())
diff --git a/Engine/Source/Editor/AssetPlacementEdMode/Public/AssetPlacementEdMode.h b/Engine/Source/Editor/AssetPlacementEdMode/Public/AssetPlacementEdMode.h
--- a/Engine/Source/Editor/AssetPlacementEdMode/Public/AssetPlacementEdMode.h
+++ b/Engine/Source/Editor/AssetPlacementEdMode/Public/AssetPlacementEdMode.h
@@ -37,6 +37,7 @@ public:
 	// End of UBaseLegacyWidgetEdMode interface
 	//////////////////
 
+	static bool IsEnabled();
 	static bool DoesPaletteSupportElement(const FTypedElementHandle& InElementToCheck, const TArray<FPaletteItem>& InPaletteToCheck);
 
 protected:
diff --git a/Engine/Source/Editor/EditorFramework/Private/Tools/Modes.cpp b/Engine/Source/Editor/EditorFramework/Private/Tools/Modes.cpp
--- a/Engine/Source/Editor/EditorFramework/Private/Tools/Modes.cpp
+++ b/Engine/Source/Editor/EditorFramework/Private/Tools/Modes.cpp
@@ -5,8 +5,8 @@
 
 FEditorModeInfo::FEditorModeInfo()
 	: ID(NAME_None)
-	, bVisible(false)
 	, PriorityOrder(MAX_int32)
+	, Visibility(false)
 {
 }
 
@@ -14,18 +14,23 @@ FEditorModeInfo::FEditorModeInfo(
 	FEditorModeID InID,
 	FText InName,
 	FSlateIcon InIconBrush,
-	bool InIsVisible,
+	TAttribute<bool> InVisibility,
 	int32 InPriorityOrder
 	)
 	: ID(InID)
 	, ToolbarCustomizationName(*(InID.ToString() + TEXT("Toolbar")))
 	, Name(InName)
 	, IconBrush(InIconBrush)
-	, bVisible(InIsVisible)
 	, PriorityOrder(InPriorityOrder)
+	, Visibility(InVisibility)
 {
 	if (!InIconBrush.IsSet())
 	{
 		IconBrush = FSlateIcon(FEditorStyle::GetStyleSetName(), "LevelEditor.EditorModes");
 	}
+}
+
+bool FEditorModeInfo::IsVisible() const
+{
+	return Visibility.Get(false);
 }
\ No newline at end of file
diff --git a/Engine/Source/Editor/EditorFramework/Public/Tools/Modes.h b/Engine/Source/Editor/EditorFramework/Public/Tools/Modes.h
--- a/Engine/Source/Editor/EditorFramework/Public/Tools/Modes.h
+++ b/Engine/Source/Editor/EditorFramework/Public/Tools/Modes.h
@@ -6,6 +6,7 @@
 #include "Internationalization/Text.h"
 #include "Textures/SlateIcon.h"
 #include "UObject/NameTypes.h"
+#include "Misc/Attribute.h"
 
 /** The shorthand identifier used for editor modes */
 typedef FName FEditorModeID;
@@ -20,7 +21,7 @@ struct FEditorModeInfo
 		FEditorModeID InID,
 		FText InName = FText(),
 		FSlateIcon InIconBrush = FSlateIcon(),
-		bool bInIsVisible = false,
+		TAttribute<bool> InVisibility = false,
 		int32 InPriorityOrder = MAX_int32
 		);
 
@@ -36,11 +37,14 @@ struct FEditorModeInfo
 	/** The mode icon */
 	FSlateIcon IconBrush;
 
-	/** Whether or not the mode should be visible in the mode menu */
-	bool bVisible;
-
 	/** The priority of this mode which will determine its default order and shift+X command assignment */
 	int32 PriorityOrder;
+
+	/** Whether or not the mode should be visible in the mode toolbar */
+	EDITORFRAMEWORK_API bool IsVisible() const;
+
+protected:
+	TAttribute<bool> Visibility;
 };
 
 DECLARE_MULTICAST_DELEGATE(FRegisteredModesChangedEvent);
diff --git a/Engine/Source/Editor/LevelEditor/Private/LevelEditorModesActions.cpp b/Engine/Source/Editor/LevelEditor/Private/LevelEditorModesActions.cpp
--- a/Engine/Source/Editor/LevelEditor/Private/LevelEditorModesActions.cpp
+++ b/Engine/Source/Editor/LevelEditor/Private/LevelEditorModesActions.cpp
@@ -15,16 +15,10 @@ void FLevelEditorModesCommands::RegisterCommands()
 	EditorModeCommands.Empty();
 
 	int editorMode = 0;
-	FKey EdModeKeys[9] = { EKeys::One, EKeys::Two, EKeys::Three, EKeys::Four, EKeys::Five, EKeys::Six, EKeys::Seven, EKeys::Eight, EKeys::Nine };
+	static const TArray<FKey, TInlineAllocator<9>> EdModeKeys = { EKeys::One, EKeys::Two, EKeys::Three, EKeys::Four, EKeys::Five, EKeys::Six, EKeys::Seven, EKeys::Eight, EKeys::Nine };
 
 	for ( const FEditorModeInfo& Mode : GEditor->GetEditorSubsystem<UAssetEditorSubsystem>()->GetEditorModeInfoOrderedByPriority())
 	{
-		// If the mode isn't visible don't create a menu option for it.
-		if (!Mode.bVisible)
-		{
-			continue;
-		}
-
 		FName EditorModeCommandName = FName(*(FString("EditorMode.") + Mode.ID.ToString()));
 
 		TSharedPtr<FUICommandInfo> EditorModeCommand = 
@@ -34,23 +28,33 @@ void FLevelEditorModesCommands::RegisterCommands()
 		if ( !EditorModeCommand.IsValid() )
 		{
 			FFormatNamedArguments Args;
-			Args.Add( TEXT("Mode"), Mode.Name );
+			FText ModeName = Mode.Name;
+			if (ModeName.IsEmpty())
+			{
+				ModeName = FText::FromName(Mode.ID);
+			}
+			Args.Add(TEXT("Mode"), ModeName);
 			const FText Tooltip = FText::Format( NSLOCTEXT("LevelEditor", "ModeTooltipF", "Activate {Mode} Editing Mode"), Args );
 
+			FInputChord DefaultKeyBinding;
+			if (Mode.IsVisible() && editorMode < EdModeKeys.Num())
+			{
+				DefaultKeyBinding = FInputChord(EModifierKey::Shift, EdModeKeys[editorMode]);
+				++editorMode;
+			}
+
 			FUICommandInfo::MakeCommandInfo(
 				this->AsShared(),
 				EditorModeCommand,
 				EditorModeCommandName,
-				Mode.Name,
+				ModeName,
 				Tooltip,
 				Mode.IconBrush,
 				EUserInterfaceActionType::ToggleButton,
-				editorMode < 9 ? FInputChord( EModifierKey::Shift, EdModeKeys[editorMode] ) : FInputChord() );
+				DefaultKeyBinding);
 
 			EditorModeCommands.Add(EditorModeCommand);
 		}
-
-		editorMode++;
 	}
 }
 
diff --git a/Engine/Source/Editor/LevelEditor/Private/LevelEditorToolBar.cpp b/Engine/Source/Editor/LevelEditor/Private/LevelEditorToolBar.cpp
--- a/Engine/Source/Editor/LevelEditor/Private/LevelEditorToolBar.cpp
+++ b/Engine/Source/Editor/LevelEditor/Private/LevelEditorToolBar.cpp
@@ -1254,62 +1254,59 @@ void FLevelEditorToolBar::RegisterLevelEditorToolBar( const TSharedRef<FUIComman
 	UToolMenu* ModesToolbar = UToolMenus::Get()->RegisterMenu("LevelEditor.LevelEditorToolBar.ModesToolBar", NAME_None, EMultiBoxType::SlimHorizontalToolBar);
 	ModesToolbar->StyleName = "AssetEditorToolbar";
 	{
-		FToolMenuSection& Section = ModesToolbar->AddSection("EditorModes");
-
-		Section.AddSeparator(NAME_None);
-
-		const FLevelEditorModule& LevelEditorModule = FModuleManager::LoadModuleChecked<FLevelEditorModule>("LevelEditor");
-		const FLevelEditorModesCommands& Commands = LevelEditorModule.GetLevelEditorModesCommands();
+		TWeakPtr<SLevelEditor> LevelEditorPtr = InLevelEditor;
+		ModesToolbar->AddDynamicSection("EditorModes", FNewToolMenuDelegate::CreateLambda([LevelEditorPtr](UToolMenu* ToolMenu)
+		{
+			FToolMenuSection& Section = ToolMenu->AddSection("EditorModes");
+			Section.AddSeparator(NAME_None);
+			const FLevelEditorModule& LevelEditorModule = FModuleManager::LoadModuleChecked<FLevelEditorModule>("LevelEditor");
+			const FLevelEditorModesCommands& Commands = LevelEditorModule.GetLevelEditorModesCommands();
 
-		TArray<FEditorModeInfo, TInlineAllocator<1>> DefaultModes;
+			TArray<FEditorModeInfo, TInlineAllocator<1>> DefaultModes;
 
-		TArray<FEditorModeInfo, TInlineAllocator<10>> NonDefaultModes;
+			TArray<FEditorModeInfo, TInlineAllocator<10>> NonDefaultModes;
 
-		for (const FEditorModeInfo& Mode : GEditor->GetEditorSubsystem<UAssetEditorSubsystem>()->GetEditorModeInfoOrderedByPriority())
-		{
-			// If the mode isn't visible don't create a menu option for it.
-			if (!Mode.bVisible)
+			for (const FEditorModeInfo& Mode : GEditor->GetEditorSubsystem<UAssetEditorSubsystem>()->GetEditorModeInfoOrderedByPriority())
 			{
-				continue;
-			}
+				TSharedPtr<SLevelEditor> LevelEditorPin = LevelEditorPtr.Pin();
+				if (LevelEditorPin.IsValid() && LevelEditorPin->GetEditorModeManager().IsDefaultMode(Mode.ID))
+				{
+					DefaultModes.Add(Mode.ID);
+				}
+				else
+				{
+					NonDefaultModes.Add(Mode.ID);
+				}
 
-			if (GLevelEditorModeTools().IsDefaultMode(Mode.ID))
-			{
-				DefaultModes.Add(Mode.ID);
-			}
-			else
-			{
-				NonDefaultModes.Add(Mode.ID);
 			}
 
-		}
-
-		auto BuildEditorModes =
-			[&Commands, &Section](TArrayView<FEditorModeInfo> Modes)
-		{
-			for (const FEditorModeInfo& Mode : Modes)
+			auto BuildEditorModes =
+				[&Commands, &Section](TArrayView<FEditorModeInfo> Modes)
 			{
-				FName EditorModeCommandName = FName(*(FString("EditorMode.") + Mode.ID.ToString()));
+				for (const FEditorModeInfo& Mode : Modes)
+				{
+					FName EditorModeCommandName = FName(*(FString("EditorMode.") + Mode.ID.ToString()));
 
-				TSharedPtr<FUICommandInfo> EditorModeCommand =
-					FInputBindingManager::Get().FindCommandInContext(Commands.GetContextName(), EditorModeCommandName);
+					TSharedPtr<FUICommandInfo> EditorModeCommand =
+						FInputBindingManager::Get().FindCommandInContext(Commands.GetContextName(), EditorModeCommandName);
 
-				// If a command isn't yet registered for this mode, we need to register one.
-				if (!EditorModeCommand.IsValid())
-				{
-					continue;
-				}
+					// If a command isn't yet registered for this mode, we need to register one.
+					if (!EditorModeCommand.IsValid())
+					{
+						continue;
+					}
 
-				Section.AddEntry(FToolMenuEntry::InitToolBarButton(EditorModeCommand));
-			}
+					Section.AddEntry(FToolMenuEntry::InitToolBarButton(EditorModeCommand));
+				}
 
-		};
+			};
 
-		// Build Default Modes first
-		BuildEditorModes(MakeArrayView(DefaultModes));
+			// Build Default Modes first
+			BuildEditorModes(MakeArrayView(DefaultModes));
 
-		// Build non-default modes second
-		BuildEditorModes(MakeArrayView(NonDefaultModes));
+			// Build non-default modes second
+			BuildEditorModes(MakeArrayView(NonDefaultModes));
+		}));
 	}
 
 	UToolMenu* SettingsToolBar = UToolMenus::Get()->RegisterMenu("LevelEditor.LevelEditorToolBar.MarketplaceToolBar", NAME_None, EMultiBoxType::SlimHorizontalToolBar);
diff --git a/Engine/Source/Editor/LevelEditor/Private/SLevelEditor.cpp b/Engine/Source/Editor/LevelEditor/Private/SLevelEditor.cpp
--- a/Engine/Source/Editor/LevelEditor/Private/SLevelEditor.cpp
+++ b/Engine/Source/Editor/LevelEditor/Private/SLevelEditor.cpp
@@ -188,9 +188,12 @@ void SLevelEditor::Construct( const SLevelEditor::FArguments& InArgs)
 	// We need to register when modes list changes so that we can refresh the auto generated commands.
 	if (GEditor != nullptr)
 	{
-		GEditor->GetEditorSubsystem<UAssetEditorSubsystem>()->OnEditorModesChanged().AddRaw(this, &SLevelEditor::EditorModeCommandsChanged);
+		if (UAssetEditorSubsystem* AssetEditorSubsystem = GEditor->GetEditorSubsystem<UAssetEditorSubsystem>())
+		{
+			AssetEditorSubsystem->OnEditorModesChanged().AddRaw(this, &SLevelEditor::EditorModeCommandsChanged);
+		}
 	}
-	GLevelEditorModeTools().OnEditorModeIDChanged().AddSP(this, &SLevelEditor::OnEditorModeIdChanged);
+	GetEditorModeManager().OnEditorModeIDChanged().AddSP(this, &SLevelEditor::OnEditorModeIdChanged);
 
 	// @todo This is a hack to get this working for now. This won't work with multiple worlds
 	if (GEditor != nullptr)
@@ -339,11 +342,12 @@ SLevelEditor::~SLevelEditor()
 
 	if (GEditor)
 	{
-		GEditor->GetEditorSubsystem<UAssetEditorSubsystem>()->OnEditorModesChanged().RemoveAll(this);
+		if (UAssetEditorSubsystem* AssetEditorSubsystem = GEditor->GetEditorSubsystem<UAssetEditorSubsystem>())
+		{
+			AssetEditorSubsystem->OnEditorModesChanged().RemoveAll(this);
+		}
 		GEditor->OnLevelActorOuterChanged().Remove(LevelActorOuterChangedHandle);
 		GEditor->GetEditorWorldContext(true).RemoveRef(World);
-
-		GEditor->GetEditorSubsystem<UAssetEditorSubsystem>()->OnEditorModesChanged().RemoveAll(this);
 	}
 
 	// Clear USelection from using our selected element list
@@ -688,7 +692,7 @@ TSharedRef<SDockTab> SLevelEditor::SpawnLevelEditorTab( const FSpawnTabArgs& Arg
 	}
 	else if (TabIdentifier == FEditorModeTools::EditorModeToolbarTabName)
 	{
-		return GLevelEditorModeTools().MakeModeToolbarTab();
+		return GetEditorModeManager().MakeModeToolbarTab();
 	}
 	else if( TabIdentifier == TEXT("LevelEditorSelectionDetails") || TabIdentifier == TEXT("LevelEditorSelectionDetails2") || TabIdentifier == TEXT("LevelEditorSelectionDetails3") || TabIdentifier == TEXT("LevelEditorSelectionDetails4") )
 	{
@@ -957,7 +961,7 @@ TSharedRef<SDockTab> SLevelEditor::SpawnLevelEditorTab( const FSpawnTabArgs& Arg
 
 bool SLevelEditor::CanSpawnEditorModeToolbarTab(const FSpawnTabArgs& Args) const
 {
-	return GLevelEditorModeTools().ShouldShowModeToolbar();
+	return GetEditorModeManager().ShouldShowModeToolbar();
 }
 
 bool SLevelEditor::CanSpawnEditorModeToolboxTab(const FSpawnTabArgs& Args) const
@@ -1079,7 +1083,7 @@ void SLevelEditor::OnViewportTabClosed(TSharedRef<SDockTab> ClosedTab)
 
 void SLevelEditor::OnToolboxTabClosed(TSharedRef<SDockTab> ClosedTab)
 {
-	GLevelEditorModeTools().ActivateDefaultMode();
+	GetEditorModeManager().ActivateDefaultMode();
 }
 
 
@@ -1518,15 +1522,13 @@ void SLevelEditor::ToggleEditorMode( FEditorModeID ModeID )
 	// Prompt the user if Matinee must be closed before activating new mode
 	if (ModeID != FBuiltinEditorModes::EM_InterpEdit)
 	{
-		FEdMode* MatineeMode = GLevelEditorModeTools().GetActiveMode(FBuiltinEditorModes::EM_InterpEdit);
+		FEdMode* MatineeMode = GetEditorModeManager().GetActiveMode(FBuiltinEditorModes::EM_InterpEdit);
 		if (MatineeMode && !MatineeMode->IsCompatibleWith(ModeID))
 		{
-			FEditorModeInfo MatineeModeInfo;
-			GEditor->GetEditorSubsystem<UAssetEditorSubsystem>()->FindEditorModeInfo(ModeID, MatineeModeInfo);
 			FFormatNamedArguments Args;
-			Args.Add(TEXT("ModeName"), MatineeModeInfo.Name);
+			Args.Add(TEXT("ModeName"), MatineeMode->GetModeInfo().Name);
 			FText Msg = FText::Format(NSLOCTEXT("LevelEditor", "ModeSwitchCloseMatineeQ", "Activating '{ModeName}' editor mode will close UnrealMatinee.  Continue?"), Args);
-			
+
 			if (EAppReturnType::Yes != FMessageDialog::Open(EAppMsgType::YesNo, Msg))
 			{
 				return;
@@ -1536,7 +1538,7 @@ void SLevelEditor::ToggleEditorMode( FEditorModeID ModeID )
 		
 	// *Important* - activate the mode first since FEditorModeTools::DeactivateMode will
 	// activate the default mode when the stack becomes empty, resulting in multiple active visible modes.
-	GLevelEditorModeTools().ActivateMode( ModeID );
+	GetEditorModeManager().ActivateMode( ModeID );
 
 	// Invoke the tab regardless if if the mode was activated or already activated so that it the mode tools come into the foreground
 	LevelEditorTabManager->TryInvokeTab(LevelEditorTabIds::LevelEditorToolBox);
@@ -1545,7 +1547,21 @@ void SLevelEditor::ToggleEditorMode( FEditorModeID ModeID )
 
 bool SLevelEditor::IsModeActive( FEditorModeID ModeID )
 {
-	return GLevelEditorModeTools().IsModeActive( ModeID );
+	return GetEditorModeManager().IsModeActive( ModeID );
+}
+
+bool SLevelEditor::ShouldShowModeInToolbar(FEditorModeID ModeID)
+{
+	if (UAssetEditorSubsystem* AssetEditorSubsystem = GEditor->GetEditorSubsystem<UAssetEditorSubsystem>())
+	{
+		FEditorModeInfo ModeInfo;
+		if (AssetEditorSubsystem->FindEditorModeInfo(ModeID, ModeInfo))
+		{
+			return ModeInfo.IsVisible();
+		}
+	}
+
+	return false;
 }
 
 void SLevelEditor::EditorModeCommandsChanged()
@@ -1595,31 +1611,29 @@ void SLevelEditor::RefreshEditorModeCommands()
 	const FLevelEditorModesCommands& Commands = FLevelEditorModesCommands::Get();
 
 	int32 CommandIndex = 0;
-	for( const FEditorModeInfo& Mode : GEditor->GetEditorSubsystem<UAssetEditorSubsystem>()->GetEditorModeInfoOrderedByPriority() )
+	if (UAssetEditorSubsystem* AssetEditorSubsystem = GEditor->GetEditorSubsystem<UAssetEditorSubsystem>())
 	{
-		// If the mode isn't visible don't create a menu option for it.
-		if( !Mode.bVisible )
+		for (const FEditorModeInfo& Mode : AssetEditorSubsystem->GetEditorModeInfoOrderedByPriority())
 		{
-			continue;
-		}
+			FName EditorModeTabName = GetEditorModeTabId(Mode.ID);
+			FName EditorModeCommandName = FName(*(FString("EditorMode.") + Mode.ID.ToString()));
 
-		FName EditorModeTabName = GetEditorModeTabId( Mode.ID );
-		FName EditorModeCommandName = FName(*(FString("EditorMode.") + Mode.ID.ToString()));
+			TSharedPtr<FUICommandInfo> EditorModeCommand =
+				FInputBindingManager::Get().FindCommandInContext(Commands.GetContextName(), EditorModeCommandName);
 
-		TSharedPtr<FUICommandInfo> EditorModeCommand = 
-			FInputBindingManager::Get().FindCommandInContext(Commands.GetContextName(), EditorModeCommandName);
+			// If a command isn't yet registered for this mode, we need to register one.
+			if (EditorModeCommand.IsValid() && !LevelEditorCommands->IsActionMapped(Commands.EditorModeCommands[CommandIndex]))
+			{
+				LevelEditorCommands->MapAction(
+					Commands.EditorModeCommands[CommandIndex],
+					FExecuteAction::CreateSP(SharedThis(this), &SLevelEditor::ToggleEditorMode, Mode.ID),
+					FCanExecuteAction(),
+					FIsActionChecked::CreateSP(SharedThis(this), &SLevelEditor::IsModeActive, Mode.ID),
+					FIsActionButtonVisible::CreateSP(SharedThis(this), &SLevelEditor::ShouldShowModeInToolbar, Mode.ID));
+			}
 
-		// If a command isn't yet registered for this mode, we need to register one.
-		if ( EditorModeCommand.IsValid() && !LevelEditorCommands->IsActionMapped(Commands.EditorModeCommands[CommandIndex]) )
-		{
-			LevelEditorCommands->MapAction(
-				Commands.EditorModeCommands[CommandIndex],
-				FExecuteAction::CreateStatic( &SLevelEditor::ToggleEditorMode, Mode.ID ),
-				FCanExecuteAction(),
-				FIsActionChecked::CreateSP( GLevelEditorModeTools().AsShared(), &FEditorModeTools::IsModeActive, Mode.ID ));
+			CommandIndex++;
 		}
-
-		CommandIndex++;
 	}
 
 	for( const auto& ToolBoxTab : ToolBoxTabs )
diff --git a/Engine/Source/Editor/LevelEditor/Private/SLevelEditor.h b/Engine/Source/Editor/LevelEditor/Private/SLevelEditor.h
--- a/Engine/Source/Editor/LevelEditor/Private/SLevelEditor.h
+++ b/Engine/Source/Editor/LevelEditor/Private/SLevelEditor.h
@@ -183,10 +183,13 @@ private:
 	static FName GetEditorModeTabId( FEditorModeID ModeID );
 
 	/** Toggles the editor mode on and off, this is what the auto generated editor mode commands are mapped to. */
-	static void ToggleEditorMode( FEditorModeID ModeID );
+	void ToggleEditorMode( FEditorModeID ModeID );
 
 	/** Checks if the editor mode is active for the auto-generated editor mode command. */
-	static bool IsModeActive( FEditorModeID ModeID );
+	bool IsModeActive(FEditorModeID ModeID);
+
+	/** Checks if the editor mode is visible for the auto-generated editor mode command. */
+	bool ShouldShowModeInToolbar(FEditorModeID ModeID);
 
 	/**
 	 * Processes keybindings on the level editor
diff --git a/Engine/Source/Editor/LevelEditor/Private/SLevelEditorToolBox.cpp b/Engine/Source/Editor/LevelEditor/Private/SLevelEditorToolBox.cpp
--- a/Engine/Source/Editor/LevelEditor/Private/SLevelEditorToolBox.cpp
+++ b/Engine/Source/Editor/LevelEditor/Private/SLevelEditorToolBox.cpp
@@ -117,7 +117,7 @@ void SLevelEditorToolBox::UpdateModeLegacyToolBar()
 		for ( const FEditorModeInfo& Mode : GEditor->GetEditorSubsystem<UAssetEditorSubsystem>()->GetEditorModeInfoOrderedByPriority())
 		{
 			// If the mode isn't visible don't create a menu option for it.
-			if ( !Mode.bVisible )
+			if ( !Mode.IsVisible() )
 			{
 				continue;
 			}
diff --git a/Engine/Source/Editor/UnrealEd/Classes/Settings/LevelEditorMiscSettings.h b/Engine/Source/Editor/UnrealEd/Classes/Settings/LevelEditorMiscSettings.h
--- a/Engine/Source/Editor/UnrealEd/Classes/Settings/LevelEditorMiscSettings.h
+++ b/Engine/Source/Editor/UnrealEd/Classes/Settings/LevelEditorMiscSettings.h
@@ -48,7 +48,7 @@ public:
 	bool bEnableExperimentalLevelEditor = false;
 
 	/** If enabled, the experimental asset placement mode will be available in the level editor */
-	UPROPERTY(EditAnywhere, config, AdvancedDisplay, Category = Editing, meta = (ConfigRestartRequired = true))
+	UPROPERTY(EditAnywhere, config, AdvancedDisplay, Category = Editing)
 	bool bEnableAssetPlacementMode = false;
 
 	/** If enabled, will avoid relabeling actors in UUnrealEdEngine::edactPasteSelected */
diff --git a/Engine/Source/Editor/UnrealEd/Private/EditorModeManager.cpp b/Engine/Source/Editor/UnrealEd/Private/EditorModeManager.cpp
--- a/Engine/Source/Editor/UnrealEd/Private/EditorModeManager.cpp
+++ b/Engine/Source/Editor/UnrealEd/Private/EditorModeManager.cpp
@@ -352,7 +352,7 @@ void FEditorModeTools::DeactivateOtherVisibleModes(FEditorModeID InMode)
 	TArray<UEdMode*> TempModes(ActiveScriptableModes);
 	for (const UEdMode* Mode : TempModes)
 	{
-		if (Mode->GetID() != InMode && Mode->GetModeInfo().bVisible)
+		if (Mode->GetID() != InMode && Mode->GetModeInfo().IsVisible())
 		{
 			DeactivateMode(Mode->GetID());
 		}
@@ -432,7 +432,7 @@ bool FEditorModeTools::IsOnlyVisibleActiveMode(FEditorModeID InMode) const
 	// Only return true if this is the *only* active mode
 	for (const UEdMode* Mode : ActiveScriptableModes)
 	{
-		if (Mode->GetModeInfo().bVisible && Mode->GetID() != InMode)
+		if (Mode->GetModeInfo().IsVisible() && Mode->GetID() != InMode)
 		{
 			return false;
 		}
@@ -780,7 +780,7 @@ bool FEditorModeTools::ShouldShowModeToolbox() const
 {
 	for (const UEdMode* Mode : ActiveScriptableModes)
 	{
-		if (Mode->GetModeInfo().bVisible && Mode->UsesToolkits())
+		if (Mode->GetModeInfo().IsVisible() && Mode->UsesToolkits())
 		{
 			return true;
 		}
@@ -844,11 +844,11 @@ void FEditorModeTools::ActivateMode(FEditorModeID InID, bool bToggle)
 	}
 
 	// Remove anything that isn't compatible with this mode
-	const bool bIsVisibleMode = ScriptableMode->GetModeInfo().bVisible;
+	const bool bIsVisibleMode = ScriptableMode->GetModeInfo().IsVisible();
 	for (int32 ModeIndex = ActiveScriptableModes.Num() - 1; ModeIndex >= 0; ModeIndex--)
 	{
 		const bool bModesAreCompatible = ScriptableMode->IsCompatibleWith(ActiveScriptableModes[ModeIndex]->GetID()) || ActiveScriptableModes[ModeIndex]->IsCompatibleWith(ScriptableMode->GetID());
-		if (!bModesAreCompatible || (bIsVisibleMode && ActiveScriptableModes[ModeIndex]->GetModeInfo().bVisible))
+		if (!bModesAreCompatible || (bIsVisibleMode && ActiveScriptableModes[ModeIndex]->GetModeInfo().IsVisible()))
 		{
 			DeactivateScriptableModeAtIndex(ModeIndex);
 		}
diff --git a/Engine/Source/Editor/UnrealEd/Private/Settings/SettingsClasses.cpp b/Engine/Source/Editor/UnrealEd/Private/Settings/SettingsClasses.cpp
--- a/Engine/Source/Editor/UnrealEd/Private/Settings/SettingsClasses.cpp
+++ b/Engine/Source/Editor/UnrealEd/Private/Settings/SettingsClasses.cpp
@@ -39,6 +39,7 @@
 #include "InstalledPlatformInfo.h"
 #include "DrawDebugHelpers.h"
 #include "ToolMenus.h"
+#include "Subsystems/AssetEditorSubsystem.h"
 
 #define LOCTEXT_NAMESPACE "SettingsClasses"
 
@@ -442,6 +443,14 @@ void ULevelEditorMiscSettings::PostEditChangeProperty( struct FPropertyChangedEv
 		FNavigationSystem::SetNavigationAutoUpdateEnabled(bNavigationAutoUpdate, EditorContext.World()->GetNavigationSystem());
 	}
 
+	if (Name == GET_MEMBER_NAME_CHECKED(ULevelEditorMiscSettings, bEnableAssetPlacementMode))
+	{
+		if (UAssetEditorSubsystem* AssetEditorSubsystem = GEditor->GetEditorSubsystem<UAssetEditorSubsystem>())
+		{
+			AssetEditorSubsystem->OnEditorModesChanged().Broadcast();
+		}
+	}
+
 	if (!FUnrealEdMisc::Get().IsDeletePreferences())
 	{
 		SaveConfig();
