Separate Panel for Modes Palette, applied to Modeling Mode

diff --git a/Engine/Plugins/Experimental/ModelingToolsEditorMode/Source/ModelingToolsEditorMode/Public/ModelingToolsEditorModeToolkit.h b/Engine/Plugins/Experimental/ModelingToolsEditorMode/Source/ModelingToolsEditorMode/Public/ModelingToolsEditorModeToolkit.h
--- a/Engine/Plugins/Experimental/ModelingToolsEditorMode/Source/ModelingToolsEditorMode/Public/ModelingToolsEditorModeToolkit.h
+++ b/Engine/Plugins/Experimental/ModelingToolsEditorMode/Source/ModelingToolsEditorMode/Public/ModelingToolsEditorModeToolkit.h
@@ -44,6 +44,8 @@ public:
 	virtual FText GetToolPaletteDisplayName(FName PaletteName) const; 
 	virtual void BuildToolPalette(FName PaletteName, class FToolBarBuilder& ToolbarBuilder);
 	virtual void OnToolPaletteChanged(FName PaletteName) override;
+	virtual bool HasIntegratedToolPalettes() const { return false; }
+	virtual bool HasExclusiveToolPalettes() const { return false; }
 
 	virtual FText GetActiveToolDisplayName() const override { return ActiveToolName; }
 	virtual FText GetActiveToolMessage() const override { return ActiveToolMessage; }
diff --git a/Engine/Source/Editor/EditorStyle/Private/StarshipStyle.cpp b/Engine/Source/Editor/EditorStyle/Private/StarshipStyle.cpp
--- a/Engine/Source/Editor/EditorStyle/Private/StarshipStyle.cpp
+++ b/Engine/Source/Editor/EditorStyle/Private/StarshipStyle.cpp
@@ -1779,13 +1779,13 @@ void FStarshipEditorStyle::FStyle::SetupGeneralStyles()
 		
 		PaletteToolBarStyle.SetBackground(FSlateColorBrush(FStyleColors::Recessed));
 
-		PaletteToolBarStyle.SetLabelPadding(FMargin(0.0f, 4.0f, 0.0f, 0.0f));
+		PaletteToolBarStyle.SetLabelPadding(FMargin(0.0f, 1.0f, 0.0f, 0.0f));
 
-		PaletteToolBarStyle.SetButtonPadding(       FMargin(2.0f, 2.0f));
-		PaletteToolBarStyle.SetCheckBoxPadding(     FMargin(1.0f, 1.0f));
-		PaletteToolBarStyle.SetComboButtonPadding(  FMargin(1.0f, 1.0f));
-		PaletteToolBarStyle.SetIndentedBlockPadding(FMargin(1.0f, 1.0f));
-		PaletteToolBarStyle.SetBlockPadding(        FMargin(1.0f, 1.0f));
+		PaletteToolBarStyle.SetButtonPadding(       FMargin(0.0f, 0.0f));
+		PaletteToolBarStyle.SetCheckBoxPadding(     FMargin(0.0f, 0.0f));
+		PaletteToolBarStyle.SetComboButtonPadding(  FMargin(0.0f, 0.0f));
+		PaletteToolBarStyle.SetIndentedBlockPadding(FMargin(0.0f, 0.0f));
+		PaletteToolBarStyle.SetBlockPadding(        FMargin(0.0f, 0.0f));
 
 		Set( "PaletteToolBar.Tab",  FCheckBoxStyle()
 			.SetCheckBoxType(            ESlateCheckBoxType::ToggleButton)
@@ -1806,8 +1806,11 @@ void FStarshipEditorStyle::FStyle::SetupGeneralStyles()
 			.SetPadding(FMargin(2.f, 6.f))
 		);
 
-		Set("PaletteToolBar.MaxUniformToolbarSize", 64.f);
-		Set("PaletteToolBar.MinUniformToolbarSize", 64.f);
+		Set("PaletteToolBar.MaxUniformToolbarSize", 48.f);
+		Set("PaletteToolBar.MinUniformToolbarSize", 48.f);
+
+		Set("PaletteToolBar.ExpandableAreaHeader", new FSlateRoundedBoxBrush(FStyleColors::Dropdown, FVector4(4.0, 4.0, 0.0, 0.0)));
+		Set("PaletteToolBar.ExpandableAreaBody", new FSlateRoundedBoxBrush(FStyleColors::Recessed, FVector4(0.0, 0.0, 4.0, 4.0)));
 		
 		FLinearColor ButtonHoverColor(FColor(81, 81, 81));
 
diff --git a/Engine/Source/Editor/LevelEditor/Private/SLevelEditor.cpp b/Engine/Source/Editor/LevelEditor/Private/SLevelEditor.cpp
--- a/Engine/Source/Editor/LevelEditor/Private/SLevelEditor.cpp
+++ b/Engine/Source/Editor/LevelEditor/Private/SLevelEditor.cpp
@@ -1411,8 +1411,8 @@ TSharedRef<SWidget> SLevelEditor::RestoreContentArea( const TSharedRef<SDockTab>
 	//     9.2. Engine/Config/BaseEditorLayout.ini
 	//     9.3. Etc
 	// 10. Push the new "DefaultLayout.ini" together with your new code.
-	// 11. Also update these instructions if you change the version number (e.g., from "UnrealEd_Layout_v1.4" to "UnrealEd_Layout_v1.5").
-	const FName LayoutName = TEXT("LevelEditor_Layout_v1.5");
+	// 11. Also update these instructions if you change the version number (e.g., from "UnrealEd_Layout_v1.6" to "UnrealEd_Layout_v1.7").
+	const FName LayoutName = TEXT("LevelEditor_Layout_v1.6");
 	const TSharedRef<FTabManager::FLayout> DefaultLayout = FLayoutSaveRestore::LoadFromConfig(GEditorLayoutIni,
 		FTabManager::NewLayout( LayoutName )
 		->AddArea
@@ -1431,6 +1431,13 @@ TSharedRef<SWidget> SLevelEditor::RestoreContentArea( const TSharedRef<SDockTab>
 					->SetSizeCoefficient( .75f )
 					->SetOrientation(Orient_Horizontal)
 					->Split
+					(
+						FTabManager::NewStack()
+						->SetSizeCoefficient( .18f )
+						->SetHideTabWell(true)
+						->AddTab(FEditorModeTools::EditorModeToolbarTabName, ETabState::ClosedTab)
+					)
+					->Split
 					(
 						FTabManager::NewStack()
 						->SetSizeCoefficient( 0.3f )
@@ -1439,22 +1446,10 @@ TSharedRef<SWidget> SLevelEditor::RestoreContentArea( const TSharedRef<SDockTab>
 					)
 					->Split
 					(
-						FTabManager::NewSplitter()
-						->SetOrientation(Orient_Vertical)
-						->SetSizeCoefficient( 1.15f )
-						->Split
-						(
-							FTabManager::NewStack()
-							->SetHideTabWell(true)
-							->AddTab(FEditorModeTools::EditorModeToolbarTabName, ETabState::ClosedTab)
-						)
-						->Split
-						(
-							FTabManager::NewStack()
-							->SetHideTabWell(true)
-							->SetSizeCoefficient( 1.0f )
-							->AddTab(LevelEditorTabIds::LevelEditorViewport, ETabState::OpenedTab)
-						)
+						FTabManager::NewStack()
+						->SetHideTabWell(true)
+						->SetSizeCoefficient( 1.0f )
+						->AddTab(LevelEditorTabIds::LevelEditorViewport, ETabState::OpenedTab)
 					)
 				)
 				->Split
diff --git a/Engine/Source/Editor/LevelEditor/Private/SLevelEditorToolBox.cpp b/Engine/Source/Editor/LevelEditor/Private/SLevelEditorToolBox.cpp
--- a/Engine/Source/Editor/LevelEditor/Private/SLevelEditorToolBox.cpp
+++ b/Engine/Source/Editor/LevelEditor/Private/SLevelEditorToolBox.cpp
@@ -17,7 +17,6 @@
 #include "EdMode.h"
 #include "Widgets/Docking/SDockTab.h"
 #include "Widgets/Input/SCheckBox.h"
-#include "Widgets/Layout/SExpandableArea.h"
 #include "Widgets/Layout/SWidgetSwitcher.h"
 #include "Widgets/Layout/SUniformWrapPanel.h"
 #include "StatusBarSubsystem.h"
@@ -173,95 +172,100 @@ void SLevelEditorToolBox::UpdateInlineContent(const TSharedPtr<IToolkit>& Toolki
 
 			TWeakPtr<FModeToolkit> ModeToolkit = StaticCastSharedPtr<FModeToolkit>(Toolkit);
 
-			TSharedRef< SUniformWrapPanel> PaletteTabBox = SNew(SUniformWrapPanel)
-				.SlotPadding( FMargin(1.f, 2.f) )
-				.HAlign(HAlign_Center);
-
-			// Only show if there is more than one child in the switcher
-			PaletteTabBox->SetVisibility(TAttribute<EVisibility>::Create(TAttribute<EVisibility>::FGetter::CreateLambda([PaletteTabBox]() -> EVisibility { return PaletteTabBox->GetChildren()->Num() > 1 ? EVisibility::Visible : EVisibility::Collapsed; }) ));
-
-			// Also build the toolkit here 
-			TArray<FName> PaletteNames;
-			ModeToolkit.Pin()->GetToolPaletteNames(PaletteNames);
-
-			TSharedPtr<FUICommandList> CommandList;
-			CommandList = ModeToolkit.Pin()->GetToolkitCommands();
-
-			TSharedRef< SWidgetSwitcher > PaletteSwitcher = SNew(SWidgetSwitcher)
-			.WidgetIndex_Lambda( [PaletteNames, ModeToolkit] () -> int32 { 
-				int32 FoundIndex;
-				TSharedPtr<FModeToolkit> PinnedToolkit = ModeToolkit.Pin();
-				if (PinnedToolkit.IsValid() && PaletteNames.Find(PinnedToolkit->GetCurrentPalette(), FoundIndex))
-				{
-					return FoundIndex;	
-				}
-				return 0;
-			} );
-			
-			for(auto Palette : PaletteNames)
+			if (ModeToolkit.Pin()->HasIntegratedToolPalettes())
 			{
-				FName ToolbarCustomizationName = ModeToolkit.Pin()->GetEditorMode() ? ModeToolkit.Pin()->GetEditorMode()->GetModeInfo().ToolbarCustomizationName : ModeToolkit.Pin()->GetScriptableEditorMode()->GetModeInfo().ToolbarCustomizationName;
-				FUniformToolBarBuilder ModeToolbarBuilder(CommandList, FMultiBoxCustomization(ToolbarCustomizationName));
-				ModeToolbarBuilder.SetStyle(&FEditorStyle::Get(), "PaletteToolBar");
+				TSharedRef< SUniformWrapPanel> PaletteTabBox = SNew(SUniformWrapPanel)
+					.SlotPadding( FMargin(1.f, 2.f) )
+					.HAlign(HAlign_Center);
+
+				// Only show if there is more than one child in the switcher
+				PaletteTabBox->SetVisibility(TAttribute<EVisibility>::Create(TAttribute<EVisibility>::FGetter::CreateLambda([PaletteTabBox]() -> EVisibility { return PaletteTabBox->GetChildren()->Num() > 1 ? EVisibility::Visible : EVisibility::Collapsed; }) ));
+
+				// Also build the toolkit here 
+				TArray<FName> PaletteNames;
+				ModeToolkit.Pin()->GetToolPaletteNames(PaletteNames);
+
+				TSharedPtr<FUICommandList> CommandList;
+				CommandList = ModeToolkit.Pin()->GetToolkitCommands();
+
+				TSharedRef< SWidgetSwitcher > PaletteSwitcher = SNew(SWidgetSwitcher)
+				.WidgetIndex_Lambda( [PaletteNames, ModeToolkit] () -> int32 { 
+					int32 FoundIndex;
+					TSharedPtr<FModeToolkit> PinnedToolkit = ModeToolkit.Pin();
+					if (PinnedToolkit.IsValid() && PaletteNames.Find(PinnedToolkit->GetCurrentPalette(), FoundIndex))
+					{
+						return FoundIndex;	
+					}
+					return 0;
+				} );
+				
+				
+				for(auto Palette : PaletteNames)
+				{
+					FName ToolbarCustomizationName = ModeToolkit.Pin()->GetEditorMode() ? ModeToolkit.Pin()->GetEditorMode()->GetModeInfo().ToolbarCustomizationName : ModeToolkit.Pin()->GetScriptableEditorMode()->GetModeInfo().ToolbarCustomizationName;
+					FUniformToolBarBuilder ModeToolbarBuilder(CommandList, FMultiBoxCustomization(ToolbarCustomizationName));
+					ModeToolbarBuilder.SetStyle(&FEditorStyle::Get(), "PaletteToolBar");
 
-				ModeToolkit.Pin()->BuildToolPalette(Palette, ModeToolbarBuilder);
+					ModeToolkit.Pin()->BuildToolPalette(Palette, ModeToolbarBuilder);
 
-				TSharedRef<SWidget> PaletteWidget = ModeToolbarBuilder.MakeWidget();
+					TSharedRef<SWidget> PaletteWidget = ModeToolbarBuilder.MakeWidget();
 
-				PaletteTabBox->AddSlot()
-				[
-					SNew(SCheckBox)
-					.Padding(FMargin(8.f, 4.f, 8.f, 5.f))
-					.Style( FEditorStyle::Get(),  "PaletteToolBar.Tab" )
-					.OnCheckStateChanged_Lambda([/*PaletteSwitcher, PaletteWidget, */ModeToolkit, Palette](const ECheckBoxState) {
+					PaletteTabBox->AddSlot()
+					[
+						SNew(SCheckBox)
+						.Padding(FMargin(8.f, 4.f, 8.f, 5.f))
+						.Style( FEditorStyle::Get(),  "PaletteToolBar.Tab" )
+						.OnCheckStateChanged_Lambda([ModeToolkit, Palette](const ECheckBoxState) {
+								TSharedPtr<FModeToolkit> PinnedToolkit = ModeToolkit.Pin();
+								if (PinnedToolkit.IsValid())
+								{
+									PinnedToolkit->SetCurrentPalette(Palette);
+								}
+							} 
+						)
+						// .IsChecked_Lambda( [PaletteSwitcher, PaletteWidget] () -> ECheckBoxState { return PaletteSwitcher->GetActiveWidget() == PaletteWidget ? ECheckBoxState::Checked : ECheckBoxState::Unchecked; })
+						.IsChecked_Lambda([ModeToolkit, Palette]() -> ECheckBoxState {
 							TSharedPtr<FModeToolkit> PinnedToolkit = ModeToolkit.Pin();
-							if (PinnedToolkit.IsValid())
+							if (PinnedToolkit.IsValid() && (PinnedToolkit->GetCurrentPalette() == Palette))
 							{
-								PinnedToolkit->SetCurrentPalette(Palette);
+								return ECheckBoxState::Checked;
 							}
-						} 
-					)
-					// .IsChecked_Lambda( [PaletteSwitcher, PaletteWidget] () -> ECheckBoxState { return PaletteSwitcher->GetActiveWidget() == PaletteWidget ? ECheckBoxState::Checked : ECheckBoxState::Unchecked; })
-					.IsChecked_Lambda([ModeToolkit, Palette]() -> ECheckBoxState {
-						TSharedPtr<FModeToolkit> PinnedToolkit = ModeToolkit.Pin();
-						if (PinnedToolkit.IsValid() && (PinnedToolkit->GetCurrentPalette() == Palette))
-						{
-							return ECheckBoxState::Checked;
-						}
-						return ECheckBoxState::Unchecked;
-					})
+							return ECheckBoxState::Unchecked;
+						})
+						[
+							SNew(STextBlock)
+							.TextStyle(FAppStyle::Get(), "NormalText")
+							.Text(ModeToolkit.Pin()->GetToolPaletteDisplayName(Palette))
+							.Justification(ETextJustify::Center)
+						]
+					];
+
+
+					PaletteSwitcher->AddSlot()
 					[
-						SNew(STextBlock)
-						.TextStyle(FAppStyle::Get(), "NormalText")
-						.Text(ModeToolkit.Pin()->GetToolPaletteDisplayName(Palette))
-						.Justification(ETextJustify::Center)
-					]
-				];
-
+						PaletteWidget
+					]; 
+				}
 
-				PaletteSwitcher->AddSlot()
-				[
-					PaletteWidget
-				]; 
-			}
 
+				ModeToolHeader->SetContent(
+					SNew(SVerticalBox)
 
-			ModeToolHeader->SetContent(
-				SNew(SVerticalBox)
+					+SVerticalBox::Slot()
+					.Padding(8.f, 0.f, 0.f, 8.f)
+					.AutoHeight()
+					[
+						PaletteTabBox
+					]
 
-				+SVerticalBox::Slot()
-				.Padding(8.f, 0.f, 0.f, 8.f)
-				.AutoHeight()
-				[
-					PaletteTabBox
-				]
+					+SVerticalBox::Slot()
+					.AutoHeight()
+					[
+						PaletteSwitcher
+					]
+				);
 
-				+SVerticalBox::Slot()
-				.AutoHeight()
-				[
-					PaletteSwitcher
-				]
-			);
+			}
 
 			if (StatusBarMessageHandle.IsValid())
 			{
diff --git a/Engine/Source/Editor/UnrealEd/Private/EditorModeManager.cpp b/Engine/Source/Editor/UnrealEd/Private/EditorModeManager.cpp
--- a/Engine/Source/Editor/UnrealEd/Private/EditorModeManager.cpp
+++ b/Engine/Source/Editor/UnrealEd/Private/EditorModeManager.cpp
@@ -30,6 +30,8 @@
 #include "Subsystems/BrushEditingSubsystem.h"
 #include "Tools/UEdMode.h"
 #include "Widgets/Images/SImage.h"
+#include "Widgets/Layout/SExpandableArea.h"
+#include "Widgets/Layout/SScrollBox.h"
 #include "InputRouter.h"
 #include "InteractiveGizmoManager.h"
 #include "EdModeInteractiveToolsContext.h"
@@ -592,7 +594,12 @@ void FEditorModeTools::RebuildModeToolBar()
 	{
 		ModeToolbarBoxPinned->ClearChildren();
 
-		TSharedRef< SHorizontalBox > PaletteTabBox = SNew(SHorizontalBox);
+		bool bExclusivePalettes = true;
+		TSharedRef<SVerticalBox> ToolBoxVBox = SNew(SVerticalBox);
+
+		TSharedRef< SUniformWrapPanel> PaletteTabBox = SNew(SUniformWrapPanel)
+			.SlotPadding(FMargin(1.f, 2.f))
+			.HAlign(HAlign_Left);
 		TSharedRef< SWidgetSwitcher > PaletteSwitcher = SNew(SWidgetSwitcher);
 
 		int32 PaletteCount = ActiveToolBarRows.Num();
@@ -604,7 +611,6 @@ void FEditorModeTools::RebuildModeToolBar()
 				if (ensure(Row.ToolbarWidget.IsValid()))
 				{
 					TSharedRef<SWidget> PaletteWidget = Row.ToolbarWidget.ToSharedRef();
-
 					
 					TSharedPtr<FModeToolkit> RowToolkit;
 					if (FEdMode* Mode = GetActiveMode(Row.ModeID))
@@ -616,32 +622,56 @@ void FEditorModeTools::RebuildModeToolBar()
 						RowToolkit = ScriptableMode->GetToolkit();
 					}
 
-					// Don't show Palette Tabs if there is only one
-					if (PaletteCount > 1)
+					bExclusivePalettes = RowToolkit->HasExclusiveToolPalettes();
+
+					if (!bExclusivePalettes)
 					{
-						PaletteTabBox->AddSlot()
-						.AutoWidth()
-						.Padding(FMargin(0.f, 1.0f, 1.0f, 0.0f))
+						ToolBoxVBox->AddSlot()
+						.AutoHeight()
+						.Padding(FMargin(2.0, 2.0))
 						[
-							SNew(SCheckBox)
-							.Style( FEditorStyle::Get(),  "ToolPalette.DockingTab" )
-							.OnCheckStateChanged_Lambda( [PaletteSwitcher, Row, RowToolkit] (const ECheckBoxState) { 
-									PaletteSwitcher->SetActiveWidget(Row.ToolbarWidget.ToSharedRef());
-									RowToolkit->SetCurrentPalette(Row.PaletteName);
-								} 
-							)
-							.IsChecked_Lambda( [PaletteSwitcher, PaletteWidget] () -> ECheckBoxState { return PaletteSwitcher->GetActiveWidget() == PaletteWidget ? ECheckBoxState::Checked : ECheckBoxState::Unchecked; })
+							SNew(SExpandableArea)
+							.AreaTitle(Row.DisplayName)
+							.AreaTitleFont(FAppStyle::Get().GetFontStyle("NormalFont"))
+							.BorderImage(FAppStyle::Get().GetBrush("PaletteToolbar.ExpandableAreaHeader"))
+							.BodyBorderImage(FAppStyle::Get().GetBrush("PaletteToolbar.ExpandableAreaBody"))
+							.HeaderPadding(FMargin(4.f))
+							.Padding(FMargin(4.0, 0.0))
+							.BodyContent()
 							[
-								SNew(STextBlock)
-								.Text(Row.DisplayName)
+								PaletteWidget	
 							]
 						];
 					}
-
-					PaletteSwitcher->AddSlot()
-					[
-						PaletteWidget
-					]; 
+					else 
+					{
+						// Don't show Palette Tabs if there is only one
+						if (PaletteCount > 1)
+						{
+							PaletteTabBox->AddSlot()
+							.AutoWidth()
+							.Padding(FMargin(0.f, 1.0f, 1.0f, 0.0f))
+							[
+								SNew(SCheckBox)
+								.Style( FEditorStyle::Get(),  "ToolPalette.DockingTab" )
+								.OnCheckStateChanged_Lambda( [PaletteSwitcher, Row, RowToolkit] (const ECheckBoxState) { 
+										PaletteSwitcher->SetActiveWidget(Row.ToolbarWidget.ToSharedRef());
+										RowToolkit->SetCurrentPalette(Row.PaletteName);
+									} 
+								)
+								.IsChecked_Lambda( [PaletteSwitcher, PaletteWidget] () -> ECheckBoxState { return PaletteSwitcher->GetActiveWidget() == PaletteWidget ? ECheckBoxState::Checked : ECheckBoxState::Unchecked; })
+								[
+									SNew(STextBlock)
+									.Text(Row.DisplayName)
+								]
+							];
+						}
+
+						PaletteSwitcher->AddSlot()
+						[
+							PaletteWidget
+						]; 
+					}
 				}
 			}
 
@@ -662,15 +692,26 @@ void FEditorModeTools::RebuildModeToolBar()
 			];
 
 			ModeToolbarBoxPinned->AddSlot()
+			.AutoHeight()
 			.Padding(1.f)
 			[
 				SNew(SBox)
-				.HeightOverride(45.f)
+				.HeightOverride(PaletteSwitcher->GetNumWidgets() > 0 ? 45.f : 0.f)
 				[
 					PaletteSwitcher 
 				]
 			];
 
+			ModeToolbarBoxPinned->AddSlot()
+			[
+				SNew(SScrollBox)
+
+				+SScrollBox::Slot()
+				[
+					ToolBoxVBox	
+				]
+			];
+
 			ModeToolbarPaletteSwitcher = PaletteSwitcher;
 		}
 		else
@@ -680,6 +721,55 @@ void FEditorModeTools::RebuildModeToolBar()
 	}
 }
 
+void FEditorModeTools::SpawnOrUpdateModeToolbar()
+{
+	if(ShouldShowModeToolbar())
+	{
+		if (ModeToolbarTab.IsValid())
+		{
+			RebuildModeToolBar();
+		}
+		else if (ToolkitHost.IsValid())
+		{
+			ToolkitHost.Pin()->GetTabManager()->TryInvokeTab(EditorModeToolbarTabName);
+		}
+	}
+}
+
+void FEditorModeTools::InvokeToolPaletteTab(FEditorModeID InModeID, FName InPaletteName)
+{
+	if (!ModeToolbarPaletteSwitcher.Pin()) 
+	{
+		return;
+	}
+
+	for (auto Row: ActiveToolBarRows)
+	{
+		if (Row.ModeID == InModeID && Row.PaletteName == InPaletteName)
+		{
+			TSharedRef<SWidget> PaletteWidget = Row.ToolbarWidget.ToSharedRef();
+
+			TSharedPtr<FModeToolkit> RowToolkit;
+			if (FEdMode* Mode = GetActiveMode(InModeID))
+			{
+				RowToolkit = Mode->GetToolkit();
+			}
+			else if (UEdMode* ScriptableMode = GetActiveScriptableMode(InModeID))
+			{
+				RowToolkit = ScriptableMode->GetToolkit();
+			}
+
+			TSharedPtr<SWidget> ActiveWidget = ModeToolbarPaletteSwitcher.Pin()->GetActiveWidget();
+			if (RowToolkit && ActiveWidget.Get() != Row.ToolbarWidget.Get())
+			{
+				ModeToolbarPaletteSwitcher.Pin()->SetActiveWidget(Row.ToolbarWidget.ToSharedRef());
+				RowToolkit->OnToolPaletteChanged(Row.PaletteName);
+			}
+			break;	
+		}
+	}	
+}
+
 void FEditorModeTools::RemoveAllDelegateHandlers()
 {
 	if (GEditor)
@@ -761,7 +851,6 @@ TSharedRef<SDockTab> FEditorModeTools::MakeModeToolbarTab()
 	TSharedRef<SDockTab> ToolbarTabRef = 	
 		SNew(SDockTab)
 		.Label(NSLOCTEXT("EditorModes", "EditorModesToolbarTitle", "Mode Toolbar"))
-		.ShouldAutosize(true)
 		.ContentPadding(0.0f)
 		.Icon(FStyleDefaults::GetNoBrush())
 		[
@@ -887,8 +976,12 @@ void FEditorModeTools::ActivateMode(FEditorModeID InID, bool bToggle)
 			ActiveToolBarRows.Emplace(ScriptableMode->GetID(), Palette, Toolkit->GetToolPaletteDisplayName(Palette), ModeToolbarBuilder.MakeWidget());
 			PaletteCount++;
 		}
+
+		if (!Toolkit->HasIntegratedToolPalettes() && PaletteCount > 0)
+		{
+			SpawnOrUpdateModeToolbar();
+		}
 	}
-	
 	// Update the editor UI
 	FEditorSupportDelegates::UpdateUI.Broadcast();
 }
diff --git a/Engine/Source/Editor/UnrealEd/Public/EditorModeManager.h b/Engine/Source/Editor/UnrealEd/Public/EditorModeManager.h
--- a/Engine/Source/Editor/UnrealEd/Public/EditorModeManager.h
+++ b/Engine/Source/Editor/UnrealEd/Public/EditorModeManager.h
@@ -117,6 +117,7 @@ protected:
 		
 private:
 	void RebuildModeToolBar();
+	void SpawnOrUpdateModeToolbar();
 public:
 
 	/**
@@ -498,6 +499,11 @@ public:
 	void UpdateInternalData();
 	bool IsOnlyVisibleActiveMode(FEditorModeID InMode) const;
 
+	/*
+	* Sets the active Modes ToolBar Palette Tab to the named Palette
+	*/
+	void  InvokeToolPaletteTab(FEditorModeID InMode, FName InPaletteName);
+
 	/** returns true if all active EdModes are OK with an AutoSave happening now  */
 	bool CanAutoSave() const;
 
diff --git a/Engine/Source/Editor/UnrealEd/Public/Toolkits/BaseToolkit.h b/Engine/Source/Editor/UnrealEd/Public/Toolkits/BaseToolkit.h
--- a/Engine/Source/Editor/UnrealEd/Public/Toolkits/BaseToolkit.h
+++ b/Engine/Source/Editor/UnrealEd/Public/Toolkits/BaseToolkit.h
@@ -138,6 +138,12 @@ public:
 	 **/
 	virtual FText GetToolPaletteDisplayName(FName Palette) const { return FText(); }
 
+	/* Exclusive Tool Palettes only allow users to use tools from one palette at a time */
+	virtual bool HasExclusiveToolPalettes() const { return true; }
+
+	/* Integrated Tool Palettes show up in the same panel as their details */
+	virtual bool HasIntegratedToolPalettes() const { return true; }
+
 	/**
 	 * @param PaletteIndex      The index of the ToolPalette to build
 	 * @param ToolbarBuilder    The builder to use for given PaletteIndex
diff --git a/Engine/Source/Runtime/Slate/Private/Framework/MultiBox/MultiBox.cpp b/Engine/Source/Runtime/Slate/Private/Framework/MultiBox/MultiBox.cpp
--- a/Engine/Source/Runtime/Slate/Private/Framework/MultiBox/MultiBox.cpp
+++ b/Engine/Source/Runtime/Slate/Private/Framework/MultiBox/MultiBox.cpp
@@ -877,11 +877,11 @@ void SMultiBoxWidget::BuildMultiBoxWidget()
 			[
 				SAssignNew(UniformToolbarPanel, SUniformWrapPanel)
 				.HAlign(HAlign_Left)
-				.MinDesiredSlotWidth(68.f)
-				.MinDesiredSlotHeight(52.f)
-				.MaxDesiredSlotWidth(68.f)
-				.MaxDesiredSlotHeight(52.f)
-				.SlotPadding(0.f)
+				.MinDesiredSlotWidth(46.f)
+				.MinDesiredSlotHeight(43.f)
+				.MaxDesiredSlotWidth(46.f)
+				.MaxDesiredSlotHeight(43.f)
+				.SlotPadding(FMargin(1.f, 1.f))
 			];
 		}
 		break;
diff --git a/Engine/Source/Runtime/Slate/Private/Widgets/Layout/SExpandableArea.cpp b/Engine/Source/Runtime/Slate/Private/Widgets/Layout/SExpandableArea.cpp
--- a/Engine/Source/Runtime/Slate/Private/Widgets/Layout/SExpandableArea.cpp
+++ b/Engine/Source/Runtime/Slate/Private/Widgets/Layout/SExpandableArea.cpp
@@ -36,8 +36,7 @@ void SExpandableArea::Construct( const FArguments& InArgs )
 		HeaderContent = 
 			SNew(STextBlock)
 			.Text(InArgs._AreaTitle)
-			.Font(InArgs._AreaTitleFont)
-			.ShadowOffset(FVector2D(1.0f, 1.0f));
+			.Font(InArgs._AreaTitleFont);
 	}
 
 	// If the user wants the body of the expanded section to be different from the title area,
