Convert Editor layout version mismatch warning from a modal prompt to an output log message. This is not a user facing message and should not block the editor

diff --git a/Engine/Source/Editor/LevelEditor/Private/SLevelEditor.cpp b/Engine/Source/Editor/LevelEditor/Private/SLevelEditor.cpp
--- a/Engine/Source/Editor/LevelEditor/Private/SLevelEditor.cpp
+++ b/Engine/Source/Editor/LevelEditor/Private/SLevelEditor.cpp
@@ -1491,10 +1491,12 @@ TSharedRef<SWidget> SLevelEditor::RestoreContentArea( const TSharedRef<SDockTab>
 	if (RemovedOlderLayoutVersions.Num() > 0)
 	{
 		// FMessageDialog - Notify the user that the layout version was updated and the current layout uses a deprecated one
-		const FText TextTitle = LOCTEXT("LevelEditorVersionErrorTitle", "Unreal Level Editor Layout Version Mismatch");
+		//const FText TextTitle = LOCTEXT("LevelEditorVersionErrorTitle", "Unreal Level Editor Layout Version Mismatch");
 		const FText TextBody = FText::Format(LOCTEXT("LevelEditorVersionErrorBody", "The expected Unreal Level Editor layout version is \"{0}\", while only version \"{1}\" was found. I.e., the current layout was created with a previous version of Unreal that is deprecated and no longer compatible.\n\nUnreal will continue with the default layout for its current version, the deprecated one has been removed.\n\nYou can create and save your custom layouts with \"Window\"->\"Save Layout\"->\"Save Layout As...\"."),
 			FText::FromString(LayoutName.ToString()), FText::FromString(RemovedOlderLayoutVersions[0]));
-		FMessageDialog::Open(EAppMsgType::Ok, TextBody, &TextTitle);
+		//FMessageDialog::Open(EAppMsgType::Ok, TextBody, &TextTitle);
+
+		UE_LOG(LogSlate, Log, TEXT("%s"), *TextBody.ToString());
 	}
 
 	FLayoutExtender LayoutExtender;
