Pathway for Toolkits to add Viewport Overlay Widgets

diff --git a/Engine/Source/Editor/EditorFramework/Public/Toolkits/IToolkitHost.h b/Engine/Source/Editor/EditorFramework/Public/Toolkits/IToolkitHost.h
--- a/Engine/Source/Editor/EditorFramework/Public/Toolkits/IToolkitHost.h
+++ b/Engine/Source/Editor/EditorFramework/Public/Toolkits/IToolkitHost.h
@@ -6,6 +6,7 @@
 
 class FEditorModeTools;
 class UTypedElementCommonActions;
+class IAssetViewport;
 
 /**
  * Base interface class for toolkit hosts
@@ -15,6 +16,9 @@ class IToolkitHost
 
 public:
 
+	/* Notificaton when the active viewport changed */
+	DECLARE_MULTICAST_DELEGATE_TwoParams(FOnActiveViewportChanged,  TSharedPtr<IAssetViewport>, TSharedPtr<IAssetViewport>)
+
 	/** Gets a widget that can be used to parent a modal window or pop-up to.  You shouldn't be using this widget for
 	    anything other than parenting, as the type of widget and behavior/lifespan is completely up to the host. */
 	virtual TSharedRef< class SWidget > GetParentWidget() = 0;
@@ -39,4 +43,13 @@ public:
 
 	/** Returns the common actions implementation for this toolkit host */
 	virtual UTypedElementCommonActions* GetCommonActions() const = 0;
+
+	/** Allows Toolkits to push widgets to the viewport.  Passing a nullptr for the Viewport will add  or 
+	    remove the OverlayWidget to or from the ActiveViewport */
+	virtual void AddViewportOverlayWidget(TSharedRef<SWidget>, TSharedPtr<IAssetViewport> InViewport = nullptr)  {}
+	virtual void RemoveViewportOverlayWidget(TSharedRef<SWidget>, TSharedPtr<IAssetViewport> InViewport = nullptr) {}
+
+	/** Gets a multicast delegate which is executed whenever the toolkit host's active viewport changes. */
+	virtual FOnActiveViewportChanged& OnActiveViewportChanged() = 0;
+
 };
diff --git a/Engine/Source/Editor/EditorStyle/Private/StarshipStyle.cpp b/Engine/Source/Editor/EditorStyle/Private/StarshipStyle.cpp
--- a/Engine/Source/Editor/EditorStyle/Private/StarshipStyle.cpp
+++ b/Engine/Source/Editor/EditorStyle/Private/StarshipStyle.cpp
@@ -2507,6 +2507,11 @@ void FStarshipEditorStyle::FStyle::SetupViewportStyles()
 		Set("EditorViewportToolBar.Maximize.Checked", new IMAGE_BRUSH_SVG("Starship/EditorViewport/quad", Icon16x16));
 		Set("EditorViewportToolBar.RestoreFromImmersive.Normal", new IMAGE_BRUSH("Icons/icon_RestoreFromImmersive_16px", Icon16x16));
 
+		FLinearColor ViewportOverlayColor = FStyleColors::Input.GetSpecifiedColor();
+		ViewportOverlayColor.A = 0.75f;
+
+		Set("EditorViewport.OverlayBrush", new FSlateRoundedBoxBrush(ViewportOverlayColor, 8.0, FStyleColors::Dropdown, 1.0));
+
 	}
 
 	// Legacy Viewport ToolbarBar
diff --git a/Engine/Source/Editor/LevelEditor/Private/SLevelEditor.cpp b/Engine/Source/Editor/LevelEditor/Private/SLevelEditor.cpp
--- a/Engine/Source/Editor/LevelEditor/Private/SLevelEditor.cpp
+++ b/Engine/Source/Editor/LevelEditor/Private/SLevelEditor.cpp
@@ -445,6 +445,17 @@ bool SLevelEditor::HasActivePlayInEditorViewport() const
 	return false;
 }
 
+void SLevelEditor::Tick(const FGeometry& AllottedGeometry, const double InCurrentTime, const float InDeltaTime)
+{
+	// Update the ActiveViewport
+	TSharedPtr<SLevelViewport> LastActiveViewport = CachedActiveViewport.Pin();
+	if (!LastActiveViewport || &LastActiveViewport->GetLevelViewportClient() != GCurrentLevelEditingViewportClient)
+	{
+		TSharedPtr<SLevelViewport> NewActiveViewport = GetActiveViewport();
+		CachedActiveViewport = NewActiveViewport;
+		OnActiveViewportChanged().Broadcast(LastActiveViewport, NewActiveViewport);
+	}
+}
 
 TSharedPtr<SLevelViewport> SLevelEditor::GetActiveViewport()
 {
@@ -2011,4 +2022,29 @@ UTypedElementCommonActions* SLevelEditor::GetCommonActions() const
 	return CommonActions;
 }
 
+void SLevelEditor::AddViewportOverlayWidget(TSharedRef<SWidget> InWidget, TSharedPtr<IAssetViewport> InViewport)
+{
+	if (InViewport != nullptr)
+	{
+		InViewport->AddOverlayWidget(InWidget);
+	}
+	else if (TSharedPtr<SLevelViewport> ActiveViewport = GetActiveViewport())
+	{
+		ActiveViewport->AddOverlayWidget(InWidget);
+	}
+}
+
+void SLevelEditor::RemoveViewportOverlayWidget(TSharedRef<SWidget> InWidget, TSharedPtr<IAssetViewport> InViewport)
+{
+	if (InViewport != nullptr)
+	{
+		InViewport->RemoveOverlayWidget(InWidget);
+	}
+	else if (TSharedPtr<SLevelViewport> ActiveViewport = GetActiveViewport())
+	{
+		ActiveViewport->RemoveOverlayWidget(InWidget);
+	}
+}
+
+
 #undef LOCTEXT_NAMESPACE
\ No newline at end of file
diff --git a/Engine/Source/Editor/LevelEditor/Private/SLevelEditor.h b/Engine/Source/Editor/LevelEditor/Private/SLevelEditor.h
--- a/Engine/Source/Editor/LevelEditor/Private/SLevelEditor.h
+++ b/Engine/Source/Editor/LevelEditor/Private/SLevelEditor.h
@@ -130,6 +130,13 @@ public:
 	virtual TSharedRef<SWidget> CreateToolBox() override;
 	virtual FEditorModeTools& GetEditorModeManager() const override;
 	virtual UTypedElementCommonActions* GetCommonActions() const override;
+	virtual FOnActiveViewportChanged& OnActiveViewportChanged() { return OnActiveViewportChangedDelegate; }
+	virtual void AddViewportOverlayWidget(TSharedRef<SWidget>, TSharedPtr<IAssetViewport> InViewport = nullptr) override;
+	virtual void RemoveViewportOverlayWidget(TSharedRef<SWidget>, TSharedPtr<IAssetViewport> InViewport = nullptr) override; 
+
+
+	/* Tick to check the ActiveViewport */
+	virtual void Tick(const FGeometry& AllottedGeometry, const double InCurrentTime, const float InDeltaTime) override;
 
 	/** SWidget overrides */
 	virtual bool SupportsKeyboardFocus() const override
@@ -256,6 +263,9 @@ private:
 	// A list of any standalone editor viewports that aren't in tabs
 	TArray< TWeakPtr<SLevelViewport> > StandaloneViewports;
 
+	// The last known active viewport
+	TWeakPtr<class SLevelViewport> CachedActiveViewport;
+
 	// Border that hosts the document content for the level editor.
 	TSharedPtr< SBorder > DocumentsAreaBorder;
 	
@@ -315,6 +325,9 @@ private:
 	/** Actor details SCS editor customization */
 	TSharedPtr<ISCSEditorUICustomization> ActorDetailsSCSEditorUICustomization;
 		
+	/** A delegate which is called any time the LevelEditor's active viewport changes. */
+	FOnActiveViewportChanged OnActiveViewportChangedDelegate;
+
 	/** If this flag is raised we will force refresh on next selection update. */
 	bool bNeedsRefresh : 1;
 };
diff --git a/Engine/Source/Editor/UnrealEd/Private/Toolkits/SStandaloneAssetEditorToolkitHost.h b/Engine/Source/Editor/UnrealEd/Private/Toolkits/SStandaloneAssetEditorToolkitHost.h
--- a/Engine/Source/Editor/UnrealEd/Private/Toolkits/SStandaloneAssetEditorToolkitHost.h
+++ b/Engine/Source/Editor/UnrealEd/Private/Toolkits/SStandaloneAssetEditorToolkitHost.h
@@ -57,6 +57,7 @@ public:
 	virtual void OnToolkitHostingFinished( const TSharedRef< class IToolkit >& Toolkit ) override;
 	virtual UWorld* GetWorld() const override;
 	virtual UTypedElementCommonActions* GetCommonActions() const override;
+	virtual FOnActiveViewportChanged& OnActiveViewportChanged() override { return OnActiveViewportChangedDelegate; }
 	void CreateDefaultStandaloneMenuBar(UToolMenu* MenuBar);
 
 	/** SWidget overrides */
@@ -123,6 +124,10 @@ private:
 	/** The menu extenders to populate the main toolkit host menu with */
 	TArray< TSharedPtr<FExtender> > MenuExtenders;
 
+	/** A delegate which is called any time the LevelEditor's active viewport changes. */
+	FOnActiveViewportChanged OnActiveViewportChangedDelegate;
+
+
 	SVerticalBox::FSlot* ToolbarSlot;
 };
 
