Ported "CanDelete" menu validation to use typed elements

#rnx

diff --git a/Engine/Source/Editor/LevelEditor/Private/LevelEditorActions.cpp b/Engine/Source/Editor/LevelEditor/Private/LevelEditorActions.cpp
--- a/Engine/Source/Editor/LevelEditor/Private/LevelEditorActions.cpp
+++ b/Engine/Source/Editor/LevelEditor/Private/LevelEditorActions.cpp
@@ -1625,40 +1625,30 @@ bool FLevelEditorActionCallbacks::Delete_CanExecute()
 		return false;
 	}
 
-	bool bCanDelete = false;
-	if (GEditor->GetSelectedComponentCount() > 0)
+	static const FName NAME_LevelEditor = "LevelEditor";
+	if (TSharedPtr<ILevelEditor> LevelEditor = FModuleManager::GetModuleChecked<FLevelEditorModule>(NAME_LevelEditor).GetLevelEditorInstance().Pin())
 	{
-		TArray<UActorComponent*> SelectedComponents;
-		for (FSelectionIterator It(GEditor->GetSelectedComponentIterator()); It; ++It)
-		{
-			SelectedComponents.Add(CastChecked<UActorComponent>(*It));
-		}
+		bool bCanDelete = false;
 
-		bCanDelete = FComponentEditorUtils::CanDeleteComponents(SelectedComponents);
-	}
-	else
-	{
-		UWorld* World = GetWorld();
-		if (World)
+		const UTypedElementSelectionSet* SelectionSet = LevelEditor->GetElementSelectionSet();
+		SelectionSet->ForEachSelectedElement<UTypedElementWorldInterface>([&bCanDelete](const TTypedElement<UTypedElementWorldInterface>& InWorldElement)
 		{
-			bCanDelete = GUnrealEd->CanDeleteSelectedActors(World, true, false);
-		}
-	}
+			bCanDelete |= InWorldElement.CanDeleteElement();
+			return !bCanDelete;
+		});
 
-	if (!bCanDelete)
-	{
-		TWeakPtr<class ILevelEditor> LevelEditor = FModuleManager::GetModuleChecked<FLevelEditorModule>(TEXT("LevelEditor")).GetLevelEditorInstance();
-		if (LevelEditor.IsValid())
+		if (!bCanDelete)
 		{
-			TSharedPtr<class ISceneOutliner> SceneOutlinerPtr = LevelEditor.Pin()->GetSceneOutliner();
-			if (SceneOutlinerPtr.IsValid())
+			if (TSharedPtr<ISceneOutliner> SceneOutlinerPtr = LevelEditor->GetSceneOutliner())
 			{
 				bCanDelete = SceneOutlinerPtr->Delete_CanExecute();
 			}
 		}
+
+		return bCanDelete;
 	}
 
-	return bCanDelete;
+	return false;
 }
 
 void FLevelEditorActionCallbacks::Rename_Execute()
diff --git a/Engine/Source/Editor/UnrealEd/Private/UnrealEdSrv.cpp b/Engine/Source/Editor/UnrealEd/Private/UnrealEdSrv.cpp
--- a/Engine/Source/Editor/UnrealEd/Private/UnrealEdSrv.cpp
+++ b/Engine/Source/Editor/UnrealEd/Private/UnrealEdSrv.cpp
@@ -2494,7 +2494,7 @@ bool UUnrealEdEngine::Exec_Actor( UWorld* InWorld, const TCHAR* Str, FOutputDevi
 		{
 			if (UTypedElementCommonActions* CommonActions = LevelEditor->GetCommonActions())
 			{
-				UTypedElementSelectionSet* SelectionSet = GetSelectedActors()->GetElementSelectionSet();
+				UTypedElementSelectionSet* SelectionSet = LevelEditor->GetMutableElementSelectionSet();
 
 				TStrongObjectPtr<UTypedElementList> ElementsToDelete(UTypedElementRegistry::GetInstance()->CreateElementList());
 				CommonActions->GetSelectedElementsToDelete(SelectionSet, ElementsToDelete.Get());
