AActor::IsUserManaged() to prevent some type of actors to be deleted, replaced, etc

diff --git a/Engine/Source/Editor/LevelEditor/Private/LevelEditorCreateActorMenu.cpp b/Engine/Source/Editor/LevelEditor/Private/LevelEditorCreateActorMenu.cpp
--- a/Engine/Source/Editor/LevelEditor/Private/LevelEditorCreateActorMenu.cpp
+++ b/Engine/Source/Editor/LevelEditor/Private/LevelEditorCreateActorMenu.cpp
@@ -44,6 +44,8 @@
 #include "Engine/TriggerBase.h"
 #include "LevelEditorMenuContext.h"
 #include "ThumbnailRendering/ThumbnailManager.h"
+#include "Engine/Selection.h"
+#include "ActorEditorUtils.h"
 
 class SMenuThumbnail : public SCompoundWidget
 {
@@ -211,7 +213,22 @@ class SAssetMenuEntry : public SCompoundWidget
 
 static bool CanReplaceActors()
 {
-	return ( GEditor->GetSelectedActorCount() > 0 && !AssetSelectionUtils::IsBuilderBrushSelected() );
+	bool bCanReplace = false;
+
+	for (FSelectionIterator SelectionIter = GEditor->GetSelectedActorIterator(); SelectionIter; ++SelectionIter)
+	{
+		if (AActor* Actor = Cast<AActor>(*SelectionIter))
+		{ 
+			bCanReplace = true;
+			if(!Actor->IsUserManaged() || FActorEditorUtils::IsABuilderBrush(Actor))
+			{
+				bCanReplace = false;
+				break;
+			}
+		}
+	}
+
+	return bCanReplace;
 }
 
 /**
diff --git a/Engine/Source/Editor/UnrealEd/Private/EditorEngine.cpp b/Engine/Source/Editor/UnrealEd/Private/EditorEngine.cpp
--- a/Engine/Source/Editor/UnrealEd/Private/EditorEngine.cpp
+++ b/Engine/Source/Editor/UnrealEd/Private/EditorEngine.cpp
@@ -4997,7 +4997,7 @@ void UEditorEngine::ReplaceSelectedActors(UActorFactory* Factory, const FAssetDa
 	for (FSelectionIterator It = GetSelectedActorIterator(); It; ++It)
 	{
 		AActor* Actor = Cast<AActor>(*It);
-		if ( Actor && !FActorEditorUtils::IsABuilderBrush(Actor) )
+		if ( Actor && Actor->IsUserManaged() && !FActorEditorUtils::IsABuilderBrush(Actor) )
 		{
 			ActorsToReplace.Add(Actor);
 		}
diff --git a/Engine/Source/Runtime/Engine/Classes/GameFramework/Actor.h b/Engine/Source/Runtime/Engine/Classes/GameFramework/Actor.h
--- a/Engine/Source/Runtime/Engine/Classes/GameFramework/Actor.h
+++ b/Engine/Source/Runtime/Engine/Classes/GameFramework/Actor.h
@@ -1842,8 +1842,11 @@ public:
 	virtual void PostTransacted(const FTransactionObjectEvent& TransactionEvent) override;
 	virtual bool IsSelectedInEditor() const override;
 
+	/** Used to know if actor supports some editor operations. (Delete, Replace) */
+	virtual bool IsUserManaged() const { return true; }
+
 	/** When selected can this actor be deleted? */
-	virtual bool CanDeleteSelectedActor(FText& OutReason) const { return true; }
+	virtual bool CanDeleteSelectedActor(FText& OutReason) const;
 
 	/** Does this actor supports external packaging? */
 	virtual bool SupportsExternalPackaging() const;
diff --git a/Engine/Source/Runtime/Engine/Private/Actor.cpp b/Engine/Source/Runtime/Engine/Private/Actor.cpp
--- a/Engine/Source/Runtime/Engine/Private/Actor.cpp
+++ b/Engine/Source/Runtime/Engine/Private/Actor.cpp
@@ -3111,6 +3111,19 @@ static void DispatchOnComponentsCreated(AActor* NewActor)
 }
 
 #if WITH_EDITOR
+
+bool AActor::CanDeleteSelectedActor(FText& OutReason) const
+{
+	if (!IsUserManaged())
+	{
+		OutReason = LOCTEXT("UserManaged", "Actor is not user managed");
+		return false;
+	}
+
+	return true;
+}
+
+
 void AActor::PostEditImport()
 {
 	Super::PostEditImport();
diff --git a/Engine/Source/Runtime/Engine/Private/ActorPartition/PartitionActor.cpp b/Engine/Source/Runtime/Engine/Private/ActorPartition/PartitionActor.cpp
--- a/Engine/Source/Runtime/Engine/Private/ActorPartition/PartitionActor.cpp
+++ b/Engine/Source/Runtime/Engine/Private/ActorPartition/PartitionActor.cpp
@@ -27,9 +27,9 @@ TUniquePtr<FWorldPartitionActorDesc> APartitionActor::CreateClassActorDesc() con
 	return TUniquePtr<FWorldPartitionActorDesc>(new FPartitionActorDesc());
 }
 
-bool APartitionActor::CanDeleteSelectedActor(FText& OutReason) const
+bool APartitionActor::IsUserManaged() const
 {
-	if (!Super::CanDeleteSelectedActor(OutReason))
+	if (!Super::IsUserManaged())
 	{
 		return false;
 	}
@@ -37,7 +37,6 @@ bool APartitionActor::CanDeleteSelectedActor(FText& OutReason) const
 	check(GetLevel());
 	if (GetLevel()->bIsPartitioned)
 	{
-		OutReason = LOCTEXT("CantDeleteSelectedPartitionActor", "Can't delete Partition Actor (automatically managed).");
 		return false;
 	}
 
diff --git a/Engine/Source/Runtime/Engine/Public/ActorPartition/PartitionActor.h b/Engine/Source/Runtime/Engine/Public/ActorPartition/PartitionActor.h
--- a/Engine/Source/Runtime/Engine/Public/ActorPartition/PartitionActor.h
+++ b/Engine/Source/Runtime/Engine/Public/ActorPartition/PartitionActor.h
@@ -24,7 +24,7 @@ public:
 	virtual TUniquePtr<class FWorldPartitionActorDesc> CreateClassActorDesc() const override;
 	virtual uint32 GetDefaultGridSize(UWorld* InWorld) const PURE_VIRTUAL(APartitionActor, return 0;)
 	virtual FGuid GetGridGuid() const { return FGuid(); }
-	virtual bool CanDeleteSelectedActor(FText& OutReason) const override;
+	virtual bool IsUserManaged() const override;
 #endif
 	//~ End AActor Interface	
 
