Optimized GetElementsToMove for large group actors

Each group actor is now only enumerated once at the end, rather than once for each selected actor within a group

diff --git a/Engine/Source/Editor/EditorFramework/Private/Elements/Framework/TypedElementViewportInteraction.cpp b/Engine/Source/Editor/EditorFramework/Private/Elements/Framework/TypedElementViewportInteraction.cpp
--- a/Engine/Source/Editor/EditorFramework/Private/Elements/Framework/TypedElementViewportInteraction.cpp
+++ b/Engine/Source/Editor/EditorFramework/Private/Elements/Framework/TypedElementViewportInteraction.cpp
@@ -5,7 +5,7 @@
 #include "Elements/Framework/TypedElementRegistry.h"
 #include "Elements/Framework/TypedElementUtil.h"
 
-void FTypedElementViewportInteractionCustomization::GetElementsToMove(const TTypedElement<UTypedElementWorldInterface>& InElementWorldHandle, const ETypedElementViewportInteractionWorldType InWorldType, const UTypedElementSelectionSet* InSelectionSet, UTypedElementList* OutElementsToMove)
+void FTypedElementViewportInteractionCustomization::GetElementsToMove(const TTypedElement<UTypedElementWorldInterface>& InElementWorldHandle, const ETypedElementViewportInteractionWorldType InWorldType, const UTypedElementSelectionSet* InSelectionSet, UTypedElementList* OutElementsToMove, FElementToMoveFinalizerMap& OutElementsToMoveFinalizers)
 {
 	OutElementsToMove->Add(InElementWorldHandle);
 }
@@ -139,16 +139,23 @@ void FTypedElementViewportInteractionCustomization::MirrorElement(const TTypedEl
 void UTypedElementViewportInteraction::GetSelectedElementsToMove(const UTypedElementSelectionSet* InSelectionSet, const ETypedElementViewportInteractionWorldType InWorldType, UTypedElementList* OutElementsToMove) const
 {
 	OutElementsToMove->Reset();
-	InSelectionSet->ForEachSelectedElement<UTypedElementWorldInterface>([this, InSelectionSet, InWorldType, OutElementsToMove](const TTypedElement<UTypedElementWorldInterface>& InElementWorldHandle)
+
+	FTypedElementViewportInteractionCustomization::FElementToMoveFinalizerMap ElementsToMoveFinalizers;
+	InSelectionSet->ForEachSelectedElement<UTypedElementWorldInterface>([this, InSelectionSet, InWorldType, OutElementsToMove, &ElementsToMoveFinalizers](const TTypedElement<UTypedElementWorldInterface>& InElementWorldHandle)
 	{
 		if (InElementWorldHandle.CanEditElement())
 		{
 			FTypedElementViewportInteractionElement ViewportInteractionElement(InElementWorldHandle, GetInterfaceCustomizationByTypeId(InElementWorldHandle.GetId().GetTypeId()));
 			check(ViewportInteractionElement.IsSet());
-			ViewportInteractionElement.GetElementsToMove(InWorldType, InSelectionSet, OutElementsToMove);
+			ViewportInteractionElement.GetElementsToMove(InWorldType, InSelectionSet, OutElementsToMove, ElementsToMoveFinalizers);
 		}
 		return true;
 	});
+
+	for (const auto& FinalizerPair : ElementsToMoveFinalizers)
+	{
+		FinalizerPair.Value(FinalizerPair.Key);
+	}
 }
 
 void UTypedElementViewportInteraction::BeginGizmoManipulation(const UTypedElementList* InElementsToMove, const UE::Widget::EWidgetMode InWidgetMode)
diff --git a/Engine/Source/Editor/EditorFramework/Public/Elements/Framework/TypedElementViewportInteraction.h b/Engine/Source/Editor/EditorFramework/Public/Elements/Framework/TypedElementViewportInteraction.h
--- a/Engine/Source/Editor/EditorFramework/Public/Elements/Framework/TypedElementViewportInteraction.h
+++ b/Engine/Source/Editor/EditorFramework/Public/Elements/Framework/TypedElementViewportInteraction.h
@@ -25,8 +25,11 @@ class EDITORFRAMEWORK_API FTypedElementViewportInteractionCustomization
 public:
 	virtual ~FTypedElementViewportInteractionCustomization() = default;
 
+	using FElementToMoveFinalizerFunc = TFunction<void(const FTypedElementHandle&)>;
+	using FElementToMoveFinalizerMap  = TMap<FTypedElementHandle, FElementToMoveFinalizerFunc>;
+
 	//~ See UTypedElementViewportInteraction for API docs
-	virtual void GetElementsToMove(const TTypedElement<UTypedElementWorldInterface>& InElementWorldHandle, const ETypedElementViewportInteractionWorldType InWorldType, const UTypedElementSelectionSet* InSelectionSet, UTypedElementList* OutElementsToMove);
+	virtual void GetElementsToMove(const TTypedElement<UTypedElementWorldInterface>& InElementWorldHandle, const ETypedElementViewportInteractionWorldType InWorldType, const UTypedElementSelectionSet* InSelectionSet, UTypedElementList* OutElementsToMove, FElementToMoveFinalizerMap& OutElementsToMoveFinalizers);
 	virtual bool GetGizmoPivotLocation(const TTypedElement<UTypedElementWorldInterface>& InElementWorldHandle, const UE::Widget::EWidgetMode InWidgetMode, FVector& OutPivotLocation);
 	virtual void PreGizmoManipulationStarted(TArrayView<const FTypedElementHandle> InElementHandles, const UE::Widget::EWidgetMode InWidgetMode);
 	virtual void GizmoManipulationStarted(const TTypedElement<UTypedElementWorldInterface>& InElementWorldHandle, const UE::Widget::EWidgetMode InWidgetMode);
@@ -68,7 +71,7 @@ public:
 	}
 
 	//~ See UTypedElementViewportInteraction for API docs
-	void GetElementsToMove(const ETypedElementViewportInteractionWorldType InWorldType, const UTypedElementSelectionSet* InSelectionSet, UTypedElementList* OutElementsToMove) { ViewportInteractionCustomization->GetElementsToMove(ElementWorldHandle, InWorldType, InSelectionSet, OutElementsToMove); }
+	void GetElementsToMove(const ETypedElementViewportInteractionWorldType InWorldType, const UTypedElementSelectionSet* InSelectionSet, UTypedElementList* OutElementsToMove, FTypedElementViewportInteractionCustomization::FElementToMoveFinalizerMap& OutElementsToMoveFinalizers) { ViewportInteractionCustomization->GetElementsToMove(ElementWorldHandle, InWorldType, InSelectionSet, OutElementsToMove, OutElementsToMoveFinalizers); }
 	bool GetGizmoPivotLocation(const UE::Widget::EWidgetMode InWidgetMode, FVector& OutPivotLocation) const { return ViewportInteractionCustomization->GetGizmoPivotLocation(ElementWorldHandle, InWidgetMode, OutPivotLocation); }
 	void GizmoManipulationStarted(const UE::Widget::EWidgetMode InWidgetMode) const { ViewportInteractionCustomization->GizmoManipulationStarted(ElementWorldHandle, InWidgetMode); }
 	void GizmoManipulationDeltaUpdate(const UE::Widget::EWidgetMode InWidgetMode, const EAxisList::Type InDragAxis, const FInputDeviceState& InInputState, const FTransform& InDeltaTransform, const FVector& InPivotLocation) const { ViewportInteractionCustomization->GizmoManipulationDeltaUpdate(ElementWorldHandle, InWidgetMode, InDragAxis, InInputState, InDeltaTransform, InPivotLocation); }
diff --git a/Engine/Source/Editor/LevelEditor/Private/Elements/Actor/ActorElementLevelEditorViewportInteractionCustomization.cpp b/Engine/Source/Editor/LevelEditor/Private/Elements/Actor/ActorElementLevelEditorViewportInteractionCustomization.cpp
--- a/Engine/Source/Editor/LevelEditor/Private/Elements/Actor/ActorElementLevelEditorViewportInteractionCustomization.cpp
+++ b/Engine/Source/Editor/LevelEditor/Private/Elements/Actor/ActorElementLevelEditorViewportInteractionCustomization.cpp
@@ -15,7 +15,7 @@
 #include "ActorGroupingUtils.h"
 #include "LevelEditorViewport.h"
 
-void FActorElementLevelEditorViewportInteractionCustomization::GetElementsToMove(const TTypedElement<UTypedElementWorldInterface>& InElementWorldHandle, const ETypedElementViewportInteractionWorldType InWorldType, const UTypedElementSelectionSet* InSelectionSet, UTypedElementList* OutElementsToMove)
+void FActorElementLevelEditorViewportInteractionCustomization::GetElementsToMove(const TTypedElement<UTypedElementWorldInterface>& InElementWorldHandle, const ETypedElementViewportInteractionWorldType InWorldType, const UTypedElementSelectionSet* InSelectionSet, UTypedElementList* OutElementsToMove, FElementToMoveFinalizerMap& OutElementsToMoveFinalizers)
 {
 	AActor* Actor = ActorElementDataUtil::GetActorFromHandleChecked(InElementWorldHandle);
 
@@ -28,7 +28,7 @@ void FActorElementLevelEditorViewportInteractionCustomization::GetElementsToMove
 
 	if (CanMoveActorInViewport(Actor, InWorldType))
 	{
-		AppendActorsToMove(Actor, InSelectionSet, OutElementsToMove);
+		AppendActorsToMove(Actor, InSelectionSet, OutElementsToMove, OutElementsToMoveFinalizers);
 	}
 }
 
@@ -146,12 +146,11 @@ bool FActorElementLevelEditorViewportInteractionCustomization::CanMoveActorInVie
 	return true;
 }
 
-void FActorElementLevelEditorViewportInteractionCustomization::AppendActorsToMove(AActor* InActor, const UTypedElementSelectionSet* InSelectionSet, UTypedElementList* OutElementsToMove)
+void FActorElementLevelEditorViewportInteractionCustomization::AppendActorsToMove(AActor* InActor, const UTypedElementSelectionSet* InSelectionSet, UTypedElementList* OutElementsToMove, FElementToMoveFinalizerMap& OutElementsToMoveFinalizers)
 {
-	auto AddActorElement = [&OutElementsToMove](AActor* InActorToAdd)
+	auto AddActorElement = [OutElementsToMove](AActor* InActorToAdd)
 	{
-		FTypedElementHandle ActorElementHandle = UEngineElementsLibrary::AcquireEditorActorElementHandle(InActorToAdd);
-		if (ActorElementHandle)
+		if(FTypedElementHandle ActorElementHandle = UEngineElementsLibrary::AcquireEditorActorElementHandle(InActorToAdd))
 		{
 			OutElementsToMove->Add(MoveTemp(ActorElementHandle));
 		}
@@ -160,7 +159,16 @@ void FActorElementLevelEditorViewportInteractionCustomization::AppendActorsToMov
 	AGroupActor* ParentGroup = AGroupActor::GetRootForActor(InActor, true, true);
 	if (ParentGroup && UActorGroupingUtils::IsGroupingActive())
 	{
-		ParentGroup->ForEachMovableActorInGroup(InSelectionSet, AddActorElement);
+		// Defer group enumeration until the finalization phase, so that each group is 
+		// enumerated once, regardless of how many actors within that group are selected
+		FTypedElementHandle ParentGroupElementHandle = UEngineElementsLibrary::AcquireEditorActorElementHandle(ParentGroup);
+		if (ParentGroupElementHandle && !OutElementsToMoveFinalizers.Contains(ParentGroupElementHandle))
+		{
+			OutElementsToMoveFinalizers.Add(ParentGroupElementHandle, [ParentGroup, InSelectionSet, AddActorElement](const FTypedElementHandle&)
+			{
+				ParentGroup->ForEachMovableActorInGroup(InSelectionSet, AddActorElement);
+			});
+		}
 	}
 	else
 	{
diff --git a/Engine/Source/Editor/LevelEditor/Private/Elements/Component/ComponentElementLevelEditorViewportInteractionCustomization.cpp b/Engine/Source/Editor/LevelEditor/Private/Elements/Component/ComponentElementLevelEditorViewportInteractionCustomization.cpp
--- a/Engine/Source/Editor/LevelEditor/Private/Elements/Component/ComponentElementLevelEditorViewportInteractionCustomization.cpp
+++ b/Engine/Source/Editor/LevelEditor/Private/Elements/Component/ComponentElementLevelEditorViewportInteractionCustomization.cpp
@@ -12,7 +12,7 @@
 #include "LevelEditorViewport.h"
 #include "Kismet2/ComponentEditorUtils.h"
 
-void FComponentElementLevelEditorViewportInteractionCustomization::GetElementsToMove(const TTypedElement<UTypedElementWorldInterface>& InElementWorldHandle, const ETypedElementViewportInteractionWorldType InWorldType, const UTypedElementSelectionSet* InSelectionSet, UTypedElementList* OutElementsToMove)
+void FComponentElementLevelEditorViewportInteractionCustomization::GetElementsToMove(const TTypedElement<UTypedElementWorldInterface>& InElementWorldHandle, const ETypedElementViewportInteractionWorldType InWorldType, const UTypedElementSelectionSet* InSelectionSet, UTypedElementList* OutElementsToMove, FElementToMoveFinalizerMap& OutElementsToMoveFinalizers)
 {
 	UActorComponent* Component = ComponentElementDataUtil::GetComponentFromHandleChecked(InElementWorldHandle);
 
@@ -33,7 +33,7 @@ void FComponentElementLevelEditorViewportInteractionCustomization::GetElementsTo
 				if (ComponentOwner->GetRootComponent() == SceneComponent)
 				{
 					// If it is a root component, use the parent actor instead
-					FActorElementLevelEditorViewportInteractionCustomization::AppendActorsToMove(ComponentOwner, InSelectionSet, OutElementsToMove);
+					FActorElementLevelEditorViewportInteractionCustomization::AppendActorsToMove(ComponentOwner, InSelectionSet, OutElementsToMove, OutElementsToMoveFinalizers);
 				}
 				else
 				{
diff --git a/Engine/Source/Editor/LevelEditor/Public/Elements/Actor/ActorElementLevelEditorViewportInteractionCustomization.h b/Engine/Source/Editor/LevelEditor/Public/Elements/Actor/ActorElementLevelEditorViewportInteractionCustomization.h
--- a/Engine/Source/Editor/LevelEditor/Public/Elements/Actor/ActorElementLevelEditorViewportInteractionCustomization.h
+++ b/Engine/Source/Editor/LevelEditor/Public/Elements/Actor/ActorElementLevelEditorViewportInteractionCustomization.h
@@ -8,14 +8,14 @@
 class LEVELEDITOR_API FActorElementLevelEditorViewportInteractionCustomization : public FActorElementEditorViewportInteractionCustomization, public FTypedElementAssetEditorLevelEditorViewportClientMixin
 {
 public:
-	virtual void GetElementsToMove(const TTypedElement<UTypedElementWorldInterface>& InElementWorldHandle, const ETypedElementViewportInteractionWorldType InWorldType, const UTypedElementSelectionSet* InSelectionSet, UTypedElementList* OutElementsToMove) override;
+	virtual void GetElementsToMove(const TTypedElement<UTypedElementWorldInterface>& InElementWorldHandle, const ETypedElementViewportInteractionWorldType InWorldType, const UTypedElementSelectionSet* InSelectionSet, UTypedElementList* OutElementsToMove, FElementToMoveFinalizerMap& OutElementsToMoveFinalizers) override;
 	virtual void GizmoManipulationStarted(const TTypedElement<UTypedElementWorldInterface>& InElementWorldHandle, const UE::Widget::EWidgetMode InWidgetMode) override;
 	virtual void GizmoManipulationDeltaUpdate(const TTypedElement<UTypedElementWorldInterface>& InElementWorldHandle, const UE::Widget::EWidgetMode InWidgetMode, const EAxisList::Type InDragAxis, const FInputDeviceState& InInputState, const FTransform& InDeltaTransform, const FVector& InPivotLocation) override;
 	virtual void GizmoManipulationStopped(const TTypedElement<UTypedElementWorldInterface>& InElementWorldHandle, const UE::Widget::EWidgetMode InWidgetMode) override;
 	virtual void PostGizmoManipulationStopped(TArrayView<const FTypedElementHandle> InElementHandles, const UE::Widget::EWidgetMode InWidgetMode) override;
 
 	static bool CanMoveActorInViewport(const AActor* InActor, const ETypedElementViewportInteractionWorldType InWorldType);
-	static void AppendActorsToMove(AActor* InActor, const UTypedElementSelectionSet* InSelectionSet, UTypedElementList* OutElementsToMove);
+	static void AppendActorsToMove(AActor* InActor, const UTypedElementSelectionSet* InSelectionSet, UTypedElementList* OutElementsToMove, FElementToMoveFinalizerMap& OutElementsToMoveFinalizers);
 
 private:
 	void ModifyScale(AActor* InActor, const EAxisList::Type InDragAxis, FVector& ScaleDelta, bool bCheckSmallExtent);
diff --git a/Engine/Source/Editor/LevelEditor/Public/Elements/Component/ComponentElementLevelEditorViewportInteractionCustomization.h b/Engine/Source/Editor/LevelEditor/Public/Elements/Component/ComponentElementLevelEditorViewportInteractionCustomization.h
--- a/Engine/Source/Editor/LevelEditor/Public/Elements/Component/ComponentElementLevelEditorViewportInteractionCustomization.h
+++ b/Engine/Source/Editor/LevelEditor/Public/Elements/Component/ComponentElementLevelEditorViewportInteractionCustomization.h
@@ -10,7 +10,7 @@ class FLevelEditorViewportClient;
 class LEVELEDITOR_API FComponentElementLevelEditorViewportInteractionCustomization : public FComponentElementEditorViewportInteractionCustomization, public FTypedElementAssetEditorLevelEditorViewportClientMixin
 {
 public:
-	virtual void GetElementsToMove(const TTypedElement<UTypedElementWorldInterface>& InElementWorldHandle, const ETypedElementViewportInteractionWorldType InWorldType, const UTypedElementSelectionSet* InSelectionSet, UTypedElementList* OutElementsToMove) override;
+	virtual void GetElementsToMove(const TTypedElement<UTypedElementWorldInterface>& InElementWorldHandle, const ETypedElementViewportInteractionWorldType InWorldType, const UTypedElementSelectionSet* InSelectionSet, UTypedElementList* OutElementsToMove, FElementToMoveFinalizerMap& OutElementsToMoveFinalizers) override;
 	virtual void GizmoManipulationStarted(const TTypedElement<UTypedElementWorldInterface>& InElementWorldHandle, const UE::Widget::EWidgetMode InWidgetMode) override;
 	virtual void GizmoManipulationDeltaUpdate(const TTypedElement<UTypedElementWorldInterface>& InElementWorldHandle, const UE::Widget::EWidgetMode InWidgetMode, const EAxisList::Type InDragAxis, const FInputDeviceState& InInputState, const FTransform& InDeltaTransform, const FVector& InPivotLocation) override;
 	virtual void GizmoManipulationStopped(const TTypedElement<UTypedElementWorldInterface>& InElementWorldHandle, const UE::Widget::EWidgetMode InWidgetMode) override;
