Don't show the change level viewport widget for partitioned worlds.

diff --git a/Engine/Source/Editor/LevelEditor/Private/SLevelViewport.cpp b/Engine/Source/Editor/LevelEditor/Private/SLevelViewport.cpp
--- a/Engine/Source/Editor/LevelEditor/Private/SLevelViewport.cpp
+++ b/Engine/Source/Editor/LevelEditor/Private/SLevelViewport.cpp
@@ -79,6 +79,7 @@
 #include "Editor/EditorPerformanceSettings.h"
 #include "UnrealWidget.h"
 #include "WorldPartition/WorldPartitionSubsystem.h"
+#include "WorldPartition/WorldPartition.h"
 #include "WorldPartition/DataLayer/WorldDataLayers.h"
 #include "WorldPartition/DataLayer/DataLayer.h"
 #include "DataLayer/DataLayerEditorSubsystem.h"
@@ -3910,6 +3911,7 @@ EVisibility SLevelViewport::GetCurrentLevelTextVisibility() const
 	return (&GetLevelViewportClient() == GCurrentLevelEditingViewportClient) 
 		&& !IsPlayInEditorViewportActive() 
 		&& GetWorld() && GetWorld()->GetCurrentLevel()->OwningWorld->GetLevels().Num() > 1
+		&& !GetWorld()->GetWorldPartition()
 		?  ContentVisibility : EVisibility::Collapsed;
 }
 
