Editor: Added utility functions for viewports

GetViewports, GetActiveViewport, SetAllowsCinematicControl, GetAllowsCinematicControl.

PilotLevelActor, EjectPilotLevelActor, EditorSetGameView, EditorGetGameView now take an optional config key to indicate which viewport should be set/get. Default is none, which acts on the active viewport (no change in existing behavior)

diff --git a/Engine/Source/Editor/LevelEditor/Private/LevelEditorSubsystem.cpp b/Engine/Source/Editor/LevelEditor/Private/LevelEditorSubsystem.cpp
--- a/Engine/Source/Editor/LevelEditor/Private/LevelEditorSubsystem.cpp
+++ b/Engine/Source/Editor/LevelEditor/Private/LevelEditorSubsystem.cpp
@@ -21,6 +21,33 @@ DEFINE_LOG_CATEGORY_STATIC(LevelEditorSubsystem, Log, All);
 
 #define LOCTEXT_NAMESPACE "LevelEditorSubsystem"
 
+namespace InternalEditorLevelLibrary
+{
+
+TSharedPtr<SLevelViewport> GetLevelViewport(const FName& ViewportConfigKey)
+{
+	FLevelEditorModule& LevelEditorModule = FModuleManager::GetModuleChecked<FLevelEditorModule>("LevelEditor");
+
+	TSharedPtr<ILevelEditor> LevelEditor = LevelEditorModule.GetFirstLevelEditor();
+	if (LevelEditor.IsValid())
+	{
+		for (TSharedPtr<SLevelViewport> LevelViewport : LevelEditor->GetViewports())
+		{
+			if (LevelViewport.IsValid() && LevelViewport->GetConfigKey() == ViewportConfigKey)
+			{
+				return LevelViewport;
+			}
+		}
+
+		TSharedPtr<SLevelViewport> ActiveLevelViewport = LevelEditor->GetActiveViewportInterface();
+		return ActiveLevelViewport;
+	}
+
+	return nullptr;
+}
+
+}
+
 void ULevelEditorSubsystem::Initialize(FSubsystemCollectionBase& Collection)
 {
 	UToolMenus::RegisterStartupCallback(FSimpleMulticastDelegate::FDelegate::CreateUObject(this, &ULevelEditorSubsystem::ExtendQuickActionMenu));
@@ -67,14 +94,12 @@ void ULevelEditorSubsystem::PilotLevelActor(const FToolMenuContext& InContext)
 	PilotLevelActor(SelectedActor);
 }
 
-void ULevelEditorSubsystem::PilotLevelActor(AActor* ActorToPilot) 
+void ULevelEditorSubsystem::PilotLevelActor(AActor* ActorToPilot, FName ViewportConfigKey) 
 {
-	FLevelEditorModule& LevelEditorModule = FModuleManager::GetModuleChecked<FLevelEditorModule>("LevelEditor");
-
-	TSharedPtr<SLevelViewport> ActiveLevelViewport = LevelEditorModule.GetFirstActiveLevelViewport();
-	if (ActiveLevelViewport.IsValid())
+	TSharedPtr<SLevelViewport> LevelViewport = InternalEditorLevelLibrary::GetLevelViewport(ViewportConfigKey);
+	if (LevelViewport.IsValid())
 	{
-		FLevelEditorViewportClient& LevelViewportClient = ActiveLevelViewport->GetLevelViewportClient();
+		FLevelEditorViewportClient& LevelViewportClient = LevelViewport->GetLevelViewportClient();
 
 		LevelViewportClient.SetActorLock(ActorToPilot);
 		if (LevelViewportClient.IsPerspective() && LevelViewportClient.GetActiveActorLock().IsValid())
@@ -84,14 +109,12 @@ void ULevelEditorSubsystem::PilotLevelActor(AActor* ActorToPilot)
 	}
 }
 
-void ULevelEditorSubsystem::EjectPilotLevelActor()
+void ULevelEditorSubsystem::EjectPilotLevelActor(FName ViewportConfigKey)
 {
-	FLevelEditorModule& LevelEditorModule = FModuleManager::GetModuleChecked<FLevelEditorModule>("LevelEditor");
-
-	TSharedPtr<SLevelViewport> ActiveLevelViewport = LevelEditorModule.GetFirstActiveLevelViewport();
-	if (ActiveLevelViewport.IsValid())
+	TSharedPtr<SLevelViewport> LevelViewport = InternalEditorLevelLibrary::GetLevelViewport(ViewportConfigKey);
+	if (LevelViewport.IsValid())
 	{
-		FLevelEditorViewportClient& LevelViewportClient = ActiveLevelViewport->GetLevelViewportClient();
+		FLevelEditorViewportClient& LevelViewportClient = LevelViewport->GetLevelViewportClient();
 
 		if (AActor* LockedActor = LevelViewportClient.GetActiveActorLock().Get())
 		{
@@ -137,20 +160,29 @@ void ULevelEditorSubsystem::EditorInvalidateViewports()
 	}
 }
 
-void ULevelEditorSubsystem::EditorSetGameView(bool bGameView)
+void ULevelEditorSubsystem::EditorSetGameView(bool bGameView, FName ViewportConfigKey)
 {
-	FLevelEditorModule& LevelEditorModule = FModuleManager::GetModuleChecked<FLevelEditorModule>("LevelEditor");
-
-	TSharedPtr<IAssetViewport> ActiveLevelViewport = LevelEditorModule.GetFirstActiveViewport();
-	if (ActiveLevelViewport.IsValid())
+	TSharedPtr<SLevelViewport> LevelViewport = InternalEditorLevelLibrary::GetLevelViewport(ViewportConfigKey);
+	if (LevelViewport.IsValid())
 	{
-		if (ActiveLevelViewport->IsInGameView() != bGameView)
+		if (LevelViewport->IsInGameView() != bGameView)
 		{
-			ActiveLevelViewport->ToggleGameView();
+			LevelViewport->ToggleGameView();
 		}
 	}
 }
 
+bool ULevelEditorSubsystem::EditorGetGameView(FName ViewportConfigKey)
+{
+	TSharedPtr<SLevelViewport> LevelViewport = InternalEditorLevelLibrary::GetLevelViewport(ViewportConfigKey);
+	if (LevelViewport.IsValid())
+	{
+		return LevelViewport->IsInGameView();
+	}
+	return false;
+}
+
+
 void ULevelEditorSubsystem::EditorRequestEndPlay()
 {
 	GUnrealEd->RequestEndPlayMap();
@@ -161,6 +193,61 @@ bool ULevelEditorSubsystem::IsInPlayInEditor() const
 	return GUnrealEd->IsPlayingSessionInEditor();
 }
 
+TArray<FName> ULevelEditorSubsystem::GetViewportConfigKeys()
+{
+	TArray<FName> ViewportConfigKeys;
+	FLevelEditorModule& LevelEditorModule = FModuleManager::GetModuleChecked<FLevelEditorModule>("LevelEditor");
+
+	TSharedPtr<ILevelEditor> LevelEditor = LevelEditorModule.GetFirstLevelEditor();
+	if (LevelEditor.IsValid())
+	{
+		for (TSharedPtr<SLevelViewport> LevelViewport : LevelEditor->GetViewports())
+		{
+			if (LevelViewport.IsValid())
+			{
+				ViewportConfigKeys.Add(LevelViewport->GetConfigKey());
+			}
+		}
+	}
+
+	return ViewportConfigKeys;
+}
+
+FName ULevelEditorSubsystem::GetActiveViewportConfigKey()
+{
+	FLevelEditorModule& LevelEditorModule = FModuleManager::GetModuleChecked<FLevelEditorModule>("LevelEditor");
+
+	TSharedPtr<ILevelEditor> LevelEditor = LevelEditorModule.GetFirstLevelEditor();
+	if (LevelEditor.IsValid())
+	{
+		TSharedPtr<SLevelViewport> ActiveLevelViewport = LevelEditor->GetActiveViewportInterface();
+		if (ActiveLevelViewport.IsValid())
+		{
+			return ActiveLevelViewport->GetConfigKey();
+		}
+	}
+	return NAME_None;
+}
+
+void ULevelEditorSubsystem::SetAllowsCinematicControl(bool bAllow, FName ViewportConfigKey)
+{
+	TSharedPtr<SLevelViewport> LevelViewport = InternalEditorLevelLibrary::GetLevelViewport(ViewportConfigKey);
+	if (LevelViewport.IsValid())
+	{
+		return LevelViewport->SetAllowsCinematicControl(bAllow);
+	}
+}
+
+bool ULevelEditorSubsystem::GetAllowsCinematicControl(FName ViewportConfigKey)
+{
+	TSharedPtr<SLevelViewport> LevelViewport = InternalEditorLevelLibrary::GetLevelViewport(ViewportConfigKey);
+	if (LevelViewport.IsValid())
+	{
+		return LevelViewport->GetAllowsCinematicControl();
+	}
+	return false;
+}
+
 /**
  *
  * Editor Scripting | Level
diff --git a/Engine/Source/Editor/LevelEditor/Private/SLevelViewport.cpp b/Engine/Source/Editor/LevelEditor/Private/SLevelViewport.cpp
--- a/Engine/Source/Editor/LevelEditor/Private/SLevelViewport.cpp
+++ b/Engine/Source/Editor/LevelEditor/Private/SLevelViewport.cpp
@@ -2535,6 +2535,16 @@ bool SLevelViewport::IsLockedCameraViewEnabled() const
 	return LevelViewportClient->bLockedCameraView;
 }
 
+void SLevelViewport::SetAllowsCinematicControl(bool bAllow)
+{
+	LevelViewportClient->SetAllowCinematicControl(bAllow);
+}
+
+bool SLevelViewport::GetAllowsCinematicControl() const
+{
+	return LevelViewportClient->AllowsCinematicControl();
+}
+
 void SLevelViewport::FindSelectedInLevelScript()
 {
 	GUnrealEd->FindSelectedActorsInLevelScript();
diff --git a/Engine/Source/Editor/LevelEditor/Public/LevelEditorSubsystem.h b/Engine/Source/Editor/LevelEditor/Public/LevelEditorSubsystem.h
--- a/Engine/Source/Editor/LevelEditor/Public/LevelEditorSubsystem.h
+++ b/Engine/Source/Editor/LevelEditor/Public/LevelEditorSubsystem.h
@@ -25,10 +25,10 @@ public:
 	void ExtendQuickActionMenu();
 
 	UFUNCTION(BlueprintCallable, Category = "Editor Scripting | Level Utility", meta=(DevelopmentOnly))
-	void PilotLevelActor(AActor* ActorToPilot);
+	void PilotLevelActor(AActor* ActorToPilot, FName ViewportConfigKey = NAME_None);
 	void PilotLevelActor(const FToolMenuContext& InContext);
 	UFUNCTION(BlueprintCallable, Category = "Editor Scripting | Level Utility", meta = (DevelopmentOnly))
-	void EjectPilotLevelActor();
+	void EjectPilotLevelActor(FName ViewportConfigKey = NAME_None);
 
 	UFUNCTION(BlueprintCallable, Category = "Editor Scripting | Level Utility", meta = (DevelopmentOnly))
 	void EditorPlaySimulate();
@@ -37,7 +37,10 @@ public:
 	void EditorInvalidateViewports();
 
 	UFUNCTION(BlueprintCallable, Category = "Editor Scripting | Level Utility", meta = (DevelopmentOnly))
-	void EditorSetGameView(bool bGameView);
+	void EditorSetGameView(bool bGameView, FName ViewportConfigKey = NAME_None);
+
+	UFUNCTION(BlueprintCallable, Category = "Editor Scripting | Level Utility", meta = (DevelopmentOnly))
+	bool EditorGetGameView(FName ViewportConfigKey = NAME_None);
 
 	UFUNCTION(BlueprintCallable, Category = "Editor Scripting | Level Utility", meta = (DevelopmentOnly))
 	void EditorRequestEndPlay();
@@ -45,6 +48,18 @@ public:
 	UFUNCTION(BlueprintCallable, Category = "Editor Scripting | Level Utility", meta = (DevelopmentOnly))
 	bool IsInPlayInEditor() const;
 
+	UFUNCTION(BlueprintCallable, Category = "Editor Scripting | Level Utility", meta = (DevelopmentOnly))
+	TArray<FName> GetViewportConfigKeys();
+
+	UFUNCTION(BlueprintCallable, Category = "Editor Scripting | Level Utility", meta = (DevelopmentOnly))
+	FName GetActiveViewportConfigKey();
+
+	UFUNCTION(BlueprintCallable, Category = "Editor Scripting | Level Utility", meta = (DevelopmentOnly))
+	void SetAllowsCinematicControl(bool bAllow, FName ViewportConfigKey = NAME_None);
+
+	UFUNCTION(BlueprintCallable, Category = "Editor Scripting | Level Utility", meta = (DevelopmentOnly))
+	bool GetAllowsCinematicControl(FName ViewportConfigKey = NAME_None);
+
 	/**
 	 * Close the current Persistent Level (without saving it). Create a new blank Level and save it. Load the new created level.
 	 * @param	AssetPath		Asset Path of where the level will be saved.
diff --git a/Engine/Source/Editor/LevelEditor/Public/SLevelViewport.h b/Engine/Source/Editor/LevelEditor/Public/SLevelViewport.h
--- a/Engine/Source/Editor/LevelEditor/Public/SLevelViewport.h
+++ b/Engine/Source/Editor/LevelEditor/Public/SLevelViewport.h
@@ -272,6 +272,18 @@ public:
 	 */
 	bool IsLockedCameraViewEnabled() const;
 
+	/**
+	 * Sets whether the viewport should allow cinematic control
+	 * 
+	 * @param Whether the viewport should allow cinematic control.
+     */
+	void SetAllowsCinematicControl(bool bAllow);
+
+	/**
+	 * @return Whether the viewport allows cinematic control.
+     */
+	bool GetAllowsCinematicControl() const;
+
 	/**
 	 * @return the fixed width that a column returned by CreateActorLockSceneOutlinerColumn expects to be
 	 */
