Avoid having the cursor decoration follow you everywhere on the desktop for the entire duration of a drop operation that might require a long async compilation

diff --git a/Engine/Source/Editor/LevelEditor/Private/SLevelViewport.cpp b/Engine/Source/Editor/LevelEditor/Private/SLevelViewport.cpp
--- a/Engine/Source/Editor/LevelEditor/Private/SLevelViewport.cpp
+++ b/Engine/Source/Editor/LevelEditor/Private/SLevelViewport.cpp
@@ -666,9 +666,13 @@ void SLevelViewport::OnDragEnter( const FGeometry& MyGeometry, const FDragDropEv
 		{
 			if ( HandleDragObjects(MyGeometry, DragDropEvent) )
 			{
-				if ( HandlePlaceDraggedObjects(MyGeometry, DragDropEvent, /*bCreateDropPreview=*/true) )
+				// Hide the decorator before dropping the object to avoid having a decorator present for the 
+				// entire duration of an async asset build if required. 
+				// Restore the decorator visibility if the drop fails to preserve previous behavior.
+				DragDropEvent.GetOperation()->SetDecoratorVisibility(false);
+				if ( !HandlePlaceDraggedObjects(MyGeometry, DragDropEvent, /*bCreateDropPreview=*/true) )
 				{
-					DragDropEvent.GetOperation()->SetDecoratorVisibility(false);
+					DragDropEvent.GetOperation()->SetDecoratorVisibility(true);
 				}
 			}
 		}
