Preventing an issue where users could partially convert collections of actor types to Blueprints, including unconvertible types like Brush/BSPs. We now detect potential problems and give them the option to proceed without the unconvertible actors.

diff --git a/Engine/Source/Editor/LevelEditor/Private/LevelEditorActions.cpp b/Engine/Source/Editor/LevelEditor/Private/LevelEditorActions.cpp
--- a/Engine/Source/Editor/LevelEditor/Private/LevelEditorActions.cpp
+++ b/Engine/Source/Editor/LevelEditor/Private/LevelEditorActions.cpp
@@ -2403,6 +2403,66 @@ void FLevelEditorActionCallbacks::ConvertSelectedActorsIntoBlueprintClass()
 {
 	const ECreateBlueprintFromActorMode ValidCreateModes = FCreateBlueprintFromActorDialog::GetValidCreationMethods();
 	ECreateBlueprintFromActorMode DefaultCreateMode = ECreateBlueprintFromActorMode::Harvest;
+
+	// Check all of the selected actors for any that can't be converted
+	bool bHasAnyValidActors = false;
+	TArray<AActor*> UnconvertibleSelectedActors;
+
+	for (FSelectionIterator It(GEditor->GetSelectedActorIterator()); It; ++It)
+	{
+		if (AActor* Actor = Cast<AActor>(*It))
+		{
+			if (FKismetEditorUtilities::CanCreateBlueprintOfClass(Actor->GetClass()))
+			{
+				bHasAnyValidActors = true;
+			}
+			else
+			{
+				UnconvertibleSelectedActors.Add(Actor);
+			}
+		}
+	}
+
+	if (UnconvertibleSelectedActors.Num())
+	{
+		// Let the user know that some or all of the selected actors are not convertible to BP. 
+		FString UnconvertedActorsList = FString::JoinBy(UnconvertibleSelectedActors, TEXT("\n"), [](const AActor* Actor) 
+			{
+				return FString::Format(TEXT("{0} (type '{1}')"), { Actor->GetName(), Actor->GetClass()->GetName() });
+			});
+
+		if (bHasAnyValidActors)
+		{
+			// If there are some convertible actors, give the user a choice to proceed with only the valid ones
+			const FText Message = FText::Format(LOCTEXT("ConfirmPartialConversionToBlueprint", "These selected actors cannot be used to create a blueprint. Do you want to continue conversion without them?\n\n{0}"), FText::FromString(UnconvertedActorsList));
+			if (FMessageDialog::Open(EAppMsgType::YesNo, EAppReturnType::No, Message) == EAppReturnType::No)
+			{
+				return;
+			}
+		}
+		else
+		{
+			// There are no convertible actors. Just let the user know and bail.
+			const FText Message = FText::Format(LOCTEXT("SelectedActorsCannotBeBlueprint", "No selected actors can be used to create a blueprint:\n\n{0}"), FText::FromString(UnconvertedActorsList));
+			FMessageDialog::Open(EAppMsgType::Ok, Message);
+			return;
+		}
+
+		// Deselect the unconvertible actors and clear any surface selection (common with unconvertible Brush actors)
+		GEditor->GetSelectedActors()->BeginBatchSelectOperation();
+		GEditor->DeselectAllSurfaces();
+
+		const bool bShouldSelect = false;
+		const bool bShouldNotify = false;
+		for (AActor* UnconvertibleActor : UnconvertibleSelectedActors)
+		{
+			GEditor->SelectActor(UnconvertibleActor, bShouldSelect, bShouldNotify);
+		}
+
+		GEditor->GetSelectedActors()->EndBatchSelectOperation(bShouldNotify);
+		GEditor->NoteSelectionChange();
+	}
+
 	if (!!(ValidCreateModes & ECreateBlueprintFromActorMode::Subclass) && (GEditor->GetSelectedActorCount() == 1))
 	{
 		// If a single actor is selected and it can be subclassed, use that as default
