Minimal context menu for static mesh instance elements

Currently only delete and delta transform are enabled/available

diff --git a/Engine/Source/Editor/LevelEditor/Private/LevelEditorActions.cpp b/Engine/Source/Editor/LevelEditor/Private/LevelEditorActions.cpp
--- a/Engine/Source/Editor/LevelEditor/Private/LevelEditorActions.cpp
+++ b/Engine/Source/Editor/LevelEditor/Private/LevelEditorActions.cpp
@@ -3363,10 +3363,10 @@ void FLevelEditorCommands::RegisterCommands()
 	UI_COMMAND( AlignPivotToActor, "Align Pivot to Actor", "Aligns the actor with another actor at its pivot point", EUserInterfaceActionType::Button, FInputChord() );
 	UI_COMMAND( SnapBottomCenterBoundsToActor, "Snap Bottom Center Bounds to Actor", "Snaps the actor to another actor at its bottom center bounds", EUserInterfaceActionType::Button, FInputChord() );
 	UI_COMMAND( AlignBottomCenterBoundsToActor, "Align Bottom Center Bounds to Actor", "Aligns the actor with another actor at its bottom center bounds", EUserInterfaceActionType::Button, FInputChord() );
-	UI_COMMAND( DeltaTransformToActors, "Delta Transform", "Apply Delta Transform to selected actors", EUserInterfaceActionType::Button, FInputChord() );
-	UI_COMMAND( MirrorActorX, "Mirror X", "Mirrors the actor along the X axis", EUserInterfaceActionType::Button, FInputChord() );
-	UI_COMMAND( MirrorActorY, "Mirror Y", "Mirrors the actor along the Y axis", EUserInterfaceActionType::Button, FInputChord() );
-	UI_COMMAND( MirrorActorZ, "Mirror Z", "Mirrors the actor along the Z axis", EUserInterfaceActionType::Button, FInputChord() );
+	UI_COMMAND( DeltaTransformToActors, "Delta Transform", "Apply Delta Transform to selected elements", EUserInterfaceActionType::Button, FInputChord() );
+	UI_COMMAND( MirrorActorX, "Mirror X", "Mirrors the element along the X axis", EUserInterfaceActionType::Button, FInputChord() );
+	UI_COMMAND( MirrorActorY, "Mirror Y", "Mirrors the element along the Y axis", EUserInterfaceActionType::Button, FInputChord() );
+	UI_COMMAND( MirrorActorZ, "Mirror Z", "Mirrors the element along the Z axis", EUserInterfaceActionType::Button, FInputChord() );
 	UI_COMMAND( LockActorMovement, "Lock Actor Movement", "Locks the actor so it cannot be moved", EUserInterfaceActionType::ToggleButton, FInputChord() );
 	UI_COMMAND( DetachFromParent, "Detach", "Detach the actor from its parent", EUserInterfaceActionType::Button, FInputChord() );
 	UI_COMMAND( AttachSelectedActors, "Attach Selected Actors", "Attach the selected actors to the last selected actor", EUserInterfaceActionType::Button, FInputChord(EModifierKey::Alt, EKeys::B) );
diff --git a/Engine/Source/Editor/LevelEditor/Private/LevelEditorContextMenu.cpp b/Engine/Source/Editor/LevelEditor/Private/LevelEditorContextMenu.cpp
--- a/Engine/Source/Editor/LevelEditor/Private/LevelEditorContextMenu.cpp
+++ b/Engine/Source/Editor/LevelEditor/Private/LevelEditorContextMenu.cpp
@@ -44,6 +44,7 @@
 #include "EditorViewportCommands.h"
 #include "Toolkits/GlobalEditorCommonCommands.h"
 #include "LevelEditorCreateActorMenu.h"
+#include "Elements/Interfaces/TypedElementObjectInterface.h"
 #include "SourceCodeNavigation.h"
 #include "EditorClassUtils.h"
 #include "Framework/Commands/GenericCommands.h"
@@ -125,7 +126,8 @@ public:
 	 * @param MenuBuilder	The menu to add items to
 	 * @param ContextType	The context for this editor menu
 	 */
-	static void FillEditMenu(UToolMenu* Menu);
+	static void FillEditMenu(UToolMenu* Menu) { FillEditMenu(Menu, nullptr); }
+	static void FillEditMenu(UToolMenu* Menu, FToolMenuSection* InSection);
 
 	/**
 	 * Fills in menu options for the merge actors menu
@@ -449,6 +451,33 @@ void FLevelEditorContextMenu::RegisterActorContextMenu()
 	}));
 }
 
+void FLevelEditorContextMenu::RegisterElementContextMenu()
+{
+	UToolMenus* ToolMenus = UToolMenus::Get();
+	if (ToolMenus->IsMenuRegistered("LevelEditor.ElementContextMenu"))
+	{
+		return;
+	}
+
+	UToolMenu* Menu = ToolMenus->RegisterMenu("LevelEditor.ElementContextMenu");
+	Menu->AddDynamicSection("ElementContextMenuDynamic", FNewToolMenuDelegate::CreateLambda([](UToolMenu* InMenu)
+	{
+		{
+			FToolMenuSection& Section = InMenu->AddSection("ElementEditActions", LOCTEXT("ElementEditActions", "Edit"));
+			FLevelEditorContextMenuImpl::FillEditMenu(InMenu, &Section);
+		}
+
+		{
+			FToolMenuSection& Section = InMenu->AddSection("ElementLevelActions", LOCTEXT("ElementLevelActions", "Level"));
+			Section.AddSubMenu(
+				"TransformSubMenu",
+				LOCTEXT("TransformSubMenu", "Transform"),
+				LOCTEXT("TransformSubMenu_ToolTip_Element", "Element transform utils"),
+				FNewToolMenuDelegate::CreateStatic(&FLevelEditorContextMenuImpl::FillTransformMenu));
+		}
+	}));
+}
+
 void FLevelEditorContextMenu::RegisterSceneOutlinerContextMenu()
 {
 	UToolMenus* ToolMenus = UToolMenus::Get();
@@ -503,17 +532,27 @@ void FLevelEditorContextMenu::RegisterEmptySelectionContextMenu()
 	}));
 }
 
-FName FLevelEditorContextMenu::GetContextMenuName(ELevelEditorMenuContext ContextType)
+FName FLevelEditorContextMenu::GetContextMenuName(ELevelEditorMenuContext ContextType, const UTypedElementSelectionSet* InSelectionSet)
 {
-	if (GEditor->GetSelectedComponentCount() > 0)
-	{
-		return "LevelEditor.ComponentContextMenu";
-	}
-	else if (GEditor->GetSelectedActorCount() > 0)
+	if (InSelectionSet)
 	{
-		return "LevelEditor.ActorContextMenu";
+		if (InSelectionSet->HasSelectedObjects<UActorComponent>())
+		{
+			return "LevelEditor.ComponentContextMenu";
+		}
+
+		if (InSelectionSet->HasSelectedObjects<AActor>())
+		{
+			return "LevelEditor.ActorContextMenu";
+		}
+
+		if (InSelectionSet->GetNumSelectedElements())
+		{
+			return "LevelEditor.ElementContextMenu";
+		}
 	}
-	else if (ContextType == ELevelEditorMenuContext::SceneOutliner)
+
+	if (ContextType == ELevelEditorMenuContext::SceneOutliner)
 	{
 		return "LevelEditor.SceneOutlinerContextMenu";
 	}
@@ -521,25 +560,39 @@ FName FLevelEditorContextMenu::GetContextMenuName(ELevelEditorMenuContext Contex
 	return "LevelEditor.EmptySelectionContextMenu";
 }
 
-FName FLevelEditorContextMenu::InitMenuContext(FToolMenuContext& Context, TWeakPtr<SLevelEditor> LevelEditor, ELevelEditorMenuContext ContextType, AActor* HitProxyActor)
+FName FLevelEditorContextMenu::InitMenuContext(FToolMenuContext& Context, TWeakPtr<SLevelEditor> LevelEditor, ELevelEditorMenuContext ContextType, const FTypedElementHandle& HitProxyElement)
 {
 	RegisterComponentContextMenu();
 	RegisterActorContextMenu();
+	RegisterElementContextMenu();
 	RegisterSceneOutlinerContextMenu();
 	RegisterEmptySelectionContextMenu();
 
-	TSharedPtr<FUICommandList> LevelEditorActionsList = LevelEditor.Pin()->GetLevelEditorActions();
+	TSharedPtr<SLevelEditor> LevelEditorPtr = LevelEditor.Pin();
+	check(LevelEditorPtr);
+
+	TSharedPtr<FUICommandList> LevelEditorActionsList = LevelEditorPtr->GetLevelEditorActions();
 	Context.AppendCommandList(LevelEditorActionsList);
 
 	ULevelEditorContextMenuContext* ContextObject = NewObject<ULevelEditorContextMenuContext>();
 	ContextObject->LevelEditor = LevelEditor;
 	ContextObject->ContextType = ContextType;
-	ContextObject->HitProxyActor = HitProxyActor;
+	ContextObject->CurrentSelection = LevelEditorPtr->GetElementSelectionSet();
+	ContextObject->HitProxyElement = HitProxyElement;
+	{
+		TTypedElement<UTypedElementObjectInterface> HitProxyObjectElement = ContextObject->CurrentSelection->GetElementList()->GetElement<UTypedElementObjectInterface>(ContextObject->HitProxyElement);
+		ContextObject->HitProxyActor = HitProxyObjectElement ? Cast<AActor>(HitProxyObjectElement.GetObject()) : nullptr;
+	}
 	for (FSelectedEditableComponentIterator It(GEditor->GetSelectedEditableComponentIterator()); It; ++It)
 	{
 		ContextObject->SelectedComponents.Add(CastChecked<UActorComponent>(*It));
 	}
-	Context.AddObject(ContextObject);
+	Context.AddObject(ContextObject, [](UObject* InContext)
+	{
+		ULevelEditorContextMenuContext* CastContext = CastChecked<ULevelEditorContextMenuContext>(InContext);
+		CastContext->CurrentSelection = nullptr;
+		CastContext->HitProxyElement.Release();
+	});
 
 	if (GEditor->GetSelectedComponentCount() == 0 && GEditor->GetSelectedActorCount() > 0)
 	{
@@ -564,10 +617,10 @@ FName FLevelEditorContextMenu::InitMenuContext(FToolMenuContext& Context, TWeakP
 		}
 	}
 
-	return GetContextMenuName(ContextType);
+	return GetContextMenuName(ContextType, ContextObject->CurrentSelection);
 }
 
-UToolMenu* FLevelEditorContextMenu::GenerateMenu(TWeakPtr<SLevelEditor> LevelEditor, ELevelEditorMenuContext ContextType, TSharedPtr<FExtender> Extender, AActor* HitProxyActor)
+UToolMenu* FLevelEditorContextMenu::GenerateMenu(TWeakPtr<SLevelEditor> LevelEditor, ELevelEditorMenuContext ContextType, TSharedPtr<FExtender> Extender, const FTypedElementHandle& HitProxyElement)
 {
 	FToolMenuContext Context;
 	if (Extender.IsValid())
@@ -575,15 +628,15 @@ UToolMenu* FLevelEditorContextMenu::GenerateMenu(TWeakPtr<SLevelEditor> LevelEdi
 		Context.AddExtender(Extender);
 	}
 
-	FName ContextMenuName = InitMenuContext(Context, LevelEditor, ContextType, HitProxyActor);
+	FName ContextMenuName = InitMenuContext(Context, LevelEditor, ContextType, HitProxyElement);
 	return UToolMenus::Get()->GenerateMenu(ContextMenuName, Context);
 }
 
 // NOTE: We intentionally receive a WEAK pointer here because we want to be callable by a delegate whose
 //       payload contains a weak reference to a level editor instance
-TSharedPtr< SWidget > FLevelEditorContextMenu::BuildMenuWidget(TWeakPtr< SLevelEditor > LevelEditor, ELevelEditorMenuContext ContextType, TSharedPtr<FExtender> Extender, AActor* HitProxyActor)
+TSharedPtr< SWidget > FLevelEditorContextMenu::BuildMenuWidget(TWeakPtr< SLevelEditor > LevelEditor, ELevelEditorMenuContext ContextType, TSharedPtr<FExtender> Extender, const FTypedElementHandle& HitProxyElement)
 {
-	UToolMenu* Menu = GenerateMenu(LevelEditor, ContextType, Extender, HitProxyActor);
+	UToolMenu* Menu = GenerateMenu(LevelEditor, ContextType, Extender, HitProxyElement);
 	return UToolMenus::Get()->GenerateWidget(Menu);
 }
 
@@ -688,7 +741,7 @@ void FLevelEditorContextMenu::SummonViewOptionMenu( const TSharedRef< SLevelEdit
 	BuildViewOptionMenu(LevelEditor, MakeViewOptionWidget(LevelEditor, bShouldCloseWindowAfterMenuSelection, ViewOptionType), MouseCursorLocation);
 }
 
-void FLevelEditorContextMenu::SummonMenu(const TSharedRef< SLevelEditor >& LevelEditor, ELevelEditorMenuContext ContextType, AActor* HitProxyActor)
+void FLevelEditorContextMenu::SummonMenu(const TSharedRef< SLevelEditor >& LevelEditor, ELevelEditorMenuContext ContextType, const FTypedElementHandle& HitProxyElement)
 {
 	struct Local
 	{
@@ -712,7 +765,7 @@ void FLevelEditorContextMenu::SummonMenu(const TSharedRef< SLevelEditor >& Level
 	Extender->AddMenuExtension("LevelViewportAttach", EExtensionHook::After, TSharedPtr< FUICommandList >(), FMenuExtensionDelegate::CreateStatic(&Local::ExtendMenu));
 
 	// Create the context menu!
-	TSharedPtr<SWidget> MenuWidget = BuildMenuWidget( LevelEditor, ContextType, Extender, HitProxyActor );
+	TSharedPtr<SWidget> MenuWidget = BuildMenuWidget( LevelEditor, ContextType, Extender, HitProxyElement );
 	if ( MenuWidget.IsValid() )
 	{
 		// @todo: Should actually use the location from a click event instead!
@@ -1037,8 +1090,9 @@ void FLevelEditorContextMenuImpl::FillActorLevelMenu(UToolMenu* Menu)
 
 void FLevelEditorContextMenuImpl::FillTransformMenu(UToolMenu* Menu)
 {
-	if ( FLevelEditorActionCallbacks::ActorSelected_CanExecute() )
+	if ( FLevelEditorActionCallbacks::ElementSelected_CanExecute() )
 	{
+		if ( FLevelEditorActionCallbacks::ActorSelected_CanExecute() )
 		{
 			FToolMenuSection& Section = Menu->AddSection("TransformSnapAlign");
 			Section.AddSubMenu(
@@ -1056,10 +1110,19 @@ void FLevelEditorContextMenuImpl::FillTransformMenu(UToolMenu* Menu)
 
 	{
 		FToolMenuSection& Section = Menu->AddSection("MirrorLock");
-		Section.AddMenuEntry(FLevelEditorCommands::Get().MirrorActorX);
-		Section.AddMenuEntry(FLevelEditorCommands::Get().MirrorActorY);
-		Section.AddMenuEntry(FLevelEditorCommands::Get().MirrorActorZ);
-		Section.AddMenuEntry(FLevelEditorCommands::Get().LockActorMovement);
+
+		// TODO: Need an element API to allow the mirror actions
+		if ( FLevelEditorActionCallbacks::ActorSelected_CanExecute() )
+		{
+			Section.AddMenuEntry(FLevelEditorCommands::Get().MirrorActorX);
+			Section.AddMenuEntry(FLevelEditorCommands::Get().MirrorActorY);
+			Section.AddMenuEntry(FLevelEditorCommands::Get().MirrorActorZ);
+		}
+
+		if ( FLevelEditorActionCallbacks::ActorSelected_CanExecute() )
+		{
+			Section.AddMenuEntry(FLevelEditorCommands::Get().LockActorMovement);
+		}
 	}
 }
 
@@ -1287,9 +1350,9 @@ void FLevelEditorContextMenuImpl::FillGroupMenu( UToolMenu* Menu )
 	}
 }
 
-void FLevelEditorContextMenuImpl::FillEditMenu( UToolMenu* Menu )
+void FLevelEditorContextMenuImpl::FillEditMenu( UToolMenu* Menu, FToolMenuSection* InSection )
 {
-	FToolMenuSection& Section = Menu->AddSection("Section");
+	FToolMenuSection& Section = InSection ? *InSection : Menu->AddSection("Section");
 
 	Section.AddMenuEntry( FGenericCommands::Get().Cut );
 	Section.AddMenuEntry( FGenericCommands::Get().Copy );
diff --git a/Engine/Source/Editor/LevelEditor/Private/LevelEditorContextMenu.h b/Engine/Source/Editor/LevelEditor/Private/LevelEditorContextMenu.h
--- a/Engine/Source/Editor/LevelEditor/Private/LevelEditorContextMenu.h
+++ b/Engine/Source/Editor/LevelEditor/Private/LevelEditorContextMenu.h
@@ -8,6 +8,7 @@
 
 class FExtender;
 class UToolMenu;
+class UTypedElementSelectionSet;
 struct FToolMenuContext;
 class SLevelEditor;
 class SWidget;
@@ -26,7 +27,7 @@ public:
 	 * @param	ContextType		The context we should use to specialize this menu
 	 * @param	HitProxyActor	The hitproxy actor in the case the ContextType is Viewport
 	 */
-	static void SummonMenu( const TSharedRef< class SLevelEditor >& LevelEditor, ELevelEditorMenuContext ContextType, AActor* HitProxyActor = nullptr);
+	static void SummonMenu( const TSharedRef< class SLevelEditor >& LevelEditor, ELevelEditorMenuContext ContextType, const FTypedElementHandle& HitProxyElement = FTypedElementHandle());
 
 	/**
 	 * Summons the viewport view option menu
@@ -43,7 +44,7 @@ public:
 	 * @param	HitProxyActor	The hitproxy actor in the case the ContextType is Viewport
 	 * @return	Widget for this context menu
 	 */
-	static TSharedPtr< SWidget > BuildMenuWidget(TWeakPtr< SLevelEditor > LevelEditor, ELevelEditorMenuContext ContextType, TSharedPtr<FExtender> Extender = TSharedPtr<FExtender>(), AActor* HitProxyActor = nullptr);
+	static TSharedPtr< SWidget > BuildMenuWidget(TWeakPtr< SLevelEditor > LevelEditor, ELevelEditorMenuContext ContextType, TSharedPtr<FExtender> Extender = TSharedPtr<FExtender>(), const FTypedElementHandle& HitProxyElement = FTypedElementHandle());
 
 	/**
 	 * Populates the specified menu builder for the context menu that can be inserted into a pop-up window
@@ -54,18 +55,19 @@ public:
 	 * @param	Extender		Allows extension of this menu based on context.
 	 * @param	HitProxyActor	The hitproxy actor in the case the ContextType is Viewport
 	 */
-	static UToolMenu* GenerateMenu(TWeakPtr< SLevelEditor > LevelEditor, ELevelEditorMenuContext ContextType, TSharedPtr<FExtender> Extender = TSharedPtr<FExtender>(), AActor* HitProxyActor = nullptr);
+	static UToolMenu* GenerateMenu(TWeakPtr< SLevelEditor > LevelEditor, ELevelEditorMenuContext ContextType, TSharedPtr<FExtender> Extender = TSharedPtr<FExtender>(), const FTypedElementHandle& HitProxyElement = FTypedElementHandle());
 
 	/* Adds required information to Context for build menu based on current selection */
-	static FName InitMenuContext(FToolMenuContext& Context, TWeakPtr<SLevelEditor> LevelEditor, ELevelEditorMenuContext ContextType, AActor* HitProxyActor = nullptr);
+	static FName InitMenuContext(FToolMenuContext& Context, TWeakPtr<SLevelEditor> LevelEditor, ELevelEditorMenuContext ContextType, const FTypedElementHandle& HitProxyElement = FTypedElementHandle());
 
 	/* Returns name of menu to display based on current selection */
-	static FName GetContextMenuName(ELevelEditorMenuContext ContextType);
+	static FName GetContextMenuName(ELevelEditorMenuContext ContextType, const UTypedElementSelectionSet* InSelectionSet);
 
 private:
 
 	static void RegisterComponentContextMenu();
 	static void RegisterActorContextMenu();
+	static void RegisterElementContextMenu();
 	static void RegisterSceneOutlinerContextMenu();
 	static void RegisterEmptySelectionContextMenu();
 
diff --git a/Engine/Source/Editor/LevelEditor/Private/SLevelEditor.cpp b/Engine/Source/Editor/LevelEditor/Private/SLevelEditor.cpp
--- a/Engine/Source/Editor/LevelEditor/Private/SLevelEditor.cpp
+++ b/Engine/Source/Editor/LevelEditor/Private/SLevelEditor.cpp
@@ -800,9 +800,9 @@ TSharedRef<SDockTab> SLevelEditor::SpawnLevelEditorTab( const FSpawnTabArgs& Arg
 			if (!ToolMenus->IsMenuRegistered(MenuName))
 			{
 				UToolMenu* Menu = ToolMenus->RegisterMenu(MenuName, "SceneOutliner.DefaultContextMenuBase");
-				FToolMenuSection& Section = Menu->AddDynamicSection("LevelEditorContextMenu", FNewToolMenuDelegate::CreateLambda([](UToolMenu* InMenu)
+				FToolMenuSection& Section = Menu->AddDynamicSection("LevelEditorContextMenu", FNewToolMenuDelegate::CreateLambda([SelectionSet = TWeakObjectPtr<const UTypedElementSelectionSet>(GetElementSelectionSet())](UToolMenu* InMenu)
 				{
-					FName LevelContextMenuName = FLevelEditorContextMenu::GetContextMenuName(ELevelEditorMenuContext::SceneOutliner);
+					FName LevelContextMenuName = FLevelEditorContextMenu::GetContextMenuName(ELevelEditorMenuContext::SceneOutliner, SelectionSet.Get());
 					if (LevelContextMenuName != NAME_None)
 					{
 						// Extend the menu even if no actors selected, as Edit menu should always exist for scene outliner
@@ -1754,9 +1754,9 @@ UTypedElementSelectionSet* SLevelEditor::GetMutableElementSelectionSet()
 	return SelectedElements;
 }
 
-void SLevelEditor::SummonLevelViewportContextMenu(AActor* HitProxyActor)
+void SLevelEditor::SummonLevelViewportContextMenu(const FTypedElementHandle& HitProxyElement)
 {
-	FLevelEditorContextMenu::SummonMenu( SharedThis( this ), ELevelEditorMenuContext::Viewport, HitProxyActor);
+	FLevelEditorContextMenu::SummonMenu( SharedThis( this ), ELevelEditorMenuContext::Viewport, HitProxyElement);
 }
 
 void SLevelEditor::SummonLevelViewportViewOptionMenu(const ELevelViewportType ViewOption)
diff --git a/Engine/Source/Editor/LevelEditor/Private/SLevelEditor.h b/Engine/Source/Editor/LevelEditor/Private/SLevelEditor.h
--- a/Engine/Source/Editor/LevelEditor/Private/SLevelEditor.h
+++ b/Engine/Source/Editor/LevelEditor/Private/SLevelEditor.h
@@ -67,7 +67,7 @@ public:
 	/** ILevelEditor interface */
 	virtual const UTypedElementSelectionSet* GetElementSelectionSet() const override;
 	virtual UTypedElementSelectionSet* GetMutableElementSelectionSet() override;
-	virtual void SummonLevelViewportContextMenu(AActor* HitProxyActor = nullptr) override;
+	virtual void SummonLevelViewportContextMenu(const FTypedElementHandle& HitProxyElement = FTypedElementHandle()) override;
 	virtual void SummonLevelViewportViewOptionMenu(const ELevelViewportType ViewOption) override;
 	virtual const TArray< TSharedPtr< class IToolkit > >& GetHostedToolkits() const override;
 	virtual TArray< TSharedPtr< SLevelViewport > > GetViewports() const override;
diff --git a/Engine/Source/Editor/LevelEditor/Public/ILevelEditor.h b/Engine/Source/Editor/LevelEditor/Public/ILevelEditor.h
--- a/Engine/Source/Editor/LevelEditor/Public/ILevelEditor.h
+++ b/Engine/Source/Editor/LevelEditor/Public/ILevelEditor.h
@@ -7,6 +7,8 @@
 #include "Widgets/SWidget.h"
 #include "Widgets/SCompoundWidget.h"
 #include "Editor/UnrealEdTypes.h"
+#include "Elements/Framework/TypedElementHandle.h"
+#include "Elements/Framework/EngineElementsLibrary.h"
 #include "Framework/Docking/TabManager.h"
 #include "Framework/Commands/UICommandList.h"
 #include "AssetThumbnail.h"
@@ -32,7 +34,18 @@ public:
 	virtual UTypedElementSelectionSet* GetMutableElementSelectionSet() = 0;
 
 	/** Summons a context menu for this level editor at the mouse cursor's location */
-	virtual void SummonLevelViewportContextMenu(AActor* HitProxyActor = nullptr) = 0;
+	virtual void SummonLevelViewportContextMenu(const FTypedElementHandle& HitProxyElement = FTypedElementHandle()) = 0;
+
+	UE_DEPRECATED(5.0, "Use the SummonLevelViewportContextMenu overload which takes a FTypedElementHandle.")
+	void SummonLevelViewportContextMenu(AActor* HitProxyActor)
+	{
+		FTypedElementHandle HitProxyElement;
+		if (HitProxyActor)
+		{
+			HitProxyElement = UEngineElementsLibrary::AcquireEditorActorElementHandle(HitProxyActor);
+		}
+		SummonLevelViewportContextMenu(HitProxyElement);
+	}
 
 	/** Summons a context menu for view options */
 	virtual void SummonLevelViewportViewOptionMenu(ELevelViewportType ViewOption) = 0;
diff --git a/Engine/Source/Editor/LevelEditor/Public/LevelEditorMenuContext.h b/Engine/Source/Editor/LevelEditor/Public/LevelEditorMenuContext.h
--- a/Engine/Source/Editor/LevelEditor/Public/LevelEditorMenuContext.h
+++ b/Engine/Source/Editor/LevelEditor/Public/LevelEditorMenuContext.h
@@ -3,8 +3,8 @@
 #pragma once
 
 #include "CoreMinimal.h"
-#include "UObject/ObjectMacros.h"
 #include "UObject/Object.h"
+#include "Elements/Framework/TypedElementHandle.h"
 
 #include "LevelEditorMenuContext.generated.h"
 
@@ -40,8 +40,16 @@ class LEVELEDITOR_API ULevelEditorContextMenuContext : public UObject
 public:
 
 	TWeakPtr<SLevelEditor> LevelEditor;
+
 	ELevelEditorMenuContext ContextType;
 
+	UPROPERTY(VisibleAnywhere, BlueprintReadOnly, Category="LevelEditor|Menu")
+	TObjectPtr<const UTypedElementSelectionSet> CurrentSelection;
+	
+	/** If the ContextType is Viewport this property can be set to the HitProxy element that triggered the ContextMenu. */
+	UPROPERTY(VisibleAnywhere, BlueprintReadOnly, Category="LevelEditor|Menu")
+	FTypedElementHandle HitProxyElement;
+
 	UPROPERTY()
 	TArray<TObjectPtr<UActorComponent>> SelectedComponents;
 
diff --git a/Engine/Source/Editor/UnrealEd/Private/LevelViewportClickHandlers.cpp b/Engine/Source/Editor/UnrealEd/Private/LevelViewportClickHandlers.cpp
--- a/Engine/Source/Editor/UnrealEd/Private/LevelViewportClickHandlers.cpp
+++ b/Engine/Source/Editor/UnrealEd/Private/LevelViewportClickHandlers.cpp
@@ -31,7 +31,6 @@
 #include "LightMap.h"
 #include "Subsystems/BrushEditingSubsystem.h"
 #include "Elements/Framework/TypedElementSelectionSet.h"
-#include "Elements/Interfaces/TypedElementObjectInterface.h"
 
 #define LOCTEXT_NAMESPACE "ClickHandlers"
 
@@ -55,11 +54,11 @@ namespace LevelViewportClickHandlers
 		return nullptr;
 	}
 
-	static void PrivateSummonContextMenu( FLevelEditorViewportClient* ViewportClient, AActor* HitProxyActor = nullptr)
+	static void PrivateSummonContextMenu( FLevelEditorViewportClient* ViewportClient, const FTypedElementHandle& HitProxyElement = FTypedElementHandle())
 	{
 		if( ViewportClient->ParentLevelEditor.IsValid() )
 		{
-			ViewportClient->ParentLevelEditor.Pin()->SummonLevelViewportContextMenu(HitProxyActor);
+			ViewportClient->ParentLevelEditor.Pin()->SummonLevelViewportContextMenu(HitProxyElement);
 		}
 	}	
 
@@ -215,9 +214,7 @@ namespace LevelViewportClickHandlers
 
 				if (bIsRightClickSelection)
 				{
-					TTypedElement<UTypedElementObjectInterface> ObjectElement = LevelEditorElementSelectionSet->GetElementList()->GetElement<UTypedElementObjectInterface>(ResolvedElement);
-					AActor* HitActor = ObjectElement ? Cast<AActor>(ObjectElement.GetObject()) : nullptr;
-					PrivateSummonContextMenu(ViewportClient, HitActor);
+					PrivateSummonContextMenu(ViewportClient, ResolvedElement);
 				}
 			}
 		}
@@ -267,7 +264,7 @@ namespace LevelViewportClickHandlers
 				FlushRenderingCommands();
 			}
 
-			PrivateSummonContextMenu(ViewportClient, Actor);
+			PrivateSummonContextMenu(ViewportClient, UEngineElementsLibrary::AcquireEditorActorElementHandle(Actor));
 			return true;
 		}
 		else if( Click.GetEvent() == IE_DoubleClick && Click.GetKey() == EKeys::LeftMouseButton && !Click.IsControlDown() && !Click.IsShiftDown() )
