Fixed issue where transform changes were applied twice to attached actors when selecting both the actor and it's parent.

diff --git a/Engine/Source/Editor/LevelEditor/Private/Elements/Actor/ActorElementLevelEditorViewportInteractionCustomization.cpp b/Engine/Source/Editor/LevelEditor/Private/Elements/Actor/ActorElementLevelEditorViewportInteractionCustomization.cpp
--- a/Engine/Source/Editor/LevelEditor/Private/Elements/Actor/ActorElementLevelEditorViewportInteractionCustomization.cpp
+++ b/Engine/Source/Editor/LevelEditor/Private/Elements/Actor/ActorElementLevelEditorViewportInteractionCustomization.cpp
@@ -26,6 +26,25 @@ void FActorElementLevelEditorViewportInteractionCustomization::GetElementsToMove
 		return;
 	}
 
+	// Ensure that only parent-most actor(s) are moved
+	bool bParentAlsoSelected = false;
+	AActor* Parent = Actor->GetAttachParentActor();
+	while (Parent != nullptr)
+	{
+		if (Parent->IsSelected())
+		{
+			bParentAlsoSelected = true;
+			break;
+		}
+
+		Parent = Parent->GetAttachParentActor();
+	}
+	if (bParentAlsoSelected)
+	{
+		return;
+	}
+
+
 	if (CanMoveActorInViewport(Actor, InWorldType))
 	{
 		AppendActorsToMove(Actor, InSelectionSet, OutElementsToMove, OutElementsToMoveFinalizers);
