Adding SM6 preview platform.
- Adjusting the Editor handling of preview platforms to not assume SM5 is the running feature level.
- Adding "bPreviewFeatureLevelIsMax" to user settings to signal if the last saved Preview Feature Level was the maximum supported by the RHI, aka the default. If this is true on load, we set the Preview Feature Level to the max supported. This is to allow situations were we go from DX12+SM6 to DX11+SM5 and back without forcing the user to change the Preview Feature Level.
- Preventing unsupported feature levels from being previewed.
- Fixing the preview button being enabled for the default feature level.

diff --git a/Engine/Config/Windows/DataDrivenPlatformInfo.ini b/Engine/Config/Windows/DataDrivenPlatformInfo.ini
--- a/Engine/Config/Windows/DataDrivenPlatformInfo.ini
+++ b/Engine/Config/Windows/DataDrivenPlatformInfo.ini
@@ -38,6 +38,13 @@ bSupportsDistanceFields=true
 bSupportsSSDIndirect=true
 bRequiresVendorExtensionsForAtomics=true
 bSupportsAnisotropicMaterials=true
+bSupportsGPUSkinCache=true
+
+[PreviewPlatform PCD3D_SM5]
+PlatformName=PCD3D_SM5
+ShaderFormat=PCD3D_SM5
+MenuText=LOCTEXT("PreviewMenuText_PCD3D_SM5", "Shader Model 5")
+MenuTooltip=LOCTEXT("PreviewMenuTooltip_PCD3D_SM5", "DirectX 11, OpenGL 4.3+, PS4, XB1")
 
 [ShaderPlatform PCD3D_SM6]
 Language=D3D
@@ -54,6 +61,15 @@ bSupportsDistanceFields=true
 bSupportsSSDIndirect=true
 ; TODO: SM6.6
 bRequiresVendorExtensionsForAtomics=true
+bSupportsAnisotropicMaterials=true
+bSupportsGPUSkinCache=true
+FriendlyName=LOCTEXT("FriendlyShaderPlatformName_PCD3D_SM6", "Shader Model 6")
+
+[PreviewPlatform PCD3D_SM6]
+PlatformName=PCD3D_SM6
+ShaderFormat=PCD3D_SM6
+MenuText=LOCTEXT("PreviewMenuText_PCD3D_SM6", "Shader Model 6")
+MenuTooltip=LOCTEXT("PreviewMenuTooltip_PCD3D_SM6", "Preview using SM6 device profile.")
 
 [ShaderPlatform VULKAN_SM5]
 bIsMobile = false
diff --git a/Engine/Source/Editor/LevelEditor/Private/LevelEditor.cpp b/Engine/Source/Editor/LevelEditor/Private/LevelEditor.cpp
--- a/Engine/Source/Editor/LevelEditor/Private/LevelEditor.cpp
+++ b/Engine/Source/Editor/LevelEditor/Private/LevelEditor.cpp
@@ -1725,22 +1725,20 @@ void FLevelEditorModule::BindGlobalLevelEditorCommands()
 		FIsActionChecked::CreateStatic(&FLevelEditorActionCallbacks::IsFeatureLevelPreviewActive),
 		FIsActionButtonVisible::CreateStatic(FLevelEditorActionCallbacks::IsPreviewModeButtonVisible));
 
-	ActionList.MapAction(
-		Commands.PreviewPlatformOverride_SM5,
-		FExecuteAction::CreateStatic(&FLevelEditorActionCallbacks::SetPreviewPlatform, FPreviewPlatformInfo(ERHIFeatureLevel::SM5, NAME_None, NAME_None, NAME_None, false)),
-		FCanExecuteAction(),
-		FIsActionChecked::CreateStatic(&FLevelEditorActionCallbacks::IsPreviewPlatformChecked, FPreviewPlatformInfo(ERHIFeatureLevel::SM5, NAME_None)));
-
 	for (auto It = PlatformInfo::GetPreviewPlatformMenuItems().CreateConstIterator(); It; ++It)
 	{
 		EShaderPlatform ShaderPlatform = ShaderFormatToLegacyShaderPlatform(It.Value().ShaderFormat);
 		ERHIFeatureLevel::Type FeatureLevel = GetMaxSupportedFeatureLevel(ShaderPlatform);
+		const bool IsMaxFL = FeatureLevel == GMaxRHIFeatureLevel;
+
+		bool AllowSetOnSwitch = FeatureLevel != GMaxRHIFeatureLevel;
+		FPreviewPlatformInfo PreviewFeatureLevelInfo(FeatureLevel, IsMaxFL ? NAME_None : It.Value().PlatformName, IsMaxFL ? NAME_None : It.Value().ShaderFormat, IsMaxFL ? NAME_None : It.Value().DeviceProfileName, AllowSetOnSwitch);
 
 		ActionList.MapAction(
 			*Commands.PreviewPlatformOverrides.Find(It.Key()),
-			FExecuteAction::CreateStatic(&FLevelEditorActionCallbacks::SetPreviewPlatform, FPreviewPlatformInfo(FeatureLevel, It.Value().PlatformName, It.Value().ShaderFormat, It.Value().DeviceProfileName, true)),
-			FCanExecuteAction(),
-			FIsActionChecked::CreateStatic(&FLevelEditorActionCallbacks::IsPreviewPlatformChecked, FPreviewPlatformInfo(FeatureLevel, It.Value().PlatformName, It.Value().ShaderFormat, It.Value().DeviceProfileName)));
+			FExecuteAction::CreateStatic(&FLevelEditorActionCallbacks::SetPreviewPlatform, PreviewFeatureLevelInfo),
+			FCanExecuteAction::CreateStatic(&FLevelEditorActionCallbacks::CanExecutePreviewPlatform, PreviewFeatureLevelInfo),
+			FIsActionChecked::CreateStatic(&FLevelEditorActionCallbacks::IsPreviewPlatformChecked, PreviewFeatureLevelInfo));
 	}
 
 	ActionList.MapAction(
diff --git a/Engine/Source/Editor/LevelEditor/Private/LevelEditorActions.cpp b/Engine/Source/Editor/LevelEditor/Private/LevelEditorActions.cpp
--- a/Engine/Source/Editor/LevelEditor/Private/LevelEditorActions.cpp
+++ b/Engine/Source/Editor/LevelEditor/Private/LevelEditorActions.cpp
@@ -607,16 +607,16 @@ bool FLevelEditorActionCallbacks::IsFeatureLevelPreviewEnabled()
 	{
 		return false;
 	}
-	if (GEditor->PreviewPlatform.PreviewFeatureLevel == ERHIFeatureLevel::SM5)
+	if (GEditor->PreviewPlatform.PreviewFeatureLevel == GMaxRHIFeatureLevel)
 	{
-		return true;
+		return false;
 	}
 	return GEditor->IsFeatureLevelPreviewEnabled();
 }
 
 bool FLevelEditorActionCallbacks::IsFeatureLevelPreviewActive()
 {
-	if (GEditor->PreviewPlatform.PreviewFeatureLevel == ERHIFeatureLevel::SM5)
+	if (GEditor->PreviewPlatform.PreviewFeatureLevel == GMaxRHIFeatureLevel)
 	{
 		return false;
 	}
@@ -639,6 +639,11 @@ void FLevelEditorActionCallbacks::SetPreviewPlatform(FPreviewPlatformInfo NewPre
 	GEditor->SetPreviewPlatform(NewPreviewPlatform, true);
 }
 
+bool FLevelEditorActionCallbacks::CanExecutePreviewPlatform(FPreviewPlatformInfo NewPreviewPlatform)
+{
+	return NewPreviewPlatform.PreviewFeatureLevel <= GMaxRHIFeatureLevel;
+}
+
 bool FLevelEditorActionCallbacks::IsPreviewPlatformChecked(FPreviewPlatformInfo PreviewPlatform)
 {
 	return GEditor->PreviewPlatform.Matches(PreviewPlatform);
@@ -3462,8 +3467,6 @@ void FLevelEditorCommands::RegisterCommands()
 
 	UI_COMMAND(ToggleFeatureLevelPreview, "Preview Mode Toggle", "Toggles the Preview Mode on or off for the currently selected Preview target", EUserInterfaceActionType::ToggleButton, FInputChord());
 
-	UI_COMMAND(PreviewPlatformOverride_SM5, "Shader Model 5", "DirectX 11, OpenGL 4.3+, PS4, XB1", EUserInterfaceActionType::Check, FInputChord());
-
 	// Add preview platforms
 	for (auto It = PlatformInfo::GetPreviewPlatformMenuItems().CreateConstIterator(); It; ++It)
 	{
diff --git a/Engine/Source/Editor/LevelEditor/Private/LevelEditorToolBar.cpp b/Engine/Source/Editor/LevelEditor/Private/LevelEditorToolBar.cpp
--- a/Engine/Source/Editor/LevelEditor/Private/LevelEditorToolBar.cpp
+++ b/Engine/Source/Editor/LevelEditor/Private/LevelEditorToolBar.cpp
@@ -1345,7 +1345,11 @@ void FLevelEditorToolBar::RegisterLevelEditorToolBar( const TSharedRef<FUIComman
 				{
 					const FText& RenderingAsPlatformName = GetFriendlyShaderPlatformName(GEditor->PreviewPlatform.bPreviewFeatureLevelActive ? PreviewShaderPlatform : MaxRHIFeatureLevelPlatform);
                     const FText& SwitchToPlatformName = GetFriendlyShaderPlatformName(GEditor->PreviewPlatform.bPreviewFeatureLevelActive ? MaxRHIFeatureLevelPlatform : PreviewShaderPlatform);
-                    if (GWorld->FeatureLevel == GMaxRHIFeatureLevel)
+					if (PreviewShaderPlatform == MaxRHIFeatureLevelPlatform)
+					{
+						return FText::Format(LOCTEXT("PreviewModeViewingAs", "Viewing {0}."), RenderingAsPlatformName);
+					}
+                    else if (GWorld->FeatureLevel == GMaxRHIFeatureLevel)
                     {
                         return FText::Format(LOCTEXT("PreviewModeViewingAsSwitchTo", "Viewing {0}. Click to preview {1}."), RenderingAsPlatformName, SwitchToPlatformName);
                     }
@@ -1514,9 +1518,6 @@ static void MakeShaderModelPreviewMenu( UToolMenu* InMenu )
 
 	FToolMenuSection& Section = InMenu->AddSection("EditorPreviewMode", LOCTEXT("EditorPreviewModeDevices", "Preview Devices"));
 
-	// SM5
-	Section.AddMenuEntry(FLevelEditorCommands::Get().PreviewPlatformOverride_SM5);
-
 	// Preview platforms discovered from ITargetPlatforms.
 	for (auto It = FLevelEditorCommands::Get().PreviewPlatformOverrides.CreateConstIterator(); It; ++It)
 	{
diff --git a/Engine/Source/Editor/LevelEditor/Public/LevelEditorActions.h b/Engine/Source/Editor/LevelEditor/Public/LevelEditorActions.h
--- a/Engine/Source/Editor/LevelEditor/Public/LevelEditorActions.h
+++ b/Engine/Source/Editor/LevelEditor/Public/LevelEditorActions.h
@@ -588,8 +588,6 @@ public:
 
 	TSharedPtr< FUICommandInfo > ToggleFeatureLevelPreview;
 
-	TSharedPtr< FUICommandInfo > PreviewPlatformOverride_SM5;
-
 	TMap <FName, TSharedPtr<FUICommandInfo>> PreviewPlatformOverrides;
 	
 	///**
@@ -803,6 +801,7 @@ public:
 	static bool IsFeatureLevelPreviewActive();
 	static bool IsPreviewModeButtonVisible();
 	static void SetPreviewPlatform(FPreviewPlatformInfo NewPreviewPlatform);
+	static bool CanExecutePreviewPlatform(FPreviewPlatformInfo NewPreviewPlatform);
 	static bool IsPreviewPlatformChecked(FPreviewPlatformInfo NewPreviewPlatform);
 	static void GeometryCollection_SelectAllGeometry();
 	static void GeometryCollection_SelectNone();
diff --git a/Engine/Source/Editor/UnrealEd/Classes/Editor/EditorPerProjectUserSettings.h b/Engine/Source/Editor/UnrealEd/Classes/Editor/EditorPerProjectUserSettings.h
--- a/Engine/Source/Editor/UnrealEd/Classes/Editor/EditorPerProjectUserSettings.h
+++ b/Engine/Source/Editor/UnrealEd/Classes/Editor/EditorPerProjectUserSettings.h
@@ -173,6 +173,9 @@ public:
 	UPROPERTY(config)
 	bool bPreviewFeatureLevelActive;
 
+	UPROPERTY(config)
+	bool bPreviewFeatureLevelWasDefault;
+
 	UPROPERTY(config)
 	FName PreviewDeviceProfileName;
 
diff --git a/Engine/Source/Editor/UnrealEd/Private/EditorEngine.cpp b/Engine/Source/Editor/UnrealEd/Private/EditorEngine.cpp
--- a/Engine/Source/Editor/UnrealEd/Private/EditorEngine.cpp
+++ b/Engine/Source/Editor/UnrealEd/Private/EditorEngine.cpp
@@ -364,7 +364,7 @@ UEditorEngine::UEditorEngine(const FObjectInitializer& ObjectInitializer)
 					return;
 				}
 				
-				SetPreviewPlatform(FPreviewPlatformInfo(ERHIFeatureLevel::SM5), false);
+				SetPreviewPlatform(FPreviewPlatformInfo(GMaxRHIFeatureLevel), false);
 			}
 		});
 		
@@ -7407,7 +7407,7 @@ void UEditorEngine::SetPreviewPlatform(const FPreviewPlatformInfo& NewPreviewPla
 #if RHI_RAYTRACING
 	if (IsRayTracingEnabled())
 	{
-		if (PreviewPlatform.PreviewFeatureLevel != ERHIFeatureLevel::SM5)
+		if (PreviewPlatform.PreviewFeatureLevel < ERHIFeatureLevel::SM5)
 		{
 			UE_LOG(LogEditor, Warning, TEXT("Preview feature level is incompatible with ray tracing, defaulting to Shader Model 5"));
 			PreviewPlatform.PreviewFeatureLevel = ERHIFeatureLevel::SM5;
@@ -7589,6 +7589,7 @@ void UEditorEngine::SaveEditorFeatureLevel()
 	Settings->PreviewPlatformName = PreviewPlatform.PreviewPlatformName;
 	Settings->PreviewShaderFormatName = PreviewPlatform.PreviewShaderFormatName;
 	Settings->bPreviewFeatureLevelActive = PreviewPlatform.bPreviewFeatureLevelActive;
+	Settings->bPreviewFeatureLevelWasDefault = (PreviewPlatform.PreviewFeatureLevel == GMaxRHIFeatureLevel);
 	Settings->PreviewDeviceProfileName = PreviewPlatform.DeviceProfileName;
 	Settings->PostEditChange();
 }
diff --git a/Engine/Source/Editor/UnrealEd/Private/Settings/EditorPerProjectUserSettings.cpp b/Engine/Source/Editor/UnrealEd/Private/Settings/EditorPerProjectUserSettings.cpp
--- a/Engine/Source/Editor/UnrealEd/Private/Settings/EditorPerProjectUserSettings.cpp
+++ b/Engine/Source/Editor/UnrealEd/Private/Settings/EditorPerProjectUserSettings.cpp
@@ -28,9 +28,10 @@ UEditorPerProjectUserSettings::UEditorPerProjectUserSettings(const FObjectInitia
 	SwarmNumOfConcurrentJobs = 16;
 	SwarmMaxUploadChunkSizeInMB = 100;
 	SwarmIntermediateFolder = FPaths::ConvertRelativePathToFull(FPaths::ProjectIntermediateDir() + TEXT("Simplygon/"));
-	PreviewFeatureLevel = (int32)ERHIFeatureLevel::SM5;
+	PreviewFeatureLevel = GMaxRHIFeatureLevel;
 	PreviewShaderFormatName = NAME_None;
 	bPreviewFeatureLevelActive = false;
+	bPreviewFeatureLevelWasDefault = true;
 	PreviewDeviceProfileName = NAME_None;
 }
 
@@ -41,6 +42,17 @@ void UEditorPerProjectUserSettings::PostInitProperties()
 	//Ensure the material quality cvar is set to the settings loaded.
 	static IConsoleVariable* MaterialQualityLevelVar = IConsoleManager::Get().FindConsoleVariable(TEXT("r.MaterialQualityLevel"));
 	MaterialQualityLevelVar->Set(MaterialQualityLevel, ECVF_SetByScalability);
+
+	// if we last saved as the default or we somehow are loading a preview feature level higher than we can support,
+	// fall back to the current session's maximum feature level
+	if (bPreviewFeatureLevelWasDefault || PreviewFeatureLevel > GMaxRHIFeatureLevel)
+	{
+		PreviewFeatureLevel = GMaxRHIFeatureLevel;
+		PreviewShaderFormatName = NAME_None;
+		bPreviewFeatureLevelActive = false;
+		bPreviewFeatureLevelWasDefault = true;
+		PreviewDeviceProfileName = NAME_None;
+	}
 }
 
 #if WITH_EDITOR
diff --git a/Engine/Source/Runtime/RHI/Public/RHIShaderPlatformDefinitions.inl b/Engine/Source/Runtime/RHI/Public/RHIShaderPlatformDefinitions.inl
--- a/Engine/Source/Runtime/RHI/Public/RHIShaderPlatformDefinitions.inl
+++ b/Engine/Source/Runtime/RHI/Public/RHIShaderPlatformDefinitions.inl
@@ -95,7 +95,7 @@ static FText GetFriendlyShaderPlatformName(const EShaderPlatform InShaderPlatfor
 	break;
 
 	default:
-		if (FStaticShaderPlatformNames::IsStaticPlatform(InShaderPlatform))
+		if (FDataDrivenShaderPlatformInfo::IsValid(InShaderPlatform))
 		{
 			return FDataDrivenShaderPlatformInfo::GetFriendlyName(InShaderPlatform);
 		}
