Fixed null component access

#rnx

diff --git a/Engine/Source/Editor/LevelEditor/Private/Elements/Actor/ActorElementLevelEditorSelectionCustomization.cpp b/Engine/Source/Editor/LevelEditor/Private/Elements/Actor/ActorElementLevelEditorSelectionCustomization.cpp
--- a/Engine/Source/Editor/LevelEditor/Private/Elements/Actor/ActorElementLevelEditorSelectionCustomization.cpp
+++ b/Engine/Source/Editor/LevelEditor/Private/Elements/Actor/ActorElementLevelEditorSelectionCustomization.cpp
@@ -185,6 +185,11 @@ bool FActorElementLevelEditorSelectionCustomization::SelectActorElement(const TT
 	// Bind the override delegates for the components on the selected actor
 	for (UActorComponent* Component : Actor->GetComponents())
 	{
+		if (!Component)
+		{
+			continue;
+		}
+
 		if (USceneComponent* SceneComponent = Cast<USceneComponent>(Component))
 		{
 			FComponentEditorUtils::BindComponentSelectionOverride(SceneComponent, /*bBind*/true);
@@ -250,6 +255,11 @@ bool FActorElementLevelEditorSelectionCustomization::DeselectActorElement(const
 
 		for (UActorComponent* Component : Actor->GetComponents())
 		{
+			if (!Component)
+			{
+				continue;
+			}
+
 			TTypedElement<UTypedElementSelectionInterface> ComponentSelectionHandle = InSelectionSet->GetElement<UTypedElementSelectionInterface>(UEngineElementsLibrary::AcquireEditorComponentElementHandle(Component));
 			ComponentSelectionHandle.DeselectElement(InSelectionSet, InSelectionOptions);
 
diff --git a/Engine/Source/Editor/UnrealEd/Private/GroupActor.cpp b/Engine/Source/Editor/UnrealEd/Private/GroupActor.cpp
--- a/Engine/Source/Editor/UnrealEd/Private/GroupActor.cpp
+++ b/Engine/Source/Editor/UnrealEd/Private/GroupActor.cpp
@@ -110,16 +110,18 @@ bool ActorHasParentInGroup(const TArray<AActor*>& GroupActors, const AActor* Act
 {
 	check(Actor);
 	// Check that we've not got a parent attachment within the group.
-	USceneComponent* RootComponent = Actor->GetRootComponent();
-	for (const AActor* OtherActor : GroupActors)
+	if (USceneComponent* RootComponent = Actor->GetRootComponent())
 	{
-		if (OtherActor && OtherActor != Actor)
+		for (const AActor* OtherActor : GroupActors)
 		{
-			USceneComponent* OtherRootComponent = OtherActor->GetRootComponent();
-			if (RootComponent->IsAttachedTo(OtherRootComponent))
+			if (OtherActor && OtherActor != Actor)
 			{
-				// We do have parent so don't apply the delta - our parent object will apply it instead.
-				return true;
+				USceneComponent* OtherRootComponent = OtherActor->GetRootComponent();
+				if (OtherRootComponent && RootComponent->IsAttachedTo(OtherRootComponent))
+				{
+					// We do have parent so don't apply the delta - our parent object will apply it instead.
+					return true;
+				}
 			}
 		}
 	}
