Components which are not moveable should not allow viewport interactions with gizmos.

diff --git a/Engine/Source/Editor/LevelEditor/Private/Elements/Actor/ActorElementLevelEditorViewportInteractionCustomization.cpp b/Engine/Source/Editor/LevelEditor/Private/Elements/Actor/ActorElementLevelEditorViewportInteractionCustomization.cpp
--- a/Engine/Source/Editor/LevelEditor/Private/Elements/Actor/ActorElementLevelEditorViewportInteractionCustomization.cpp
+++ b/Engine/Source/Editor/LevelEditor/Private/Elements/Actor/ActorElementLevelEditorViewportInteractionCustomization.cpp
@@ -151,11 +151,18 @@ bool FActorElementLevelEditorViewportInteractionCustomization::CanMoveActorInVie
 		return false;
 	}
 
+	// If the actor has a root component, but it cannot be moved, then the actor cannot move.
+	const bool bIsPlayInEditorWorld = (InWorldType == ETypedElementViewportInteractionWorldType::PlayInEditor);
+	if (bIsPlayInEditorWorld && InActor->GetRootComponent() && !InActor->IsRootComponentMovable())
+	{
+		return false;
+	}
+
 	// The actor needs to be in the current viewport world
 	if (GEditor->PlayWorld)
 	{
 		const UWorld* CurrentWorld = InActor->GetWorld();
-		const UWorld* RequiredWorld = InWorldType == ETypedElementViewportInteractionWorldType::PlayInEditor ? GEditor->PlayWorld : GEditor->EditorWorld;
+		const UWorld* RequiredWorld = bIsPlayInEditorWorld ? GEditor->PlayWorld : GEditor->EditorWorld;
 		if (CurrentWorld != RequiredWorld)
 		{
 			return false;
diff --git a/Engine/Source/Editor/LevelEditor/Private/Elements/Component/ComponentElementLevelEditorViewportInteractionCustomization.cpp b/Engine/Source/Editor/LevelEditor/Private/Elements/Component/ComponentElementLevelEditorViewportInteractionCustomization.cpp
--- a/Engine/Source/Editor/LevelEditor/Private/Elements/Component/ComponentElementLevelEditorViewportInteractionCustomization.cpp
+++ b/Engine/Source/Editor/LevelEditor/Private/Elements/Component/ComponentElementLevelEditorViewportInteractionCustomization.cpp
@@ -28,28 +28,29 @@ void FComponentElementLevelEditorViewportInteractionCustomization::GetElementsTo
 			}
 #endif	// DO_CHECK
 
-			if (FActorElementLevelEditorViewportInteractionCustomization::CanMoveActorInViewport(ComponentOwner, InWorldType))
+			if (ComponentOwner->GetRootComponent() == SceneComponent)
 			{
-				if (ComponentOwner->GetRootComponent() == SceneComponent)
+				if (FActorElementLevelEditorViewportInteractionCustomization::CanMoveActorInViewport(ComponentOwner, InWorldType))
 				{
 					// If it is a root component, use the parent actor instead
 					FActorElementLevelEditorViewportInteractionCustomization::AppendActorsToMove(ComponentOwner, InSelectionSet, OutElementsToMove, OutElementsToMoveFinalizers);
 				}
-				else
+			}
+			else
+			{
+				// Check to see if any parent is selected
+				bool bHasSelectedParent = false;
+				for (USceneComponent* Parent = SceneComponent->GetAttachParent(); Parent && !bHasSelectedParent; Parent = Parent->GetAttachParent())
 				{
-					// Check to see if any parent is selected
-					bool bHasSelectedParent = false;
-					for (USceneComponent* Parent = SceneComponent->GetAttachParent(); Parent && !bHasSelectedParent; Parent = Parent->GetAttachParent())
-					{
-						FTypedElementHandle ParentElementHandle = UEngineElementsLibrary::AcquireEditorComponentElementHandle(Parent, /*bAllowCreate*/false);
-						bHasSelectedParent = ParentElementHandle && InSelectionSet->IsElementSelected(ParentElementHandle, FTypedElementIsSelectedOptions());
-					}
+					FTypedElementHandle ParentElementHandle = UEngineElementsLibrary::AcquireEditorComponentElementHandle(Parent, /*bAllowCreate*/false);
+					bHasSelectedParent = ParentElementHandle && InSelectionSet->IsElementSelected(ParentElementHandle, FTypedElementIsSelectedOptions());
+				}
 
-					if (!bHasSelectedParent)
-					{
-						// If no parent of this component is also in the selection set, move it
-						OutElementsToMove->Add(InElementWorldHandle);
-					}
+				const bool bCanMoveInEditorWorld = (InWorldType == ETypedElementViewportInteractionWorldType::Editor) || (SceneComponent->Mobility == EComponentMobility::Type::Movable);
+				if (!bHasSelectedParent && bCanMoveInEditorWorld)
+				{
+					// If no parent of this component is also in the selection set, move it
+					OutElementsToMove->Add(InElementWorldHandle);
 				}
 			}
 		}
