Optimized group selection/deselection

This avoids repeatedly processing each actor in a group during a select/deselect operation, as once the group itself has been selected/deselected, we have by proxy processed its grouped actors too. The only caveat to that is if a selection was previously performed with SetAllowGroups set to false (not the default), but that's advanced usage and could be resolved by toggling the group selection again.

diff --git a/Engine/Source/Editor/LevelEditor/Private/Elements/Actor/ActorElementLevelEditorSelectionCustomization.cpp b/Engine/Source/Editor/LevelEditor/Private/Elements/Actor/ActorElementLevelEditorSelectionCustomization.cpp
--- a/Engine/Source/Editor/LevelEditor/Private/Elements/Actor/ActorElementLevelEditorSelectionCustomization.cpp
+++ b/Engine/Source/Editor/LevelEditor/Private/Elements/Actor/ActorElementLevelEditorSelectionCustomization.cpp
@@ -99,7 +99,7 @@ bool FActorElementLevelEditorSelectionCustomization::CanSelectActorElement(const
 
 	// If grouping operations are not currently allowed, don't select groups
 	AGroupActor* SelectedGroupActor = Cast<AGroupActor>(Actor);
-	if (SelectedGroupActor && (!UActorGroupingUtils::IsGroupingActive() || !InSelectionOptions.AllowGroups()))
+	if (SelectedGroupActor && !UActorGroupingUtils::IsGroupingActive())
 	{
 		return false;
 	}
@@ -283,13 +283,18 @@ bool FActorElementLevelEditorSelectionCustomization::SelectActorGroup(AGroupActo
 {
 	bool bSelectionChanged = false;
 
+	TTypedElement<UTypedElementSelectionInterface> GroupSelectionHandle = InSelectionSet->GetElement<UTypedElementSelectionInterface>(UEngineElementsLibrary::AcquireEditorActorElementHandle(InGroupActor));
+
+	const FTypedElementSelectionOptions GroupSelectionOptions = FTypedElementSelectionOptions(InSelectionOptions)
+		.SetAllowGroups(false);
+
 	// Select all actors within the group (if locked or forced)
-	if (bForce || InGroupActor->IsLocked())
+	// Skip if the group is already selected, since this logic will have already run
+	if ((bForce || InGroupActor->IsLocked()) && !GroupSelectionHandle.IsElementSelected(InSelectionSet, FTypedElementIsSelectedOptions()) && CanSelectActorElement(GroupSelectionHandle, GroupSelectionOptions))
 	{
 		FTypedElementListLegacySyncScopedBatch LegacySyncBatch(InSelectionSet, InSelectionOptions.AllowLegacyNotifications());
 
-		const FTypedElementSelectionOptions GroupSelectionOptions = FTypedElementSelectionOptions(InSelectionOptions)
-			.SetAllowGroups(false);
+		bSelectionChanged |= SelectActorElement(GroupSelectionHandle, InSelectionSet, GroupSelectionOptions);
 
 		TArray<AActor*> GroupActors;
 		InGroupActor->GetGroupActors(GroupActors);
@@ -310,13 +315,18 @@ bool FActorElementLevelEditorSelectionCustomization::DeselectActorGroup(AGroupAc
 {
 	bool bSelectionChanged = false;
 
+	TTypedElement<UTypedElementSelectionInterface> GroupSelectionHandle = InSelectionSet->GetElement<UTypedElementSelectionInterface>(UEngineElementsLibrary::AcquireEditorActorElementHandle(InGroupActor));
+
+	const FTypedElementSelectionOptions GroupSelectionOptions = FTypedElementSelectionOptions(InSelectionOptions)
+		.SetAllowGroups(false);
+
 	// Deselect all actors within the group (if locked or forced)
-	if (bForce || InGroupActor->IsLocked())
+	// Skip if the group is already deselected, since this logic will have already run
+	if ((bForce || InGroupActor->IsLocked()) && GroupSelectionHandle.IsElementSelected(InSelectionSet, FTypedElementIsSelectedOptions()) && CanDeselectActorElement(GroupSelectionHandle, GroupSelectionOptions))
 	{
 		FTypedElementListLegacySyncScopedBatch LegacySyncBatch(InSelectionSet, InSelectionOptions.AllowLegacyNotifications());
 
-		const FTypedElementSelectionOptions GroupSelectionOptions = FTypedElementSelectionOptions(InSelectionOptions)
-			.SetAllowGroups(false);
+		bSelectionChanged |= DeselectActorElement(GroupSelectionHandle, InSelectionSet, GroupSelectionOptions);
 
 		TArray<AActor*> GroupActors;
 		InGroupActor->GetGroupActors(GroupActors);
