Removed finalizer from FTypedElementViewportInteractionCustomization::GetElementsToMove

This was to optimize selecting groups, but we've handled this case elsewhere by just skipping the group logic if the group is already in the list of things that we've processed. This applies that approach to getting the elements to move.

diff --git a/Engine/Source/Editor/EditorFramework/Private/Elements/Framework/TypedElementViewportInteraction.cpp b/Engine/Source/Editor/EditorFramework/Private/Elements/Framework/TypedElementViewportInteraction.cpp
--- a/Engine/Source/Editor/EditorFramework/Private/Elements/Framework/TypedElementViewportInteraction.cpp
+++ b/Engine/Source/Editor/EditorFramework/Private/Elements/Framework/TypedElementViewportInteraction.cpp
@@ -5,7 +5,7 @@
 #include "Elements/Framework/TypedElementRegistry.h"
 #include "Elements/Framework/TypedElementUtil.h"
 
-void FTypedElementViewportInteractionCustomization::GetElementsToMove(const TTypedElement<UTypedElementWorldInterface>& InElementWorldHandle, const ETypedElementViewportInteractionWorldType InWorldType, const UTypedElementSelectionSet* InSelectionSet, UTypedElementList* OutElementsToMove, FElementToMoveFinalizerMap& OutElementsToMoveFinalizers)
+void FTypedElementViewportInteractionCustomization::GetElementsToMove(const TTypedElement<UTypedElementWorldInterface>& InElementWorldHandle, const ETypedElementViewportInteractionWorldType InWorldType, const UTypedElementSelectionSet* InSelectionSet, UTypedElementList* OutElementsToMove)
 {
 	OutElementsToMove->Add(InElementWorldHandle);
 }
@@ -140,22 +140,16 @@ void UTypedElementViewportInteraction::GetSelectedElementsToMove(const UTypedEle
 {
 	OutElementsToMove->Reset();
 
-	FTypedElementViewportInteractionCustomization::FElementToMoveFinalizerMap ElementsToMoveFinalizers;
-	InSelectionSet->ForEachSelectedElement<UTypedElementWorldInterface>([this, InSelectionSet, InWorldType, OutElementsToMove, &ElementsToMoveFinalizers](const TTypedElement<UTypedElementWorldInterface>& InElementWorldHandle)
+	InSelectionSet->ForEachSelectedElement<UTypedElementWorldInterface>([this, InSelectionSet, InWorldType, OutElementsToMove](const TTypedElement<UTypedElementWorldInterface>& InElementWorldHandle)
 	{
 		if (InElementWorldHandle.CanEditElement())
 		{
 			FTypedElementViewportInteractionElement ViewportInteractionElement(InElementWorldHandle, GetInterfaceCustomizationByTypeId(InElementWorldHandle.GetId().GetTypeId()));
 			check(ViewportInteractionElement.IsSet());
-			ViewportInteractionElement.GetElementsToMove(InWorldType, InSelectionSet, OutElementsToMove, ElementsToMoveFinalizers);
+			ViewportInteractionElement.GetElementsToMove(InWorldType, InSelectionSet, OutElementsToMove);
 		}
 		return true;
 	});
-
-	for (const auto& FinalizerPair : ElementsToMoveFinalizers)
-	{
-		FinalizerPair.Value(FinalizerPair.Key);
-	}
 }
 
 void UTypedElementViewportInteraction::BeginGizmoManipulation(const UTypedElementList* InElementsToMove, const UE::Widget::EWidgetMode InWidgetMode)
diff --git a/Engine/Source/Editor/EditorFramework/Public/Elements/Framework/TypedElementViewportInteraction.h b/Engine/Source/Editor/EditorFramework/Public/Elements/Framework/TypedElementViewportInteraction.h
--- a/Engine/Source/Editor/EditorFramework/Public/Elements/Framework/TypedElementViewportInteraction.h
+++ b/Engine/Source/Editor/EditorFramework/Public/Elements/Framework/TypedElementViewportInteraction.h
@@ -25,11 +25,8 @@ class EDITORFRAMEWORK_API FTypedElementViewportInteractionCustomization
 public:
 	virtual ~FTypedElementViewportInteractionCustomization() = default;
 
-	using FElementToMoveFinalizerFunc = TFunction<void(const FTypedElementHandle&)>;
-	using FElementToMoveFinalizerMap  = TMap<FTypedElementHandle, FElementToMoveFinalizerFunc>;
-
 	//~ See UTypedElementViewportInteraction for API docs
-	virtual void GetElementsToMove(const TTypedElement<UTypedElementWorldInterface>& InElementWorldHandle, const ETypedElementViewportInteractionWorldType InWorldType, const UTypedElementSelectionSet* InSelectionSet, UTypedElementList* OutElementsToMove, FElementToMoveFinalizerMap& OutElementsToMoveFinalizers);
+	virtual void GetElementsToMove(const TTypedElement<UTypedElementWorldInterface>& InElementWorldHandle, const ETypedElementViewportInteractionWorldType InWorldType, const UTypedElementSelectionSet* InSelectionSet, UTypedElementList* OutElementsToMove);
 	virtual bool GetGizmoPivotLocation(const TTypedElement<UTypedElementWorldInterface>& InElementWorldHandle, const UE::Widget::EWidgetMode InWidgetMode, FVector& OutPivotLocation);
 	virtual void PreGizmoManipulationStarted(TArrayView<const FTypedElementHandle> InElementHandles, const UE::Widget::EWidgetMode InWidgetMode);
 	virtual void GizmoManipulationStarted(const TTypedElement<UTypedElementWorldInterface>& InElementWorldHandle, const UE::Widget::EWidgetMode InWidgetMode);
@@ -71,7 +68,7 @@ public:
 	}
 
 	//~ See UTypedElementViewportInteraction for API docs
-	void GetElementsToMove(const ETypedElementViewportInteractionWorldType InWorldType, const UTypedElementSelectionSet* InSelectionSet, UTypedElementList* OutElementsToMove, FTypedElementViewportInteractionCustomization::FElementToMoveFinalizerMap& OutElementsToMoveFinalizers) { ViewportInteractionCustomization->GetElementsToMove(ElementWorldHandle, InWorldType, InSelectionSet, OutElementsToMove, OutElementsToMoveFinalizers); }
+	void GetElementsToMove(const ETypedElementViewportInteractionWorldType InWorldType, const UTypedElementSelectionSet* InSelectionSet, UTypedElementList* OutElementsToMove) { ViewportInteractionCustomization->GetElementsToMove(ElementWorldHandle, InWorldType, InSelectionSet, OutElementsToMove); }
 	bool GetGizmoPivotLocation(const UE::Widget::EWidgetMode InWidgetMode, FVector& OutPivotLocation) const { return ViewportInteractionCustomization->GetGizmoPivotLocation(ElementWorldHandle, InWidgetMode, OutPivotLocation); }
 	void GizmoManipulationStarted(const UE::Widget::EWidgetMode InWidgetMode) const { ViewportInteractionCustomization->GizmoManipulationStarted(ElementWorldHandle, InWidgetMode); }
 	void GizmoManipulationDeltaUpdate(const UE::Widget::EWidgetMode InWidgetMode, const EAxisList::Type InDragAxis, const FInputDeviceState& InInputState, const FTransform& InDeltaTransform, const FVector& InPivotLocation) const { ViewportInteractionCustomization->GizmoManipulationDeltaUpdate(ElementWorldHandle, InWidgetMode, InDragAxis, InInputState, InDeltaTransform, InPivotLocation); }
diff --git a/Engine/Source/Editor/LevelEditor/Private/Elements/Actor/ActorElementLevelEditorViewportInteractionCustomization.cpp b/Engine/Source/Editor/LevelEditor/Private/Elements/Actor/ActorElementLevelEditorViewportInteractionCustomization.cpp
--- a/Engine/Source/Editor/LevelEditor/Private/Elements/Actor/ActorElementLevelEditorViewportInteractionCustomization.cpp
+++ b/Engine/Source/Editor/LevelEditor/Private/Elements/Actor/ActorElementLevelEditorViewportInteractionCustomization.cpp
@@ -15,7 +15,7 @@
 #include "ActorGroupingUtils.h"
 #include "LevelEditorViewport.h"
 
-void FActorElementLevelEditorViewportInteractionCustomization::GetElementsToMove(const TTypedElement<UTypedElementWorldInterface>& InElementWorldHandle, const ETypedElementViewportInteractionWorldType InWorldType, const UTypedElementSelectionSet* InSelectionSet, UTypedElementList* OutElementsToMove, FElementToMoveFinalizerMap& OutElementsToMoveFinalizers)
+void FActorElementLevelEditorViewportInteractionCustomization::GetElementsToMove(const TTypedElement<UTypedElementWorldInterface>& InElementWorldHandle, const ETypedElementViewportInteractionWorldType InWorldType, const UTypedElementSelectionSet* InSelectionSet, UTypedElementList* OutElementsToMove)
 {
 	AActor* Actor = ActorElementDataUtil::GetActorFromHandleChecked(InElementWorldHandle);
 
@@ -47,7 +47,7 @@ void FActorElementLevelEditorViewportInteractionCustomization::GetElementsToMove
 
 	if (CanMoveActorInViewport(Actor, InWorldType))
 	{
-		AppendActorsToMove(Actor, InSelectionSet, OutElementsToMove, OutElementsToMoveFinalizers);
+		AppendActorsToMove(Actor, InSelectionSet, OutElementsToMove);
 	}
 }
 
@@ -172,7 +172,7 @@ bool FActorElementLevelEditorViewportInteractionCustomization::CanMoveActorInVie
 	return true;
 }
 
-void FActorElementLevelEditorViewportInteractionCustomization::AppendActorsToMove(AActor* InActor, const UTypedElementSelectionSet* InSelectionSet, UTypedElementList* OutElementsToMove, FElementToMoveFinalizerMap& OutElementsToMoveFinalizers)
+void FActorElementLevelEditorViewportInteractionCustomization::AppendActorsToMove(AActor* InActor, const UTypedElementSelectionSet* InSelectionSet, UTypedElementList* OutElementsToMove)
 {
 	auto AddActorElement = [OutElementsToMove](AActor* InActorToAdd)
 	{
@@ -185,15 +185,12 @@ void FActorElementLevelEditorViewportInteractionCustomization::AppendActorsToMov
 	AGroupActor* ParentGroup = AGroupActor::GetRootForActor(InActor, true, true);
 	if (ParentGroup && UActorGroupingUtils::IsGroupingActive())
 	{
-		// Defer group enumeration until the finalization phase, so that each group is 
-		// enumerated once, regardless of how many actors within that group are selected
+		// Skip if the group is already in the list of things to move, since this logic will have already run
 		FTypedElementHandle ParentGroupElementHandle = UEngineElementsLibrary::AcquireEditorActorElementHandle(ParentGroup);
-		if (ParentGroupElementHandle && !OutElementsToMoveFinalizers.Contains(ParentGroupElementHandle))
+		if (ParentGroupElementHandle && !OutElementsToMove->Contains(ParentGroupElementHandle))
 		{
-			OutElementsToMoveFinalizers.Add(ParentGroupElementHandle, [ParentGroup, InSelectionSet, AddActorElement](const FTypedElementHandle&)
-			{
-				ParentGroup->ForEachMovableActorInGroup(InSelectionSet, AddActorElement);
-			});
+			ParentGroup->ForEachMovableActorInGroup(InSelectionSet, AddActorElement);
+			check(OutElementsToMove->Contains(ParentGroupElementHandle));
 		}
 	}
 	else
diff --git a/Engine/Source/Editor/LevelEditor/Private/Elements/Component/ComponentElementLevelEditorViewportInteractionCustomization.cpp b/Engine/Source/Editor/LevelEditor/Private/Elements/Component/ComponentElementLevelEditorViewportInteractionCustomization.cpp
--- a/Engine/Source/Editor/LevelEditor/Private/Elements/Component/ComponentElementLevelEditorViewportInteractionCustomization.cpp
+++ b/Engine/Source/Editor/LevelEditor/Private/Elements/Component/ComponentElementLevelEditorViewportInteractionCustomization.cpp
@@ -12,7 +12,7 @@
 #include "LevelEditorViewport.h"
 #include "Kismet2/ComponentEditorUtils.h"
 
-void FComponentElementLevelEditorViewportInteractionCustomization::GetElementsToMove(const TTypedElement<UTypedElementWorldInterface>& InElementWorldHandle, const ETypedElementViewportInteractionWorldType InWorldType, const UTypedElementSelectionSet* InSelectionSet, UTypedElementList* OutElementsToMove, FElementToMoveFinalizerMap& OutElementsToMoveFinalizers)
+void FComponentElementLevelEditorViewportInteractionCustomization::GetElementsToMove(const TTypedElement<UTypedElementWorldInterface>& InElementWorldHandle, const ETypedElementViewportInteractionWorldType InWorldType, const UTypedElementSelectionSet* InSelectionSet, UTypedElementList* OutElementsToMove)
 {
 	UActorComponent* Component = ComponentElementDataUtil::GetComponentFromHandleChecked(InElementWorldHandle);
 
@@ -33,7 +33,7 @@ void FComponentElementLevelEditorViewportInteractionCustomization::GetElementsTo
 				if (FActorElementLevelEditorViewportInteractionCustomization::CanMoveActorInViewport(ComponentOwner, InWorldType))
 				{
 					// If it is a root component, use the parent actor instead
-					FActorElementLevelEditorViewportInteractionCustomization::AppendActorsToMove(ComponentOwner, InSelectionSet, OutElementsToMove, OutElementsToMoveFinalizers);
+					FActorElementLevelEditorViewportInteractionCustomization::AppendActorsToMove(ComponentOwner, InSelectionSet, OutElementsToMove);
 				}
 			}
 			else
diff --git a/Engine/Source/Editor/LevelEditor/Public/Elements/Actor/ActorElementLevelEditorViewportInteractionCustomization.h b/Engine/Source/Editor/LevelEditor/Public/Elements/Actor/ActorElementLevelEditorViewportInteractionCustomization.h
--- a/Engine/Source/Editor/LevelEditor/Public/Elements/Actor/ActorElementLevelEditorViewportInteractionCustomization.h
+++ b/Engine/Source/Editor/LevelEditor/Public/Elements/Actor/ActorElementLevelEditorViewportInteractionCustomization.h
@@ -8,14 +8,14 @@
 class LEVELEDITOR_API FActorElementLevelEditorViewportInteractionCustomization : public FActorElementEditorViewportInteractionCustomization, public FTypedElementAssetEditorLevelEditorViewportClientMixin
 {
 public:
-	virtual void GetElementsToMove(const TTypedElement<UTypedElementWorldInterface>& InElementWorldHandle, const ETypedElementViewportInteractionWorldType InWorldType, const UTypedElementSelectionSet* InSelectionSet, UTypedElementList* OutElementsToMove, FElementToMoveFinalizerMap& OutElementsToMoveFinalizers) override;
+	virtual void GetElementsToMove(const TTypedElement<UTypedElementWorldInterface>& InElementWorldHandle, const ETypedElementViewportInteractionWorldType InWorldType, const UTypedElementSelectionSet* InSelectionSet, UTypedElementList* OutElementsToMove) override;
 	virtual void GizmoManipulationStarted(const TTypedElement<UTypedElementWorldInterface>& InElementWorldHandle, const UE::Widget::EWidgetMode InWidgetMode) override;
 	virtual void GizmoManipulationDeltaUpdate(const TTypedElement<UTypedElementWorldInterface>& InElementWorldHandle, const UE::Widget::EWidgetMode InWidgetMode, const EAxisList::Type InDragAxis, const FInputDeviceState& InInputState, const FTransform& InDeltaTransform, const FVector& InPivotLocation) override;
 	virtual void GizmoManipulationStopped(const TTypedElement<UTypedElementWorldInterface>& InElementWorldHandle, const UE::Widget::EWidgetMode InWidgetMode) override;
 	virtual void PostGizmoManipulationStopped(TArrayView<const FTypedElementHandle> InElementHandles, const UE::Widget::EWidgetMode InWidgetMode) override;
 
 	static bool CanMoveActorInViewport(const AActor* InActor, const ETypedElementViewportInteractionWorldType InWorldType);
-	static void AppendActorsToMove(AActor* InActor, const UTypedElementSelectionSet* InSelectionSet, UTypedElementList* OutElementsToMove, FElementToMoveFinalizerMap& OutElementsToMoveFinalizers);
+	static void AppendActorsToMove(AActor* InActor, const UTypedElementSelectionSet* InSelectionSet, UTypedElementList* OutElementsToMove);
 
 private:
 	void ModifyScale(AActor* InActor, const EAxisList::Type InDragAxis, FVector& ScaleDelta, bool bCheckSmallExtent);
diff --git a/Engine/Source/Editor/LevelEditor/Public/Elements/Component/ComponentElementLevelEditorViewportInteractionCustomization.h b/Engine/Source/Editor/LevelEditor/Public/Elements/Component/ComponentElementLevelEditorViewportInteractionCustomization.h
--- a/Engine/Source/Editor/LevelEditor/Public/Elements/Component/ComponentElementLevelEditorViewportInteractionCustomization.h
+++ b/Engine/Source/Editor/LevelEditor/Public/Elements/Component/ComponentElementLevelEditorViewportInteractionCustomization.h
@@ -10,7 +10,7 @@ class FLevelEditorViewportClient;
 class LEVELEDITOR_API FComponentElementLevelEditorViewportInteractionCustomization : public FComponentElementEditorViewportInteractionCustomization, public FTypedElementAssetEditorLevelEditorViewportClientMixin
 {
 public:
-	virtual void GetElementsToMove(const TTypedElement<UTypedElementWorldInterface>& InElementWorldHandle, const ETypedElementViewportInteractionWorldType InWorldType, const UTypedElementSelectionSet* InSelectionSet, UTypedElementList* OutElementsToMove, FElementToMoveFinalizerMap& OutElementsToMoveFinalizers) override;
+	virtual void GetElementsToMove(const TTypedElement<UTypedElementWorldInterface>& InElementWorldHandle, const ETypedElementViewportInteractionWorldType InWorldType, const UTypedElementSelectionSet* InSelectionSet, UTypedElementList* OutElementsToMove) override;
 	virtual void GizmoManipulationStarted(const TTypedElement<UTypedElementWorldInterface>& InElementWorldHandle, const UE::Widget::EWidgetMode InWidgetMode) override;
 	virtual void GizmoManipulationDeltaUpdate(const TTypedElement<UTypedElementWorldInterface>& InElementWorldHandle, const UE::Widget::EWidgetMode InWidgetMode, const EAxisList::Type InDragAxis, const FInputDeviceState& InInputState, const FTransform& InDeltaTransform, const FVector& InPivotLocation) override;
 	virtual void GizmoManipulationStopped(const TTypedElement<UTypedElementWorldInterface>& InElementWorldHandle, const UE::Widget::EWidgetMode InWidgetMode) override;
