Fix content browser tab well being collapsed by default causing new user confusion when trying to dock tabs along side the content browser

diff --git a/Engine/Source/Editor/LevelEditor/Private/SLevelEditor.cpp b/Engine/Source/Editor/LevelEditor/Private/SLevelEditor.cpp
--- a/Engine/Source/Editor/LevelEditor/Private/SLevelEditor.cpp
+++ b/Engine/Source/Editor/LevelEditor/Private/SLevelEditor.cpp
@@ -1478,7 +1478,7 @@ TSharedRef<SWidget> SLevelEditor::RestoreContentArea( const TSharedRef<SDockTab>
 				(
 					FTabManager::NewStack()
 					->SetSizeCoefficient(.4)
-					->SetHideTabWell(true)
+					->SetHideTabWell(false)
 					->AddTab("ContentBrowserTab1", ETabState::ClosedTab)
 					->AddTab(LevelEditorTabIds::OutputLog, ETabState::ClosedTab)
 				)
