Adding DPI Scaling and fallback positioning for the in-viewport menu

diff --git a/Engine/Source/Editor/LevelEditor/Private/SLevelViewport.cpp b/Engine/Source/Editor/LevelEditor/Private/SLevelViewport.cpp
--- a/Engine/Source/Editor/LevelEditor/Private/SLevelViewport.cpp
+++ b/Engine/Source/Editor/LevelEditor/Private/SLevelViewport.cpp
@@ -4574,8 +4574,21 @@ void SLevelViewport::ToggleInViewportContextMenu()
 						.SetRealtimeUpdate(IsRealtime()));
 					// SceneView is deleted with the ViewFamily
 					FSceneView* SceneView = LevelViewportClient->CalcSceneView(&ViewFamily);
+					const float InvDpiScale = 1.0f / LevelViewportClient->GetDPIScale();
 					FVector2D ScreenPos;
 					SceneView->WorldToPixel(SelectedActor->GetTransform().GetLocation(), ScreenPos);
+					ScreenPos *= InvDpiScale;
+					const float EdgeFactor = 0.85f;
+					const float MinX = SceneView->UnscaledViewRect.Width() * InvDpiScale * (1 - EdgeFactor);
+					const float MinY = SceneView->UnscaledViewRect.Height() * InvDpiScale * (1 - EdgeFactor);
+					const float MaxX = SceneView->UnscaledViewRect.Width() * InvDpiScale * EdgeFactor;
+					const float MaxY = (SceneView->UnscaledViewRect.Height() * InvDpiScale * EdgeFactor);
+					const bool bOutside = ScreenPos.X < MinX || ScreenPos.X > MaxX || ScreenPos.Y < MinY || ScreenPos.Y > MaxY;
+					if (bOutside)
+					{
+						ScreenPos.X = (SceneView->UnscaledViewRect.Width() * InvDpiScale) / 2.0f;
+						ScreenPos.Y = (SceneView->UnscaledViewRect.Height() * InvDpiScale) / 2.0f;
+					}
 					UpdateInViewportMenuLocation(ScreenPos);
 				}
 				bIsInViewportMenuInitialized = true;
