Clean up UI commands lists which were using non-descript, duplicate categories, making them hard to find in the keybinding editor

diff --git a/Engine/Plugins/Animation/LiveLink/Source/LiveLinkEditor/Private/LiveLinkClientCommands.cpp b/Engine/Plugins/Animation/LiveLink/Source/LiveLinkEditor/Private/LiveLinkClientCommands.cpp
--- a/Engine/Plugins/Animation/LiveLink/Source/LiveLinkEditor/Private/LiveLinkClientCommands.cpp
+++ b/Engine/Plugins/Animation/LiveLink/Source/LiveLinkEditor/Private/LiveLinkClientCommands.cpp
@@ -6,7 +6,7 @@
 #define LOCTEXT_NAMESPACE "LiveLinkClientCommands"
 
 FLiveLinkClientCommands::FLiveLinkClientCommands()
-	: TCommands<FLiveLinkClientCommands>("LiveLinkClient.Common", LOCTEXT("Common", "Common"), NAME_None, FName(TEXT("LiveLinkStyle")))
+	: TCommands<FLiveLinkClientCommands>("LiveLinkClient.Common", LOCTEXT("LiveLinkCommandsLabel", "Live Link"), NAME_None, FName(TEXT("LiveLinkStyle")))
 {
 }
 
diff --git a/Engine/Plugins/Experimental/CommonConversation/Source/CommonConversationEditor/Private/ConversationEditorCommands.cpp b/Engine/Plugins/Experimental/CommonConversation/Source/CommonConversationEditor/Private/ConversationEditorCommands.cpp
--- a/Engine/Plugins/Experimental/CommonConversation/Source/CommonConversationEditor/Private/ConversationEditorCommands.cpp
+++ b/Engine/Plugins/Experimental/CommonConversation/Source/CommonConversationEditor/Private/ConversationEditorCommands.cpp
@@ -9,7 +9,7 @@
 //
 
 FConversationEditorCommonCommands::FConversationEditorCommonCommands()
-	: TCommands<FConversationEditorCommonCommands>("ConversationEditor.Common", LOCTEXT("Common", "Common"), NAME_None, FEditorStyle::GetStyleSetName())
+	: TCommands<FConversationEditorCommonCommands>("ConversationEditor.Common", LOCTEXT("ComversationEditorCommandsLabel", "Conversation Editor"), NAME_None, FEditorStyle::GetStyleSetName())
 {
 }
 
diff --git a/Engine/Plugins/MovieScene/MovieRenderPipeline/Source/MovieRenderPipelineEditor/Private/MoviePipelineCommands.cpp b/Engine/Plugins/MovieScene/MovieRenderPipeline/Source/MovieRenderPipelineEditor/Private/MoviePipelineCommands.cpp
--- a/Engine/Plugins/MovieScene/MovieRenderPipeline/Source/MovieRenderPipelineEditor/Private/MoviePipelineCommands.cpp
+++ b/Engine/Plugins/MovieScene/MovieRenderPipeline/Source/MovieRenderPipelineEditor/Private/MoviePipelineCommands.cpp
@@ -6,7 +6,7 @@
 #define LOCTEXT_NAMESPACE "MoviePipelineCommands"
 
 FMoviePipelineCommands::FMoviePipelineCommands()
-	: TCommands<FMoviePipelineCommands>("MovieRenderPipeline", LOCTEXT("Common", "Common"), NAME_None, FMovieRenderPipelineStyle::StyleName)
+	: TCommands<FMoviePipelineCommands>("MovieRenderPipeline", LOCTEXT("MoviePipelineCommandsLabel", "Movie Pipeline"), NAME_None, FMovieRenderPipelineStyle::StyleName)
 {
 }
 
diff --git a/Engine/Plugins/VirtualProduction/Takes/Source/TakeRecorder/Private/TakeRecorderCommands.cpp b/Engine/Plugins/VirtualProduction/Takes/Source/TakeRecorder/Private/TakeRecorderCommands.cpp
--- a/Engine/Plugins/VirtualProduction/Takes/Source/TakeRecorder/Private/TakeRecorderCommands.cpp
+++ b/Engine/Plugins/VirtualProduction/Takes/Source/TakeRecorder/Private/TakeRecorderCommands.cpp
@@ -6,7 +6,7 @@
 #define LOCTEXT_NAMESPACE "TakeRecorderCommands"
 
 FTakeRecorderCommands::FTakeRecorderCommands()
-	: TCommands<FTakeRecorderCommands>("TakeRecorder", LOCTEXT("Common", "Common"), NAME_None, FTakeRecorderStyle::StyleName)
+	: TCommands<FTakeRecorderCommands>("TakeRecorder", LOCTEXT("TakeRecorderCommandLabel", "Take Recorder"), NAME_None, FTakeRecorderStyle::StyleName)
 {
 }
 
diff --git a/Engine/Source/Editor/BehaviorTreeEditor/Private/BehaviorTreeEditorCommands.cpp b/Engine/Source/Editor/BehaviorTreeEditor/Private/BehaviorTreeEditorCommands.cpp
--- a/Engine/Source/Editor/BehaviorTreeEditor/Private/BehaviorTreeEditorCommands.cpp
+++ b/Engine/Source/Editor/BehaviorTreeEditor/Private/BehaviorTreeEditorCommands.cpp
@@ -6,7 +6,7 @@
 #define LOCTEXT_NAMESPACE "BehaviorTreeEditorCommands"
 
 FBTCommonCommands::FBTCommonCommands() 
-	: TCommands<FBTCommonCommands>("BTEditor.Common", LOCTEXT("Common", "Common"), NAME_None, FEditorStyle::GetStyleSetName())
+	: TCommands<FBTCommonCommands>("BTEditor.Common", LOCTEXT("BehaviorTreeCommandLabel", "Behavior Tree"), NAME_None, FEditorStyle::GetStyleSetName())
 {
 }
 
diff --git a/Engine/Source/Editor/LevelEditor/Private/LevelViewportActions.cpp b/Engine/Source/Editor/LevelEditor/Private/LevelViewportActions.cpp
--- a/Engine/Source/Editor/LevelEditor/Private/LevelViewportActions.cpp
+++ b/Engine/Source/Editor/LevelEditor/Private/LevelViewportActions.cpp
@@ -193,7 +193,7 @@ void FLevelViewportCommands::HandleNewStat(const FName& InStatName, const FName&
 		if (!FInputBindingManager::Get().FindCommandInContext(this->GetContextName(), InStatName).IsValid() && !GroupCategory.IsEmpty())
 		{
 			// Find or Add the category
-			TArray< FShowMenuCommand >* ShowStatCommands = ShowStatCatCommands.Find(GroupCategory);
+			TArray<FShowMenuCommand>* ShowStatCommands = ShowStatCatCommands.Find(GroupCategory);
 			if (!ShowStatCommands)
 			{
 				// New category means we'll need to resort
@@ -216,7 +216,7 @@ void FLevelViewportCommands::HandleNewStat(const FName& InStatName, const FName&
 				}
 
 				TSharedPtr<FUICommandInfo> StatCommand
-					= FUICommandInfoDecl(this->AsShared(), InStatName, FText::GetEmpty(), DescriptionName)
+					= FUICommandInfoDecl(this->AsShared(), InStatName, DisplayName, DescriptionName)
 					.UserInterfaceType(EUserInterfaceActionType::ToggleButton);
 
 				FLevelViewportCommands::FShowMenuCommand ShowStatCommand(StatCommand, DisplayName);
diff --git a/Engine/Source/Editor/SequenceRecorder/Private/SequenceRecorderCommands.cpp b/Engine/Source/Editor/SequenceRecorder/Private/SequenceRecorderCommands.cpp
--- a/Engine/Source/Editor/SequenceRecorder/Private/SequenceRecorderCommands.cpp
+++ b/Engine/Source/Editor/SequenceRecorder/Private/SequenceRecorderCommands.cpp
@@ -6,7 +6,7 @@
 #define LOCTEXT_NAMESPACE "SequenceRecorderCommands"
 
 FSequenceRecorderCommands::FSequenceRecorderCommands()
-	: TCommands<FSequenceRecorderCommands>("SequenceRecorder.Common", LOCTEXT("Common", "Common"), NAME_None, FEditorStyle::GetStyleSetName())
+	: TCommands<FSequenceRecorderCommands>("SequenceRecorder.Common", LOCTEXT("SequenceRecorderCommandsLabel", "Sequence Recorder"), NAME_None, FEditorStyle::GetStyleSetName())
 {
 }
 
diff --git a/Engine/Source/Editor/UMGEditor/Private/Designer/DesignerCommands.h b/Engine/Source/Editor/UMGEditor/Private/Designer/DesignerCommands.h
--- a/Engine/Source/Editor/UMGEditor/Private/Designer/DesignerCommands.h
+++ b/Engine/Source/Editor/UMGEditor/Private/Designer/DesignerCommands.h
@@ -16,7 +16,7 @@ public:
 		: TCommands<FDesignerCommands>
 		(
 			TEXT("WidgetDesigner"), // Context name for fast lookup
-			NSLOCTEXT("Contexts", "DesignerCommands", "Common Designer Commands"), // Localized context name for displaying
+			NSLOCTEXT("Contexts", "DesignerCommands", "UMG Designer Commands"), // Localized context name for displaying
 			NAME_None, // Parent
 			FEditorStyle::GetStyleSetName() // Icon Style Set
 		)
