Added actor context menu "Copy Actor File Path" to copy the file path (level file or actor file) for the current selection.

diff --git a/Engine/Source/Editor/LevelEditor/Private/LevelEditor.cpp b/Engine/Source/Editor/LevelEditor/Private/LevelEditor.cpp
--- a/Engine/Source/Editor/LevelEditor/Private/LevelEditor.cpp
+++ b/Engine/Source/Editor/LevelEditor/Private/LevelEditor.cpp
@@ -729,6 +729,11 @@ void FLevelEditorModule::BindGlobalLevelEditorCommands()
 		FExecuteAction::CreateStatic(&FLevelEditorActionCallbacks::SnapObjectToView_Clicked)
 		);
 
+	ActionList.MapAction(
+		Commands.CopyActorFilePathtoClipboard,
+		FExecuteAction::CreateStatic(&FLevelEditorActionCallbacks::CopyActorFilePathtoClipboard_Clicked)
+		);
+
 	ActionList.MapAction(
 		Commands.GoToCodeForActor, 
 		FExecuteAction::CreateStatic( &FLevelEditorActionCallbacks::GoToCodeForActor_Clicked )
diff --git a/Engine/Source/Editor/LevelEditor/Private/LevelEditorActions.cpp b/Engine/Source/Editor/LevelEditor/Private/LevelEditorActions.cpp
--- a/Engine/Source/Editor/LevelEditor/Private/LevelEditorActions.cpp
+++ b/Engine/Source/Editor/LevelEditor/Private/LevelEditorActions.cpp
@@ -8,6 +8,7 @@
 #include "HAL/PlatformFileManager.h"
 #include "Misc/MessageDialog.h"
 #include "HAL/FileManager.h"
+#include "HAL/PlatformApplicationMisc.h"
 #include "Misc/App.h"
 #include "Modules/ModuleManager.h"
 #include "Layout/WidgetPath.h"
@@ -2627,7 +2628,25 @@ void FLevelEditorActionCallbacks::SnapObjectToView_Clicked()
 
 		LevelDirtyCallback.Request();
 	}
+}
+
+void FLevelEditorActionCallbacks::CopyActorFilePathtoClipboard_Clicked()
+{
+	FString Result;
+
+	for (FSelectionIterator It(GEditor->GetSelectedActorIterator()); It; ++It)
+	{
+		const AActor* Actor = Cast<AActor>(*It);
+		const UPackage* Package = Actor->GetPackage();
+		const FString LocalFullPath(Package->GetLoadedPath().GetLocalFullPath());
+		Result += FPaths::ConvertRelativePathToFull(LocalFullPath);
+		Result += TEXT("\n");
+	}
 
+	if (Result.Len())
+	{
+		FPlatformApplicationMisc::ClipboardCopy(*Result);
+	}
 }
 
 void FLevelEditorActionCallbacks::OnEnableActorSnap()
@@ -3246,6 +3265,8 @@ void FLevelEditorCommands::RegisterCommands()
 	UI_COMMAND( SnapCameraToObject, "Snap View to Object", "Snaps the view to the selected object", EUserInterfaceActionType::Button, FInputChord() );
 	UI_COMMAND( SnapObjectToCamera, "Snap Object to View", "Snaps the selected object to the view", EUserInterfaceActionType::Button, FInputChord() );
 
+	UI_COMMAND( CopyActorFilePathtoClipboard, "Copy Actor File Path", "Copy the file path where this actor is saved", EUserInterfaceActionType::Button, FInputChord() );
+
 	UI_COMMAND( GoToCodeForActor, "Go to C++ Code for Actor", "Opens a code editing IDE and navigates to the source file associated with the seleced actor", EUserInterfaceActionType::Button, FInputChord() );
 	UI_COMMAND( GoToDocsForActor, "Go to Documentation for Actor", "Opens documentation for the Actor in the default web browser", EUserInterfaceActionType::Button, FInputChord() );
 
diff --git a/Engine/Source/Editor/LevelEditor/Private/LevelEditorContextMenu.cpp b/Engine/Source/Editor/LevelEditor/Private/LevelEditorContextMenu.cpp
--- a/Engine/Source/Editor/LevelEditor/Private/LevelEditorContextMenu.cpp
+++ b/Engine/Source/Editor/LevelEditor/Private/LevelEditorContextMenu.cpp
@@ -203,6 +203,7 @@ void FLevelEditorContextMenu::RegisterComponentContextMenu()
 			Section.AddMenuEntry(FLevelEditorCommands::Get().GoHere);
 			Section.AddMenuEntry(FLevelEditorCommands::Get().SnapCameraToObject);
 			Section.AddMenuEntry(FLevelEditorCommands::Get().SnapObjectToCamera);
+			Section.AddMenuEntry(FLevelEditorCommands::Get().CopyActorFilePathtoClipboard);
 		}
 
 		FComponentEditorUtils::FillComponentContextMenuOptions(InMenu, SelectedComponents);
@@ -288,6 +289,7 @@ void FLevelEditorContextMenu::RegisterActorContextMenu()
 			Section.AddMenuEntry(FLevelEditorCommands::Get().GoHere);
 			Section.AddMenuEntry(FLevelEditorCommands::Get().SnapCameraToObject);
 			Section.AddMenuEntry(FLevelEditorCommands::Get().SnapObjectToCamera);
+			Section.AddMenuEntry(FLevelEditorCommands::Get().CopyActorFilePathtoClipboard);
 
 			if (SelectedActors.Num() == 1)
 			{
diff --git a/Engine/Source/Editor/LevelEditor/Public/LevelEditorActions.h b/Engine/Source/Editor/LevelEditor/Public/LevelEditorActions.h
--- a/Engine/Source/Editor/LevelEditor/Public/LevelEditorActions.h
+++ b/Engine/Source/Editor/LevelEditor/Public/LevelEditorActions.h
@@ -131,6 +131,9 @@ public:
 	/** Snaps the selected actor to the camera. */
 	TSharedPtr< FUICommandInfo > SnapObjectToCamera;
 
+	/** Copy the file path where the actor is saved. */
+	TSharedPtr< FUICommandInfo > CopyActorFilePathtoClipboard;
+
 	/** Goes to the source code for the selected actor's class. */
 	TSharedPtr< FUICommandInfo > GoToCodeForActor;
 
@@ -1255,6 +1258,11 @@ public:
 	*/
 	static void SnapObjectToView_Clicked();
 
+	/**
+	* Copy the file path where the actor is saved.
+	*/
+	static void CopyActorFilePathtoClipboard_Clicked();
+
 	/** 
 	 * Moves the selected elements to the last selected element.
 	 */
