Adding back selected actor's level info in the viewport

diff --git a/Engine/Source/Editor/LevelEditor/Private/SLevelViewport.cpp b/Engine/Source/Editor/LevelEditor/Private/SLevelViewport.cpp
--- a/Engine/Source/Editor/LevelEditor/Private/SLevelViewport.cpp
+++ b/Engine/Source/Editor/LevelEditor/Private/SLevelViewport.cpp
@@ -326,7 +326,7 @@ void SLevelViewport::ConstructViewportOverlayContent()
 			[
 				BuildFeatureLevelWidget()
 			]
-		/*	+ SVerticalBox::Slot()
+			+ SVerticalBox::Slot()
 			.AutoHeight()
 			.Padding(2.0f, 1.0f, 2.0f, 1.0f)
 			[
@@ -350,7 +350,7 @@ void SLevelViewport::ConstructViewportOverlayContent()
 					.Text(this, &SLevelViewport::GetSelectedActorsCurrentLevelText, false)
 					.ShadowOffset(FVector2D(1, 1))
 				]
-			]*/
+			]
 			+ SVerticalBox::Slot()
 			.AutoHeight()
 			.Padding(2.0f, 1.0f, 2.0f, 1.0f)
