Replace use of the old SCS editor with the Subobject editor in the Actor Details panel of the level editor.

#rnx

diff --git a/Engine/Source/Editor/LevelEditor/Private/SActorDetails.cpp b/Engine/Source/Editor/LevelEditor/Private/SActorDetails.cpp
--- a/Engine/Source/Editor/LevelEditor/Private/SActorDetails.cpp
+++ b/Engine/Source/Editor/LevelEditor/Private/SActorDetails.cpp
@@ -22,7 +22,6 @@
 #include "UnrealEdGlobals.h"
 #include "LevelEditor.h"
 #include "Kismet2/KismetEditorUtilities.h"
-#include "SSCSEditor.h"
 #include "PropertyEditorModule.h"
 #include "IDetailsView.h"
 #include "LevelEditorGenericDetails.h"
@@ -34,6 +33,11 @@
 #include "Elements/Framework/TypedElementRegistry.h"
 #include "Elements/Framework/TypedElementSelectionSet.h"
 #include "Elements/Framework/EngineElementsLibrary.h"
+#include "SSubobjectEditor.h"
+#include "SSubobjectInstanceEditor.h"
+#include "SubobjectData.h"
+#include "SubobjectDataSubsystem.h"
+#include "SSubobjectEditorModule.h"
 
 class SActorDetailsUneditableComponentWarning : public SCompoundWidget
 {
@@ -132,19 +136,21 @@ void SActorDetails::Construct(const FArguments& InArgs, UTypedElementSelectionSe
 	GEditor->RegisterForUndo(this);
 
 	ComponentsBox = SNew(SBox)
-		.Padding(FMargin(2,0,2,0))
-		.Visibility(this, &SActorDetails::GetComponentsBoxVisibility)
-		[
-			SAssignNew(SCSEditor, SSCSEditor)
-			.EditorMode(EComponentEditorMode::ActorInstance)
-			.AllowEditing(this, &SActorDetails::GetAllowComponentTreeEditing)
-			.ActorContext(this, &SActorDetails::GetActorContext)
-			.OnSelectionUpdated(this, &SActorDetails::OnSCSEditorTreeViewSelectionChanged)
-			.OnItemDoubleClicked(this, &SActorDetails::OnSCSEditorTreeViewItemDoubleClicked)
-			.OnObjectReplaced(this, &SActorDetails::OnSCSEditorTreeViewObjectReplaced)
-		];
+		.Padding(FMargin(2.0f, 0.0f, 2.0f, 0.0f))
+		.Visibility(this, &SActorDetails::GetComponentsBoxVisibility);
+
+	FModuleManager::LoadModuleChecked<FSubobjectEditorModule>("SubobjectEditor");
+	
+	SubobjectEditor = SNew(SSubobjectInstanceEditor)
+		.ObjectContext(this, &SActorDetails::GetObjectContext)
+		.AllowEditing(this, &SActorDetails::GetAllowComponentTreeEditing)
+		.OnSelectionUpdated(this, &SActorDetails::OnSubobjectEditorTreeViewSelectionChanged)
+		.OnItemDoubleClicked(this, &SActorDetails::OnSubobjectEditorTreeViewItemDoubleClicked)
+		.OnObjectReplaced(this, &SActorDetails::OnSubobjectEditorTreeViewObjectReplaced);
 
-	TSharedRef<SWidget> ButtonBox = SCSEditor->GetToolButtonsBox().ToSharedRef();
+	ComponentsBox->SetContent(SubobjectEditor.ToSharedRef());
+
+	TSharedRef<SWidget> ButtonBox = SubobjectEditor->GetToolButtonsBox().ToSharedRef();
 	DetailsView->SetNameAreaCustomContent( ButtonBox );
 
 	ChildSlot
@@ -307,12 +313,13 @@ void SActorDetails::RefreshTopLevelElements(TArrayView<const TTypedElement<UType
 	// Update the underlying details view
 	SetElementDetailsObjects(TopLevelElements, bForceRefresh, bOverrideLock);
 
-	// Update the SCS tree if we were asked to edit a single actor
+	// Update the Subobject tree if we were asked to edit a single actor
+	
 	if (AActor* Actor = GetActorContext())
 	{
 		// Enable the selection guard to prevent OnTreeSelectionChanged() from altering the editor's component selection
 		TGuardValue<bool> SelectionGuard(bSelectionGuard, true);
-		SCSEditor->UpdateTree();
+		SubobjectEditor->UpdateTree();
 		UpdateComponentTreeFromEditorSelection();
 	}
 
@@ -327,7 +334,7 @@ void SActorDetails::RefreshTopLevelElements(TArrayView<const TTypedElement<UType
 	}
 }
 
-void SActorDetails::RefreshSCSTreeElements(TArrayView<const TSharedPtr<class FSCSEditorTreeNode>> InSelectedNodes, const bool bForceRefresh, const bool bOverrideLock)
+void SActorDetails::RefreshSubobjectTreeElements(TArrayView<const FSubobjectEditorTreeNodePtrType> InSelectedNodes, const bool bForceRefresh, const bool bOverrideLock)
 {
 	// Nothing to do if this view is locked!
 	if (DetailsView->IsLocked() && !bOverrideLock)
@@ -335,37 +342,40 @@ void SActorDetails::RefreshSCSTreeElements(TArrayView<const TSharedPtr<class FSC
 		return;
 	}
 
-	// Does the SCS tree have components selected?
-	TArray<UActorComponent*> Components;
-	if (AActor* Actor = GetActorContext())
+	// Does the Subobject tree have components selected?
+	TArray<const UActorComponent*> Components;
+	if (const AActor* Actor = GetActorContext())
 	{
-		for (const FSCSEditorTreeNodePtrType& SelectedNode : InSelectedNodes)
+		for (const FSubobjectEditorTreeNodePtrType& SelectedNode : InSelectedNodes)
 		{
 			if (SelectedNode)
 			{
-				if (SelectedNode->GetNodeType() == FSCSEditorTreeNode::RootActorNode)
+				if(const FSubobjectData* Data = SelectedNode->GetDataSource())
 				{
-					// If the actor node is selected then we ignore the component selection
-					Components.Reset();
-					break;
-				}
-
-				if (SelectedNode->GetNodeType() == FSCSEditorTreeNode::ComponentNode)
-				{
-					if (UActorComponent* Component = SelectedNode->FindComponentInstanceInActor(Actor))
+					if(Data->IsDefaultSceneRoot())
 					{
-						Components.Add(Component);
+						// If the actor node is selected then we ignore the component selection
+                        Components.Reset();
+                        break;
+					}
+					
+					if (Data->IsComponent())
+					{
+						if (const UActorComponent* Component = Data->FindComponentInstanceInActor(Actor))
+						{
+							Components.Add(Component);
+						}
 					}
 				}
 			}
 		}
 	}
 
-	SCSTreeElements.Reset(Components.Num());
+	SubobjectTreeElements.Reset(Components.Num());
 	if (Components.Num() > 0)
 	{
 		UTypedElementRegistry* Registry = UTypedElementRegistry::GetInstance();
-		for (UActorComponent* Component : Components)
+		for (const UActorComponent* Component : Components)
 		{
 			if (FTypedElementHandle ComponentElementHandle = UEngineElementsLibrary::AcquireEditorComponentElementHandle(Component))
 			{
@@ -373,14 +383,14 @@ void SActorDetails::RefreshSCSTreeElements(TArrayView<const TSharedPtr<class FSC
 				{
 					if (TUniquePtr<ITypedElementDetailsObject> ElementDetailsObject = ComponentDetailsHandle.GetDetailsObject())
 					{
-						SCSTreeElements.Add(MoveTemp(ElementDetailsObject));
+						SubobjectTreeElements.Add(MoveTemp(ElementDetailsObject));
 					}
 				}
 			}
 		}
 
 		// Use the component elements
-		SetElementDetailsObjects(SCSTreeElements, bForceRefresh, bOverrideLock);
+		SetElementDetailsObjects(SubobjectTreeElements, bForceRefresh, bOverrideLock);
 	}
 	else
 	{
@@ -409,7 +419,7 @@ void SActorDetails::PostUndo(bool bSuccess)
 	TGuardValue<bool> SelectionGuard(bSelectionGuard, true);
 	
 	// Refresh the tree and update the selection to match the world
-	SCSEditor->UpdateTree();
+	SubobjectEditor->UpdateTree();
 	UpdateComponentTreeFromEditorSelection();
 }
 
@@ -438,11 +448,11 @@ void SActorDetails::SetActorDetailsRootCustomization(TSharedPtr<FDetailsViewObje
 	DetailsView->ForceRefresh();
 }
 
-void SActorDetails::SetSCSEditorUICustomization(TSharedPtr<ISCSEditorUICustomization> ActorDetailsSCSEditorUICustomization)
+void SActorDetails::SetSubobjectEditorUICustomization(TSharedPtr<ISCSEditorUICustomization> ActorDetailsSubobjectEditorUICustomization)
 {
-	if (SCSEditor.IsValid())
+	if(SubobjectEditor.IsValid())
 	{
-		SCSEditor->SetUICustomization(ActorDetailsSCSEditorUICustomization);
+		SubobjectEditor->SetUICustomization(ActorDetailsSubobjectEditorUICustomization);
 	}
 }
 
@@ -456,8 +466,7 @@ void SActorDetails::OnComponentsEditedInWorld()
 			TGuardValue<bool> SelectionGuard(bSelectionGuard, true);
 
 			// Refresh the tree and update the selection to match the world
-			SCSEditor->UpdateTree();
-
+			SubobjectEditor->UpdateTree();
 			DetailsView->ForceRefresh();
 		}
 	}
@@ -469,13 +478,18 @@ bool SActorDetails::GetAllowComponentTreeEditing() const
 }
 
 AActor* SActorDetails::GetActorContext() const
+{
+	return Cast<AActor>(GetObjectContext());
+}
+
+UObject* SActorDetails::GetObjectContext() const
 {
 	return TopLevelElements.Num() == 1
-		? Cast<AActor>(TopLevelElements[0]->GetObject())
+		? TopLevelElements[0]->GetObject()
 		: nullptr;
 }
 
-void SActorDetails::OnSCSEditorTreeViewSelectionChanged(const TArray<FSCSEditorTreeNodePtrType>& SelectedNodes)
+void SActorDetails::OnSubobjectEditorTreeViewSelectionChanged(const TArray<FSubobjectEditorTreeNodePtrType>& SelectedNodes)
 {
 	if (bSelectionGuard)
 	{
@@ -488,11 +502,11 @@ void SActorDetails::OnSCSEditorTreeViewSelectionChanged(const TArray<FSCSEditorT
 		// Don't respond to de-selecting everything...
 		return;
 	}
-	
-	AActor* Actor = GetActorContext();
-	if (!Actor)
+
+	AActor* ActorContext = GetActorContext();
+	if (!ActorContext)
 	{
-		// The SCS editor requires an actor context...
+		// The Subobject editor requires an actor context...
 		return;
 	}
 
@@ -501,10 +515,11 @@ void SActorDetails::OnSCSEditorTreeViewSelectionChanged(const TArray<FSCSEditorT
 		// Remove the compilation delegate if we are no longer displaying the full details for a single blueprint component.
 		RemoveBPComponentCompileEventDelegate();
 	}
-	else if (SelectedNodes.Num() == 1 && SelectedNodes[0]->GetNodeType() == FSCSEditorTreeNode::ComponentNode)
+	else if (SelectedNodes.Num() == 1 && SelectedNodes[0]->IsComponentNode())
 	{
 		// Add delegate to monitor blueprint component compilation if we have a full details view ( i.e. single selection )
-		if (UActorComponent* Component = SelectedNodes[0]->FindComponentInstanceInActor(Actor))
+		FSubobjectData* SelectedData = SelectedNodes[0]->GetDataSource();
+		if (UActorComponent* Component = const_cast<UActorComponent*>(SelectedData->FindComponentInstanceInActor(ActorContext)))
 		{
 			if (UBlueprintGeneratedClass* ComponentBPGC = Cast<UBlueprintGeneratedClass>(Component->GetClass()))
 			{
@@ -515,30 +530,33 @@ void SActorDetails::OnSCSEditorTreeViewSelectionChanged(const TArray<FSCSEditorT
 			}
 		}
 	}
-
+	
 	// We only actually update the editor selection state if we're not locked
 	if (!DetailsView->IsLocked())
 	{
 		TArray<FTypedElementHandle> NewEditorSelection;
-		NewEditorSelection.Add(UEngineElementsLibrary::AcquireEditorActorElementHandle(Actor));
+		NewEditorSelection.Add(UEngineElementsLibrary::AcquireEditorActorElementHandle(ActorContext));
 
-		for (const FSCSEditorTreeNodePtrType& SelectedNode : SelectedNodes)
+		for (const FSubobjectEditorTreeNodePtrType& SelectedNode : SelectedNodes)
 		{
 			if (SelectedNode)
 			{
-				if (SelectedNode->GetNodeType() == FSCSEditorTreeNode::RootActorNode)
-				{
-					// If the actor node is selected then we ignore the component selection
-					NewEditorSelection.Reset();
-					NewEditorSelection.Add(UEngineElementsLibrary::AcquireEditorActorElementHandle(Actor));
-					break;
-				}
-
-				if (SelectedNode->GetNodeType() == FSCSEditorTreeNode::ComponentNode)
+				if(FSubobjectData* Data = SelectedNode->GetDataSource())
 				{
-					if (UActorComponent* Component = SelectedNode->FindComponentInstanceInActor(Actor))
+					if(Data->IsDefaultSceneRoot())
 					{
-						NewEditorSelection.Add(UEngineElementsLibrary::AcquireEditorComponentElementHandle(Component));
+						// If the actor node is selected then we ignore the component selection
+						NewEditorSelection.Reset();
+						NewEditorSelection.Add(UEngineElementsLibrary::AcquireEditorActorElementHandle(ActorContext));
+						break;
+					}
+				
+					if (Data->IsInstancedComponent())
+					{
+						if (const UActorComponent* Component = Data->FindComponentInstanceInActor(ActorContext))
+						{
+							NewEditorSelection.Add(UEngineElementsLibrary::AcquireEditorComponentElementHandle(Component));
+						}
 					}
 				}
 			}
@@ -548,21 +566,28 @@ void SActorDetails::OnSCSEditorTreeViewSelectionChanged(const TArray<FSCSEditorT
 		const bool bShouldActuallyTransact = !GIsTransacting;
 		const FScopedTransaction Transaction(NSLOCTEXT("UnrealEd", "ClickingOnComponentInTree", "Clicking on Component (tree view)"), bShouldActuallyTransact);
 
-		// Enable the selection guard to prevent OnEditorSelectionChanged() from altering the contents of the SCSTreeWidget
+		// Enable the selection guard to prevent OnEditorSelectionChanged() from altering the contents of the SubobjectTreeWidget
 		TGuardValue<bool> SelectionGuard(bSelectionGuard, true);
 		SelectionSet->SetSelection(NewEditorSelection, FTypedElementSelectionOptions());
 		SelectionSet->NotifyPendingChanges(); // Fire while still under the selection guard
 	}
 
 	// Update the underlying details view
-	RefreshSCSTreeElements(SelectedNodes, /*bForceRefresh*/false, DetailsView->IsLocked());
+	RefreshSubobjectTreeElements(SelectedNodes, /*bForceRefresh*/false, DetailsView->IsLocked());
+}
+
+void SActorDetails::OnSubobjectEditorTreeViewObjectReplaced()
+{
+	// Enable the selection guard to prevent OnTreeSelectionChanged() from altering the editor's component selection
+	TGuardValue<bool> SelectionGuard(bSelectionGuard, true);
+	SubobjectEditor->UpdateTree();
 }
 
-void SActorDetails::OnSCSEditorTreeViewItemDoubleClicked(const TSharedPtr<class FSCSEditorTreeNode> ClickedNode)
+void SActorDetails::OnSubobjectEditorTreeViewItemDoubleClicked(const FSubobjectEditorTreeNodePtrType ClickedNode)
 {
-	if (ClickedNode && ClickedNode->GetNodeType() == FSCSEditorTreeNode::ComponentNode)
+	if (ClickedNode && ClickedNode->IsComponentNode())
 	{
-		if (USceneComponent* SceneComponent = Cast<USceneComponent>(ClickedNode->GetComponentTemplate()))
+		if (const USceneComponent* SceneComponent = Cast<USceneComponent>(ClickedNode->GetComponentTemplate()))
 		{
 			const bool bActiveViewportOnly = false;
 			GEditor->MoveViewportCamerasToComponent(SceneComponent, bActiveViewportOnly);
@@ -570,13 +595,6 @@ void SActorDetails::OnSCSEditorTreeViewItemDoubleClicked(const TSharedPtr<class
 	}
 }
 
-void SActorDetails::OnSCSEditorTreeViewObjectReplaced()
-{
-	// Enable the selection guard to prevent OnTreeSelectionChanged() from altering the editor's component selection
-	TGuardValue<bool> SelectionGuard(bSelectionGuard, true);
-	SCSEditor->UpdateTree();
-}
-
 void SActorDetails::UpdateComponentTreeFromEditorSelection()
 {
 	if (DetailsView->IsLocked())
@@ -587,39 +605,39 @@ void SActorDetails::UpdateComponentTreeFromEditorSelection()
 	// Enable the selection guard to prevent OnTreeSelectionChanged() from altering the editor's component selection
 	TGuardValue<bool> SelectionGuard(bSelectionGuard, true);
 
-	TSharedPtr<SSCSTreeType>& SCSTreeWidget = SCSEditor->SCSTreeWidget;
+	TSharedPtr<SSubobjectEditorDragDropTree> TreeWidget = SubobjectEditor->GetDragDropTree();
 
 	// Update the tree selection to match the level editor component selection
-	SCSTreeWidget->ClearSelection();
-	SelectionSet->ForEachSelectedObject<UActorComponent>([this, &SCSTreeWidget](UActorComponent* InComponent)
+	SubobjectEditor->ClearSelection();
+	SelectionSet->ForEachSelectedObject<UActorComponent>([this, &TreeWidget](UActorComponent* InComponent)
 	{
-		FSCSEditorTreeNodePtrType SCSTreeNode = SCSEditor->GetNodeFromActorComponent(InComponent, false);
-		if (SCSTreeNode && SCSTreeNode->GetComponentTemplate())
+		FSubobjectEditorTreeNodePtrType TreeNode = SubobjectEditor->FindSlateNodeForObject(InComponent, false);
+		if (TreeNode && TreeNode->GetComponentTemplate())
 		{
-			SCSTreeWidget->RequestScrollIntoView(SCSTreeNode);
-			SCSTreeWidget->SetItemSelection(SCSTreeNode, true);
-			check(InComponent == SCSTreeNode->GetComponentTemplate());
+			TreeWidget->RequestScrollIntoView(TreeNode);
+			TreeWidget->SetItemSelection(TreeNode, true);
+			check(InComponent == TreeNode->GetComponentTemplate());
 		}
 		return true;
 	});
 
-	TArray<TSharedPtr<FSCSEditorTreeNode>> SelectedNodes = SCSEditor->GetSelectedNodes();
+	TArray<FSubobjectEditorTreeNodePtrType> SelectedNodes = SubobjectEditor->GetSelectedNodes();
 	if (SelectedNodes.Num() == 0)
 	{
-		SCSEditor->SelectRoot();
-		SelectedNodes = SCSEditor->GetSelectedNodes();
+		SubobjectEditor->SelectRoot();
+		SelectedNodes = SubobjectEditor->GetSelectedNodes();
 	}
 
 	// Update the underlying details view
-	RefreshSCSTreeElements(SelectedNodes, bSelectedComponentRecompiled, /*bOverrideLock*/false);
+	RefreshSubobjectTreeElements(SelectedNodes, bSelectedComponentRecompiled, /*bOverrideLock*/false);
 }
 
 bool SActorDetails::IsPropertyReadOnly(const FPropertyAndParent& PropertyAndParent) const
 {
 	bool bIsReadOnly = false;
-	for (const FSCSEditorTreeNodePtrType& Node : SCSEditor->GetSelectedNodes())
+	for (const FSubobjectEditorTreeNodePtrType& Node : SubobjectEditor->GetSelectedNodes())
 	{
-		UActorComponent* Component = Node->GetComponentTemplate();
+		const UActorComponent* Component = Node->GetComponentTemplate();
 		if (Component && Component->CreationMethod == EComponentCreationMethod::SimpleConstructionScript)
 		{
 			TSet<const FProperty*> UCSModifiedProperties;
@@ -644,12 +662,15 @@ bool SActorDetails::IsPropertyEditingEnabled() const
 	}
 
 	bool bIsEditable = true;
-	for (const FSCSEditorTreeNodePtrType& Node : SCSEditor->GetSelectedNodes())
+	for (const FSubobjectEditorTreeNodePtrType& Node : SubobjectEditor->GetSelectedNodes())
 	{
-		bIsEditable = Node->CanEdit();
-		if (!bIsEditable)
+		if(const FSubobjectData* Data = Node->GetDataSource())
 		{
-			break;
+			bIsEditable = Data->CanEdit();
+			if (!bIsEditable)
+			{
+				break;
+			}
 		}
 	}
 	return bIsEditable;
@@ -657,7 +678,7 @@ bool SActorDetails::IsPropertyEditingEnabled() const
 
 void SActorDetails::OnBlueprintedComponentWarningHyperlinkClicked(const FSlateHyperlinkRun::FMetadata& Metadata)
 {
-	UBlueprint* Blueprint = SCSEditor->GetBlueprint();
+	UBlueprint* Blueprint = SubobjectEditor->GetBlueprint();
 	if (Blueprint)
 	{
 		// Open the blueprint
@@ -668,7 +689,7 @@ void SActorDetails::OnBlueprintedComponentWarningHyperlinkClicked(const FSlateHy
 void SActorDetails::OnNativeComponentWarningHyperlinkClicked(const FSlateHyperlinkRun::FMetadata& Metadata)
 {
 	// Find the closest native parent
-	UBlueprint* Blueprint = SCSEditor->GetBlueprint();
+	UBlueprint* Blueprint = SubobjectEditor->GetBlueprint();
 	UClass* ParentClass = Blueprint ? *Blueprint->ParentClass : GetActorContext()->GetClass();
 	while (ParentClass && !ParentClass->HasAllClassFlags(CLASS_Native))
 	{
@@ -698,23 +719,27 @@ EVisibility SActorDetails::GetUCSComponentWarningVisibility() const
 	bool bIsUneditableBlueprintComponent = false;
 
 	// Check to see if any selected components are inherited from blueprint
-	for (const FSCSEditorTreeNodePtrType& Node : SCSEditor->GetSelectedNodes())
+	for (const FSubobjectEditorTreeNodePtrType& Node : SubobjectEditor->GetSelectedNodes())
 	{
-		if (!Node->IsNativeComponent())
+		if(const FSubobjectData* Data = Node->GetDataSource())
 		{
-			UActorComponent* Component = Node->GetComponentTemplate();
-			bIsUneditableBlueprintComponent = Component ? Component->CreationMethod == EComponentCreationMethod::UserConstructionScript : false;
-			if (bIsUneditableBlueprintComponent)
+			if (!Data->IsNativeComponent())
 			{
-				break;
+				const UActorComponent* Component = Data->GetComponentTemplate();
+				bIsUneditableBlueprintComponent = Component ? Component->CreationMethod == EComponentCreationMethod::UserConstructionScript : false;
+				if (bIsUneditableBlueprintComponent)
+				{
+					break;
+				}
 			}
 		}
+		
 	}
 
 	return bIsUneditableBlueprintComponent ? EVisibility::Visible : EVisibility::Collapsed;
 }
 
-bool NotEditableSetByBlueprint(UActorComponent* Component)
+bool NotEditableSetByBlueprint(const UActorComponent* Component)
 {
 	// Determine if it is locked out from a blueprint or from the native
 	UActorComponent* Archetype = CastChecked<UActorComponent>(Component->GetArchetype());
@@ -743,24 +768,27 @@ EVisibility SActorDetails::GetInheritedBlueprintComponentWarningVisibility() con
 	bool bIsUneditableBlueprintComponent = false;
 
 	// Check to see if any selected components are inherited from blueprint
-	for (const FSCSEditorTreeNodePtrType& Node : SCSEditor->GetSelectedNodes())
+	for (const FSubobjectEditorTreeNodePtrType& Node : SubobjectEditor->GetSelectedNodes())
 	{
-		if (!Node->IsNativeComponent())
+		if(const FSubobjectData* Data = Node->GetDataSource())
 		{
-			if (UActorComponent* Component = Node->GetComponentTemplate())
+			if (!Data->IsNativeComponent())
 			{
-				if (!Component->IsEditableWhenInherited() && Component->CreationMethod == EComponentCreationMethod::SimpleConstructionScript)
+				if (const UActorComponent* Component = Data->GetComponentTemplate())
 				{
-					bIsUneditableBlueprintComponent = true;
-					break;
+					if (!Component->IsEditableWhenInherited() && Component->CreationMethod == EComponentCreationMethod::SimpleConstructionScript)
+					{
+						bIsUneditableBlueprintComponent = true;
+						break;
+					}
 				}
 			}
-		}
-		else if (!Node->CanEdit() && NotEditableSetByBlueprint(Node->GetComponentTemplate()))
-		{
-			bIsUneditableBlueprintComponent = true;
-			break;
-		}
+			else if (!Data->CanEdit() && NotEditableSetByBlueprint(Data->GetComponentTemplate()))
+			{
+				bIsUneditableBlueprintComponent = true;
+				break;
+			}
+		}	
 	}
 
 	return bIsUneditableBlueprintComponent ? EVisibility::Visible : EVisibility::Collapsed;
@@ -769,14 +797,18 @@ EVisibility SActorDetails::GetInheritedBlueprintComponentWarningVisibility() con
 EVisibility SActorDetails::GetNativeComponentWarningVisibility() const
 {
 	bool bIsUneditableNative = false;
-	for (const FSCSEditorTreeNodePtrType& Node : SCSEditor->GetSelectedNodes())
+	for (const FSubobjectEditorTreeNodePtrType& Node : SubobjectEditor->GetSelectedNodes())
 	{
-		// Check to see if the component is native and not editable
-		if (Node->IsNativeComponent() && !Node->CanEdit() && !NotEditableSetByBlueprint(Node->GetComponentTemplate()))
+		if(const FSubobjectData* Data = Node->GetDataSource())
 		{
-			bIsUneditableNative = true;
-			break;
+			// Check to see if the component is native and not editable
+			if (Data->IsNativeComponent() && !Data->CanEdit() && !NotEditableSetByBlueprint(Data->GetComponentTemplate()))
+			{
+				bIsUneditableNative = true;
+				break;
+			}
 		}
+
 	}
 	
 	return bIsUneditableNative ? EVisibility::Visible : EVisibility::Collapsed;
diff --git a/Engine/Source/Editor/LevelEditor/Private/SActorDetails.h b/Engine/Source/Editor/LevelEditor/Private/SActorDetails.h
--- a/Engine/Source/Editor/LevelEditor/Private/SActorDetails.h
+++ b/Engine/Source/Editor/LevelEditor/Private/SActorDetails.h
@@ -17,11 +17,13 @@ class FTabManager;
 class FUICommandList;
 class IDetailsView;
 class SBox;
-class SSCSEditor;
 class SSplitter;
 class UBlueprint;
 class FDetailsViewObjectFilter;
 class UTypedElementSelectionSet;
+class SSubobjectEditor;
+class ISCSEditorUICustomization;
+class FSubobjectEditorTreeNode;
 
 /**
  * Wraps a details panel customized for viewing actors
@@ -65,21 +67,25 @@ public:
 	 */
 	void SetActorDetailsRootCustomization(TSharedPtr<FDetailsViewObjectFilter> ActorDetailsObjectFilter, TSharedPtr<class IDetailRootObjectCustomization> ActorDetailsRootCustomization);
 
-	/** Sets the UI customization of the SCSEditor inside this details panel. */
-	void SetSCSEditorUICustomization(TSharedPtr<class ISCSEditorUICustomization> ActorDetailsSCSEditorUICustomization);
+	UE_DEPRECATED(5.0, "SetSCSEditorUICustomization is deprecated, please use SetSubobjectEditorUICustomization instead.")
+	void SetSCSEditorUICustomization(TSharedPtr<ISCSEditorUICustomization> ActorDetailsSCSEditorUICustomization) { SetSubobjectEditorUICustomization(ActorDetailsSCSEditorUICustomization); }
+	
+	/** Sets the UI customization of the Suobject inside this details panel. */
+	void SetSubobjectEditorUICustomization(TSharedPtr<ISCSEditorUICustomization> ActorDetailsSCSEditorUICustomization);
 
 private:
 	void RefreshTopLevelElements(TArrayView<const TTypedElement<UTypedElementDetailsInterface>> InDetailsElements, const bool bForceRefresh, const bool bOverrideLock);
-	void RefreshSCSTreeElements(TArrayView<const TSharedPtr<class FSCSEditorTreeNode>> InSelectedNodes, const bool bForceRefresh, const bool bOverrideLock);
+	void RefreshSubobjectTreeElements(TArrayView<const TSharedPtr<FSubobjectEditorTreeNode>> InSelectedNodes, const bool bForceRefresh, const bool bOverrideLock);
 	void SetElementDetailsObjects(TArrayView<const TUniquePtr<ITypedElementDetailsObject>> InElementDetailsObjects, const bool bForceRefresh, const bool bOverrideLock);
 
 	AActor* GetActorContext() const;
+	UObject* GetObjectContext() const;
 	bool GetAllowComponentTreeEditing() const;
 
 	void OnComponentsEditedInWorld();
-	void OnSCSEditorTreeViewSelectionChanged(const TArray<TSharedPtr<class FSCSEditorTreeNode> >& SelectedNodes);
-	void OnSCSEditorTreeViewItemDoubleClicked(const TSharedPtr<class FSCSEditorTreeNode> ClickedNode);
-	void OnSCSEditorTreeViewObjectReplaced();
+	void OnSubobjectEditorTreeViewSelectionChanged(const TArray<TSharedPtr<FSubobjectEditorTreeNode> >& SelectedNodes);
+	void OnSubobjectEditorTreeViewItemDoubleClicked(const TSharedPtr<FSubobjectEditorTreeNode> ClickedNode);
+	void OnSubobjectEditorTreeViewObjectReplaced();
 	void UpdateComponentTreeFromEditorSelection();
 
 	bool IsPropertyReadOnly(const struct FPropertyAndParent& PropertyAndParent) const;
@@ -100,7 +106,9 @@ private:
 	TSharedPtr<SSplitter> DetailsSplitter;
 	TSharedPtr<class IDetailsView> DetailsView;
 	TSharedPtr<SBox> ComponentsBox;
-	TSharedPtr<class SSCSEditor> SCSEditor;
+
+	/** The subobject editor provides a tree widget that allows for editing of subobjects */
+	TSharedPtr<SSubobjectEditor> SubobjectEditor;
 
 	// The selection set this details panel is observing
 	UTypedElementSelectionSet* SelectionSet = nullptr;
@@ -112,9 +120,9 @@ private:
 	// Array of top-level elements that are currently being edited
 	TArray<TUniquePtr<ITypedElementDetailsObject>> TopLevelElements;
 
-	// Array of component elements that are being edited from the SCS tree selection
-	TArray<TUniquePtr<ITypedElementDetailsObject>> SCSTreeElements;
-
+	// Array of component elements that are being edited from the Subobject tree selection
+	TArray<TUniquePtr<ITypedElementDetailsObject>> SubobjectTreeElements;
+	
 	// The current component blueprint selection
 	TWeakObjectPtr<UBlueprint> SelectedBPComponentBlueprint;
 	bool bSelectedComponentRecompiled = false;
diff --git a/Engine/Source/Editor/LevelEditor/Private/SLevelEditor.cpp b/Engine/Source/Editor/LevelEditor/Private/SLevelEditor.cpp
--- a/Engine/Source/Editor/LevelEditor/Private/SLevelEditor.cpp
+++ b/Engine/Source/Editor/LevelEditor/Private/SLevelEditor.cpp
@@ -1969,7 +1969,7 @@ TSharedRef<SWidget> SLevelEditor::CreateActorDetails( const FName TabIdentifier
 	TSharedRef<SActorDetails> ActorDetails = SNew( SActorDetails, GetMutableElementSelectionSet(), TabIdentifier, LevelEditorCommands, GetTabManager() );
 
 	ActorDetails->SetActorDetailsRootCustomization(ActorDetailsObjectFilter, ActorDetailsRootCustomization);
-	ActorDetails->SetSCSEditorUICustomization(ActorDetailsSCSEditorUICustomization);
+	ActorDetails->SetSubobjectEditorUICustomization(ActorDetailsSCSEditorUICustomization);
 
 	AllActorDetailPanels.Add( ActorDetails );
 	return ActorDetails;
@@ -2018,7 +2018,7 @@ void SLevelEditor::SetActorDetailsSCSEditorUICustomization(TSharedPtr<ISCSEditor
 
 	for (TSharedRef<SActorDetails> ActorDetails : GetAllActorDetails())
 	{
-		ActorDetails->SetSCSEditorUICustomization(ActorDetailsSCSEditorUICustomization);
+		ActorDetails->SetSubobjectEditorUICustomization(ActorDetailsSCSEditorUICustomization);
 	}
 }
 
