Undo 16223481 because subobjects would be incorrectly marked as non-editable sometimes.

#rnx

diff --git a/Engine/Source/Editor/LevelEditor/Private/SActorDetails.cpp b/Engine/Source/Editor/LevelEditor/Private/SActorDetails.cpp
--- a/Engine/Source/Editor/LevelEditor/Private/SActorDetails.cpp
+++ b/Engine/Source/Editor/LevelEditor/Private/SActorDetails.cpp
@@ -22,6 +22,7 @@
 #include "UnrealEdGlobals.h"
 #include "LevelEditor.h"
 #include "Kismet2/KismetEditorUtilities.h"
+#include "SSCSEditor.h"
 #include "PropertyEditorModule.h"
 #include "IDetailsView.h"
 #include "LevelEditorGenericDetails.h"
@@ -33,11 +34,6 @@
 #include "Elements/Framework/TypedElementRegistry.h"
 #include "Elements/Framework/TypedElementSelectionSet.h"
 #include "Elements/Framework/EngineElementsLibrary.h"
-#include "SSubobjectEditor.h"
-#include "SSubobjectInstanceEditor.h"
-#include "SubobjectData.h"
-#include "SubobjectDataSubsystem.h"
-#include "SSubobjectEditorModule.h"
 
 class SActorDetailsUneditableComponentWarning : public SCompoundWidget
 {
@@ -136,21 +132,19 @@ void SActorDetails::Construct(const FArguments& InArgs, UTypedElementSelectionSe
 	GEditor->RegisterForUndo(this);
 
 	ComponentsBox = SNew(SBox)
-		.Padding(FMargin(2.0f, 0.0f, 2.0f, 0.0f))
-		.Visibility(this, &SActorDetails::GetComponentsBoxVisibility);
-
-	FModuleManager::LoadModuleChecked<FSubobjectEditorModule>("SubobjectEditor");
-	
-	SubobjectEditor = SNew(SSubobjectInstanceEditor)
-		.ObjectContext(this, &SActorDetails::GetObjectContext)
-		.AllowEditing(this, &SActorDetails::GetAllowComponentTreeEditing)
-		.OnSelectionUpdated(this, &SActorDetails::OnSubobjectEditorTreeViewSelectionChanged)
-		.OnItemDoubleClicked(this, &SActorDetails::OnSubobjectEditorTreeViewItemDoubleClicked)
-		.OnObjectReplaced(this, &SActorDetails::OnSubobjectEditorTreeViewObjectReplaced);
-
-	ComponentsBox->SetContent(SubobjectEditor.ToSharedRef());
+		.Padding(FMargin(2,0,2,0))
+		.Visibility(this, &SActorDetails::GetComponentsBoxVisibility)
+		[
+			SAssignNew(SCSEditor, SSCSEditor)
+			.EditorMode(EComponentEditorMode::ActorInstance)
+			.AllowEditing(this, &SActorDetails::GetAllowComponentTreeEditing)
+			.ActorContext(this, &SActorDetails::GetActorContext)
+			.OnSelectionUpdated(this, &SActorDetails::OnSCSEditorTreeViewSelectionChanged)
+			.OnItemDoubleClicked(this, &SActorDetails::OnSCSEditorTreeViewItemDoubleClicked)
+			.OnObjectReplaced(this, &SActorDetails::OnSCSEditorTreeViewObjectReplaced)
+		];
 
-	TSharedRef<SWidget> ButtonBox = SubobjectEditor->GetToolButtonsBox().ToSharedRef();
+	TSharedRef<SWidget> ButtonBox = SCSEditor->GetToolButtonsBox().ToSharedRef();
 	DetailsView->SetNameAreaCustomContent( ButtonBox );
 
 	ChildSlot
@@ -313,13 +307,12 @@ void SActorDetails::RefreshTopLevelElements(TArrayView<const TTypedElement<UType
 	// Update the underlying details view
 	SetElementDetailsObjects(TopLevelElements, bForceRefresh, bOverrideLock);
 
-	// Update the Subobject tree if we were asked to edit a single actor
-	
+	// Update the SCS tree if we were asked to edit a single actor
 	if (AActor* Actor = GetActorContext())
 	{
 		// Enable the selection guard to prevent OnTreeSelectionChanged() from altering the editor's component selection
 		TGuardValue<bool> SelectionGuard(bSelectionGuard, true);
-		SubobjectEditor->UpdateTree();
+		SCSEditor->UpdateTree();
 		UpdateComponentTreeFromEditorSelection();
 	}
 
@@ -334,7 +327,7 @@ void SActorDetails::RefreshTopLevelElements(TArrayView<const TTypedElement<UType
 	}
 }
 
-void SActorDetails::RefreshSubobjectTreeElements(TArrayView<const FSubobjectEditorTreeNodePtrType> InSelectedNodes, const bool bForceRefresh, const bool bOverrideLock)
+void SActorDetails::RefreshSCSTreeElements(TArrayView<const TSharedPtr<class FSCSEditorTreeNode>> InSelectedNodes, const bool bForceRefresh, const bool bOverrideLock)
 {
 	// Nothing to do if this view is locked!
 	if (DetailsView->IsLocked() && !bOverrideLock)
@@ -342,40 +335,37 @@ void SActorDetails::RefreshSubobjectTreeElements(TArrayView<const FSubobjectEdit
 		return;
 	}
 
-	// Does the Subobject tree have components selected?
-	TArray<const UActorComponent*> Components;
-	if (const AActor* Actor = GetActorContext())
+	// Does the SCS tree have components selected?
+	TArray<UActorComponent*> Components;
+	if (AActor* Actor = GetActorContext())
 	{
-		for (const FSubobjectEditorTreeNodePtrType& SelectedNode : InSelectedNodes)
+		for (const FSCSEditorTreeNodePtrType& SelectedNode : InSelectedNodes)
 		{
 			if (SelectedNode)
 			{
-				if(const FSubobjectData* Data = SelectedNode->GetDataSource())
+				if (SelectedNode->GetNodeType() == FSCSEditorTreeNode::RootActorNode)
 				{
-					if(Data->IsDefaultSceneRoot())
-					{
-						// If the actor node is selected then we ignore the component selection
-                        Components.Reset();
-                        break;
-					}
-					
-					if (Data->IsComponent())
+					// If the actor node is selected then we ignore the component selection
+					Components.Reset();
+					break;
+				}
+
+				if (SelectedNode->GetNodeType() == FSCSEditorTreeNode::ComponentNode)
+				{
+					if (UActorComponent* Component = SelectedNode->FindComponentInstanceInActor(Actor))
 					{
-						if (const UActorComponent* Component = Data->FindComponentInstanceInActor(Actor))
-						{
-							Components.Add(Component);
-						}
+						Components.Add(Component);
 					}
 				}
 			}
 		}
 	}
 
-	SubobjectTreeElements.Reset(Components.Num());
+	SCSTreeElements.Reset(Components.Num());
 	if (Components.Num() > 0)
 	{
 		UTypedElementRegistry* Registry = UTypedElementRegistry::GetInstance();
-		for (const UActorComponent* Component : Components)
+		for (UActorComponent* Component : Components)
 		{
 			if (FTypedElementHandle ComponentElementHandle = UEngineElementsLibrary::AcquireEditorComponentElementHandle(Component))
 			{
@@ -383,14 +373,14 @@ void SActorDetails::RefreshSubobjectTreeElements(TArrayView<const FSubobjectEdit
 				{
 					if (TUniquePtr<ITypedElementDetailsObject> ElementDetailsObject = ComponentDetailsHandle.GetDetailsObject())
 					{
-						SubobjectTreeElements.Add(MoveTemp(ElementDetailsObject));
+						SCSTreeElements.Add(MoveTemp(ElementDetailsObject));
 					}
 				}
 			}
 		}
 
 		// Use the component elements
-		SetElementDetailsObjects(SubobjectTreeElements, bForceRefresh, bOverrideLock);
+		SetElementDetailsObjects(SCSTreeElements, bForceRefresh, bOverrideLock);
 	}
 	else
 	{
@@ -419,7 +409,7 @@ void SActorDetails::PostUndo(bool bSuccess)
 	TGuardValue<bool> SelectionGuard(bSelectionGuard, true);
 	
 	// Refresh the tree and update the selection to match the world
-	SubobjectEditor->UpdateTree();
+	SCSEditor->UpdateTree();
 	UpdateComponentTreeFromEditorSelection();
 }
 
@@ -448,11 +438,11 @@ void SActorDetails::SetActorDetailsRootCustomization(TSharedPtr<FDetailsViewObje
 	DetailsView->ForceRefresh();
 }
 
-void SActorDetails::SetSubobjectEditorUICustomization(TSharedPtr<ISCSEditorUICustomization> ActorDetailsSubobjectEditorUICustomization)
+void SActorDetails::SetSCSEditorUICustomization(TSharedPtr<ISCSEditorUICustomization> ActorDetailsSCSEditorUICustomization)
 {
-	if(SubobjectEditor.IsValid())
+	if (SCSEditor.IsValid())
 	{
-		SubobjectEditor->SetUICustomization(ActorDetailsSubobjectEditorUICustomization);
+		SCSEditor->SetUICustomization(ActorDetailsSCSEditorUICustomization);
 	}
 }
 
@@ -466,7 +456,8 @@ void SActorDetails::OnComponentsEditedInWorld()
 			TGuardValue<bool> SelectionGuard(bSelectionGuard, true);
 
 			// Refresh the tree and update the selection to match the world
-			SubobjectEditor->UpdateTree();
+			SCSEditor->UpdateTree();
+
 			DetailsView->ForceRefresh();
 		}
 	}
@@ -478,18 +469,13 @@ bool SActorDetails::GetAllowComponentTreeEditing() const
 }
 
 AActor* SActorDetails::GetActorContext() const
-{
-	return Cast<AActor>(GetObjectContext());
-}
-
-UObject* SActorDetails::GetObjectContext() const
 {
 	return TopLevelElements.Num() == 1
-		? TopLevelElements[0]->GetObject()
+		? Cast<AActor>(TopLevelElements[0]->GetObject())
 		: nullptr;
 }
 
-void SActorDetails::OnSubobjectEditorTreeViewSelectionChanged(const TArray<FSubobjectEditorTreeNodePtrType>& SelectedNodes)
+void SActorDetails::OnSCSEditorTreeViewSelectionChanged(const TArray<FSCSEditorTreeNodePtrType>& SelectedNodes)
 {
 	if (bSelectionGuard)
 	{
@@ -502,11 +488,11 @@ void SActorDetails::OnSubobjectEditorTreeViewSelectionChanged(const TArray<FSubo
 		// Don't respond to de-selecting everything...
 		return;
 	}
-
-	AActor* ActorContext = GetActorContext();
-	if (!ActorContext)
+	
+	AActor* Actor = GetActorContext();
+	if (!Actor)
 	{
-		// The Subobject editor requires an actor context...
+		// The SCS editor requires an actor context...
 		return;
 	}
 
@@ -515,11 +501,10 @@ void SActorDetails::OnSubobjectEditorTreeViewSelectionChanged(const TArray<FSubo
 		// Remove the compilation delegate if we are no longer displaying the full details for a single blueprint component.
 		RemoveBPComponentCompileEventDelegate();
 	}
-	else if (SelectedNodes.Num() == 1 && SelectedNodes[0]->IsComponentNode())
+	else if (SelectedNodes.Num() == 1 && SelectedNodes[0]->GetNodeType() == FSCSEditorTreeNode::ComponentNode)
 	{
 		// Add delegate to monitor blueprint component compilation if we have a full details view ( i.e. single selection )
-		FSubobjectData* SelectedData = SelectedNodes[0]->GetDataSource();
-		if (UActorComponent* Component = const_cast<UActorComponent*>(SelectedData->FindComponentInstanceInActor(ActorContext)))
+		if (UActorComponent* Component = SelectedNodes[0]->FindComponentInstanceInActor(Actor))
 		{
 			if (UBlueprintGeneratedClass* ComponentBPGC = Cast<UBlueprintGeneratedClass>(Component->GetClass()))
 			{
@@ -530,33 +515,30 @@ void SActorDetails::OnSubobjectEditorTreeViewSelectionChanged(const TArray<FSubo
 			}
 		}
 	}
-	
+
 	// We only actually update the editor selection state if we're not locked
 	if (!DetailsView->IsLocked())
 	{
 		TArray<FTypedElementHandle> NewEditorSelection;
-		NewEditorSelection.Add(UEngineElementsLibrary::AcquireEditorActorElementHandle(ActorContext));
+		NewEditorSelection.Add(UEngineElementsLibrary::AcquireEditorActorElementHandle(Actor));
 
-		for (const FSubobjectEditorTreeNodePtrType& SelectedNode : SelectedNodes)
+		for (const FSCSEditorTreeNodePtrType& SelectedNode : SelectedNodes)
 		{
 			if (SelectedNode)
 			{
-				if(FSubobjectData* Data = SelectedNode->GetDataSource())
+				if (SelectedNode->GetNodeType() == FSCSEditorTreeNode::RootActorNode)
 				{
-					if(Data->IsDefaultSceneRoot())
-					{
-						// If the actor node is selected then we ignore the component selection
-						NewEditorSelection.Reset();
-						NewEditorSelection.Add(UEngineElementsLibrary::AcquireEditorActorElementHandle(ActorContext));
-						break;
-					}
-				
-					if (Data->IsInstancedComponent())
+					// If the actor node is selected then we ignore the component selection
+					NewEditorSelection.Reset();
+					NewEditorSelection.Add(UEngineElementsLibrary::AcquireEditorActorElementHandle(Actor));
+					break;
+				}
+
+				if (SelectedNode->GetNodeType() == FSCSEditorTreeNode::ComponentNode)
+				{
+					if (UActorComponent* Component = SelectedNode->FindComponentInstanceInActor(Actor))
 					{
-						if (const UActorComponent* Component = Data->FindComponentInstanceInActor(ActorContext))
-						{
-							NewEditorSelection.Add(UEngineElementsLibrary::AcquireEditorComponentElementHandle(Component));
-						}
+						NewEditorSelection.Add(UEngineElementsLibrary::AcquireEditorComponentElementHandle(Component));
 					}
 				}
 			}
@@ -566,28 +548,21 @@ void SActorDetails::OnSubobjectEditorTreeViewSelectionChanged(const TArray<FSubo
 		const bool bShouldActuallyTransact = !GIsTransacting;
 		const FScopedTransaction Transaction(NSLOCTEXT("UnrealEd", "ClickingOnComponentInTree", "Clicking on Component (tree view)"), bShouldActuallyTransact);
 
-		// Enable the selection guard to prevent OnEditorSelectionChanged() from altering the contents of the SubobjectTreeWidget
+		// Enable the selection guard to prevent OnEditorSelectionChanged() from altering the contents of the SCSTreeWidget
 		TGuardValue<bool> SelectionGuard(bSelectionGuard, true);
 		SelectionSet->SetSelection(NewEditorSelection, FTypedElementSelectionOptions());
 		SelectionSet->NotifyPendingChanges(); // Fire while still under the selection guard
 	}
 
 	// Update the underlying details view
-	RefreshSubobjectTreeElements(SelectedNodes, /*bForceRefresh*/false, DetailsView->IsLocked());
-}
-
-void SActorDetails::OnSubobjectEditorTreeViewObjectReplaced()
-{
-	// Enable the selection guard to prevent OnTreeSelectionChanged() from altering the editor's component selection
-	TGuardValue<bool> SelectionGuard(bSelectionGuard, true);
-	SubobjectEditor->UpdateTree();
+	RefreshSCSTreeElements(SelectedNodes, /*bForceRefresh*/false, DetailsView->IsLocked());
 }
 
-void SActorDetails::OnSubobjectEditorTreeViewItemDoubleClicked(const FSubobjectEditorTreeNodePtrType ClickedNode)
+void SActorDetails::OnSCSEditorTreeViewItemDoubleClicked(const TSharedPtr<class FSCSEditorTreeNode> ClickedNode)
 {
-	if (ClickedNode && ClickedNode->IsComponentNode())
+	if (ClickedNode && ClickedNode->GetNodeType() == FSCSEditorTreeNode::ComponentNode)
 	{
-		if (const USceneComponent* SceneComponent = Cast<USceneComponent>(ClickedNode->GetComponentTemplate()))
+		if (USceneComponent* SceneComponent = Cast<USceneComponent>(ClickedNode->GetComponentTemplate()))
 		{
 			const bool bActiveViewportOnly = false;
 			GEditor->MoveViewportCamerasToComponent(SceneComponent, bActiveViewportOnly);
@@ -595,6 +570,13 @@ void SActorDetails::OnSubobjectEditorTreeViewItemDoubleClicked(const FSubobjectE
 	}
 }
 
+void SActorDetails::OnSCSEditorTreeViewObjectReplaced()
+{
+	// Enable the selection guard to prevent OnTreeSelectionChanged() from altering the editor's component selection
+	TGuardValue<bool> SelectionGuard(bSelectionGuard, true);
+	SCSEditor->UpdateTree();
+}
+
 void SActorDetails::UpdateComponentTreeFromEditorSelection()
 {
 	if (DetailsView->IsLocked())
@@ -605,39 +587,39 @@ void SActorDetails::UpdateComponentTreeFromEditorSelection()
 	// Enable the selection guard to prevent OnTreeSelectionChanged() from altering the editor's component selection
 	TGuardValue<bool> SelectionGuard(bSelectionGuard, true);
 
-	TSharedPtr<SSubobjectEditorDragDropTree> TreeWidget = SubobjectEditor->GetDragDropTree();
+	TSharedPtr<SSCSTreeType>& SCSTreeWidget = SCSEditor->SCSTreeWidget;
 
 	// Update the tree selection to match the level editor component selection
-	SubobjectEditor->ClearSelection();
-	SelectionSet->ForEachSelectedObject<UActorComponent>([this, &TreeWidget](UActorComponent* InComponent)
+	SCSTreeWidget->ClearSelection();
+	SelectionSet->ForEachSelectedObject<UActorComponent>([this, &SCSTreeWidget](UActorComponent* InComponent)
 	{
-		FSubobjectEditorTreeNodePtrType TreeNode = SubobjectEditor->FindSlateNodeForObject(InComponent, false);
-		if (TreeNode && TreeNode->GetComponentTemplate())
+		FSCSEditorTreeNodePtrType SCSTreeNode = SCSEditor->GetNodeFromActorComponent(InComponent, false);
+		if (SCSTreeNode && SCSTreeNode->GetComponentTemplate())
 		{
-			TreeWidget->RequestScrollIntoView(TreeNode);
-			TreeWidget->SetItemSelection(TreeNode, true);
-			check(InComponent == TreeNode->GetComponentTemplate());
+			SCSTreeWidget->RequestScrollIntoView(SCSTreeNode);
+			SCSTreeWidget->SetItemSelection(SCSTreeNode, true);
+			check(InComponent == SCSTreeNode->GetComponentTemplate());
 		}
 		return true;
 	});
 
-	TArray<FSubobjectEditorTreeNodePtrType> SelectedNodes = SubobjectEditor->GetSelectedNodes();
+	TArray<TSharedPtr<FSCSEditorTreeNode>> SelectedNodes = SCSEditor->GetSelectedNodes();
 	if (SelectedNodes.Num() == 0)
 	{
-		SubobjectEditor->SelectRoot();
-		SelectedNodes = SubobjectEditor->GetSelectedNodes();
+		SCSEditor->SelectRoot();
+		SelectedNodes = SCSEditor->GetSelectedNodes();
 	}
 
 	// Update the underlying details view
-	RefreshSubobjectTreeElements(SelectedNodes, bSelectedComponentRecompiled, /*bOverrideLock*/false);
+	RefreshSCSTreeElements(SelectedNodes, bSelectedComponentRecompiled, /*bOverrideLock*/false);
 }
 
 bool SActorDetails::IsPropertyReadOnly(const FPropertyAndParent& PropertyAndParent) const
 {
 	bool bIsReadOnly = false;
-	for (const FSubobjectEditorTreeNodePtrType& Node : SubobjectEditor->GetSelectedNodes())
+	for (const FSCSEditorTreeNodePtrType& Node : SCSEditor->GetSelectedNodes())
 	{
-		const UActorComponent* Component = Node->GetComponentTemplate();
+		UActorComponent* Component = Node->GetComponentTemplate();
 		if (Component && Component->CreationMethod == EComponentCreationMethod::SimpleConstructionScript)
 		{
 			TSet<const FProperty*> UCSModifiedProperties;
@@ -662,15 +644,12 @@ bool SActorDetails::IsPropertyEditingEnabled() const
 	}
 
 	bool bIsEditable = true;
-	for (const FSubobjectEditorTreeNodePtrType& Node : SubobjectEditor->GetSelectedNodes())
+	for (const FSCSEditorTreeNodePtrType& Node : SCSEditor->GetSelectedNodes())
 	{
-		if(const FSubobjectData* Data = Node->GetDataSource())
+		bIsEditable = Node->CanEdit();
+		if (!bIsEditable)
 		{
-			bIsEditable = Data->CanEdit();
-			if (!bIsEditable)
-			{
-				break;
-			}
+			break;
 		}
 	}
 	return bIsEditable;
@@ -678,7 +657,7 @@ bool SActorDetails::IsPropertyEditingEnabled() const
 
 void SActorDetails::OnBlueprintedComponentWarningHyperlinkClicked(const FSlateHyperlinkRun::FMetadata& Metadata)
 {
-	UBlueprint* Blueprint = SubobjectEditor->GetBlueprint();
+	UBlueprint* Blueprint = SCSEditor->GetBlueprint();
 	if (Blueprint)
 	{
 		// Open the blueprint
@@ -689,7 +668,7 @@ void SActorDetails::OnBlueprintedComponentWarningHyperlinkClicked(const FSlateHy
 void SActorDetails::OnNativeComponentWarningHyperlinkClicked(const FSlateHyperlinkRun::FMetadata& Metadata)
 {
 	// Find the closest native parent
-	UBlueprint* Blueprint = SubobjectEditor->GetBlueprint();
+	UBlueprint* Blueprint = SCSEditor->GetBlueprint();
 	UClass* ParentClass = Blueprint ? *Blueprint->ParentClass : GetActorContext()->GetClass();
 	while (ParentClass && !ParentClass->HasAllClassFlags(CLASS_Native))
 	{
@@ -719,27 +698,23 @@ EVisibility SActorDetails::GetUCSComponentWarningVisibility() const
 	bool bIsUneditableBlueprintComponent = false;
 
 	// Check to see if any selected components are inherited from blueprint
-	for (const FSubobjectEditorTreeNodePtrType& Node : SubobjectEditor->GetSelectedNodes())
+	for (const FSCSEditorTreeNodePtrType& Node : SCSEditor->GetSelectedNodes())
 	{
-		if(const FSubobjectData* Data = Node->GetDataSource())
+		if (!Node->IsNativeComponent())
 		{
-			if (!Data->IsNativeComponent())
+			UActorComponent* Component = Node->GetComponentTemplate();
+			bIsUneditableBlueprintComponent = Component ? Component->CreationMethod == EComponentCreationMethod::UserConstructionScript : false;
+			if (bIsUneditableBlueprintComponent)
 			{
-				const UActorComponent* Component = Data->GetComponentTemplate();
-				bIsUneditableBlueprintComponent = Component ? Component->CreationMethod == EComponentCreationMethod::UserConstructionScript : false;
-				if (bIsUneditableBlueprintComponent)
-				{
-					break;
-				}
+				break;
 			}
 		}
-		
 	}
 
 	return bIsUneditableBlueprintComponent ? EVisibility::Visible : EVisibility::Collapsed;
 }
 
-bool NotEditableSetByBlueprint(const UActorComponent* Component)
+bool NotEditableSetByBlueprint(UActorComponent* Component)
 {
 	// Determine if it is locked out from a blueprint or from the native
 	UActorComponent* Archetype = CastChecked<UActorComponent>(Component->GetArchetype());
@@ -768,27 +743,24 @@ EVisibility SActorDetails::GetInheritedBlueprintComponentWarningVisibility() con
 	bool bIsUneditableBlueprintComponent = false;
 
 	// Check to see if any selected components are inherited from blueprint
-	for (const FSubobjectEditorTreeNodePtrType& Node : SubobjectEditor->GetSelectedNodes())
+	for (const FSCSEditorTreeNodePtrType& Node : SCSEditor->GetSelectedNodes())
 	{
-		if(const FSubobjectData* Data = Node->GetDataSource())
+		if (!Node->IsNativeComponent())
 		{
-			if (!Data->IsNativeComponent())
+			if (UActorComponent* Component = Node->GetComponentTemplate())
 			{
-				if (const UActorComponent* Component = Data->GetComponentTemplate())
+				if (!Component->IsEditableWhenInherited() && Component->CreationMethod == EComponentCreationMethod::SimpleConstructionScript)
 				{
-					if (!Component->IsEditableWhenInherited() && Component->CreationMethod == EComponentCreationMethod::SimpleConstructionScript)
-					{
-						bIsUneditableBlueprintComponent = true;
-						break;
-					}
+					bIsUneditableBlueprintComponent = true;
+					break;
 				}
 			}
-			else if (!Data->CanEdit() && NotEditableSetByBlueprint(Data->GetComponentTemplate()))
-			{
-				bIsUneditableBlueprintComponent = true;
-				break;
-			}
-		}	
+		}
+		else if (!Node->CanEdit() && NotEditableSetByBlueprint(Node->GetComponentTemplate()))
+		{
+			bIsUneditableBlueprintComponent = true;
+			break;
+		}
 	}
 
 	return bIsUneditableBlueprintComponent ? EVisibility::Visible : EVisibility::Collapsed;
@@ -797,18 +769,14 @@ EVisibility SActorDetails::GetInheritedBlueprintComponentWarningVisibility() con
 EVisibility SActorDetails::GetNativeComponentWarningVisibility() const
 {
 	bool bIsUneditableNative = false;
-	for (const FSubobjectEditorTreeNodePtrType& Node : SubobjectEditor->GetSelectedNodes())
+	for (const FSCSEditorTreeNodePtrType& Node : SCSEditor->GetSelectedNodes())
 	{
-		if(const FSubobjectData* Data = Node->GetDataSource())
+		// Check to see if the component is native and not editable
+		if (Node->IsNativeComponent() && !Node->CanEdit() && !NotEditableSetByBlueprint(Node->GetComponentTemplate()))
 		{
-			// Check to see if the component is native and not editable
-			if (Data->IsNativeComponent() && !Data->CanEdit() && !NotEditableSetByBlueprint(Data->GetComponentTemplate()))
-			{
-				bIsUneditableNative = true;
-				break;
-			}
+			bIsUneditableNative = true;
+			break;
 		}
-
 	}
 	
 	return bIsUneditableNative ? EVisibility::Visible : EVisibility::Collapsed;
diff --git a/Engine/Source/Editor/LevelEditor/Private/SActorDetails.h b/Engine/Source/Editor/LevelEditor/Private/SActorDetails.h
--- a/Engine/Source/Editor/LevelEditor/Private/SActorDetails.h
+++ b/Engine/Source/Editor/LevelEditor/Private/SActorDetails.h
@@ -17,13 +17,11 @@ class FTabManager;
 class FUICommandList;
 class IDetailsView;
 class SBox;
+class SSCSEditor;
 class SSplitter;
 class UBlueprint;
 class FDetailsViewObjectFilter;
 class UTypedElementSelectionSet;
-class SSubobjectEditor;
-class ISCSEditorUICustomization;
-class FSubobjectEditorTreeNode;
 
 /**
  * Wraps a details panel customized for viewing actors
@@ -67,25 +65,21 @@ public:
 	 */
 	void SetActorDetailsRootCustomization(TSharedPtr<FDetailsViewObjectFilter> ActorDetailsObjectFilter, TSharedPtr<class IDetailRootObjectCustomization> ActorDetailsRootCustomization);
 
-	UE_DEPRECATED(5.0, "SetSCSEditorUICustomization is deprecated, please use SetSubobjectEditorUICustomization instead.")
-	void SetSCSEditorUICustomization(TSharedPtr<ISCSEditorUICustomization> ActorDetailsSCSEditorUICustomization) { SetSubobjectEditorUICustomization(ActorDetailsSCSEditorUICustomization); }
-	
-	/** Sets the UI customization of the Suobject inside this details panel. */
-	void SetSubobjectEditorUICustomization(TSharedPtr<ISCSEditorUICustomization> ActorDetailsSCSEditorUICustomization);
+	/** Sets the UI customization of the SCSEditor inside this details panel. */
+	void SetSCSEditorUICustomization(TSharedPtr<class ISCSEditorUICustomization> ActorDetailsSCSEditorUICustomization);
 
 private:
 	void RefreshTopLevelElements(TArrayView<const TTypedElement<UTypedElementDetailsInterface>> InDetailsElements, const bool bForceRefresh, const bool bOverrideLock);
-	void RefreshSubobjectTreeElements(TArrayView<const TSharedPtr<FSubobjectEditorTreeNode>> InSelectedNodes, const bool bForceRefresh, const bool bOverrideLock);
+	void RefreshSCSTreeElements(TArrayView<const TSharedPtr<class FSCSEditorTreeNode>> InSelectedNodes, const bool bForceRefresh, const bool bOverrideLock);
 	void SetElementDetailsObjects(TArrayView<const TUniquePtr<ITypedElementDetailsObject>> InElementDetailsObjects, const bool bForceRefresh, const bool bOverrideLock);
 
 	AActor* GetActorContext() const;
-	UObject* GetObjectContext() const;
 	bool GetAllowComponentTreeEditing() const;
 
 	void OnComponentsEditedInWorld();
-	void OnSubobjectEditorTreeViewSelectionChanged(const TArray<TSharedPtr<FSubobjectEditorTreeNode> >& SelectedNodes);
-	void OnSubobjectEditorTreeViewItemDoubleClicked(const TSharedPtr<FSubobjectEditorTreeNode> ClickedNode);
-	void OnSubobjectEditorTreeViewObjectReplaced();
+	void OnSCSEditorTreeViewSelectionChanged(const TArray<TSharedPtr<class FSCSEditorTreeNode> >& SelectedNodes);
+	void OnSCSEditorTreeViewItemDoubleClicked(const TSharedPtr<class FSCSEditorTreeNode> ClickedNode);
+	void OnSCSEditorTreeViewObjectReplaced();
 	void UpdateComponentTreeFromEditorSelection();
 
 	bool IsPropertyReadOnly(const struct FPropertyAndParent& PropertyAndParent) const;
@@ -106,9 +100,7 @@ private:
 	TSharedPtr<SSplitter> DetailsSplitter;
 	TSharedPtr<class IDetailsView> DetailsView;
 	TSharedPtr<SBox> ComponentsBox;
-
-	/** The subobject editor provides a tree widget that allows for editing of subobjects */
-	TSharedPtr<SSubobjectEditor> SubobjectEditor;
+	TSharedPtr<class SSCSEditor> SCSEditor;
 
 	// The selection set this details panel is observing
 	UTypedElementSelectionSet* SelectionSet = nullptr;
@@ -120,9 +112,9 @@ private:
 	// Array of top-level elements that are currently being edited
 	TArray<TUniquePtr<ITypedElementDetailsObject>> TopLevelElements;
 
-	// Array of component elements that are being edited from the Subobject tree selection
-	TArray<TUniquePtr<ITypedElementDetailsObject>> SubobjectTreeElements;
-	
+	// Array of component elements that are being edited from the SCS tree selection
+	TArray<TUniquePtr<ITypedElementDetailsObject>> SCSTreeElements;
+
 	// The current component blueprint selection
 	TWeakObjectPtr<UBlueprint> SelectedBPComponentBlueprint;
 	bool bSelectedComponentRecompiled = false;
diff --git a/Engine/Source/Editor/LevelEditor/Private/SLevelEditor.cpp b/Engine/Source/Editor/LevelEditor/Private/SLevelEditor.cpp
--- a/Engine/Source/Editor/LevelEditor/Private/SLevelEditor.cpp
+++ b/Engine/Source/Editor/LevelEditor/Private/SLevelEditor.cpp
@@ -1969,7 +1969,7 @@ TSharedRef<SWidget> SLevelEditor::CreateActorDetails( const FName TabIdentifier
 	TSharedRef<SActorDetails> ActorDetails = SNew( SActorDetails, GetMutableElementSelectionSet(), TabIdentifier, LevelEditorCommands, GetTabManager() );
 
 	ActorDetails->SetActorDetailsRootCustomization(ActorDetailsObjectFilter, ActorDetailsRootCustomization);
-	ActorDetails->SetSubobjectEditorUICustomization(ActorDetailsSCSEditorUICustomization);
+	ActorDetails->SetSCSEditorUICustomization(ActorDetailsSCSEditorUICustomization);
 
 	AllActorDetailPanels.Add( ActorDetails );
 	return ActorDetails;
@@ -2018,7 +2018,7 @@ void SLevelEditor::SetActorDetailsSCSEditorUICustomization(TSharedPtr<ISCSEditor
 
 	for (TSharedRef<SActorDetails> ActorDetails : GetAllActorDetails())
 	{
-		ActorDetails->SetSubobjectEditorUICustomization(ActorDetailsSCSEditorUICustomization);
+		ActorDetails->SetSCSEditorUICustomization(ActorDetailsSCSEditorUICustomization);
 	}
 }
 
