Adds "Show Actor History" contextual action.

#changelist validated

diff --git a/Engine/Source/Editor/LevelEditor/Private/LevelEditor.cpp b/Engine/Source/Editor/LevelEditor/Private/LevelEditor.cpp
--- a/Engine/Source/Editor/LevelEditor/Private/LevelEditor.cpp
+++ b/Engine/Source/Editor/LevelEditor/Private/LevelEditor.cpp
@@ -735,6 +735,12 @@ void FLevelEditorModule::BindGlobalLevelEditorCommands()
 		FExecuteAction::CreateStatic(&FLevelEditorActionCallbacks::CopyActorFilePathtoClipboard_Clicked)
 		);
 
+	ActionList.MapAction(
+		Commands.ShowActorHistory,
+		FExecuteAction::CreateStatic(&FLevelEditorActionCallbacks::ShowActorHistory_Clicked),
+		FCanExecuteAction::CreateStatic(&FLevelEditorActionCallbacks::ShowActorHistory_CanExecute)
+	);
+
 	ActionList.MapAction(
 		Commands.GoToCodeForActor, 
 		FExecuteAction::CreateStatic( &FLevelEditorActionCallbacks::GoToCodeForActor_Clicked )
diff --git a/Engine/Source/Editor/LevelEditor/Private/LevelEditorActions.cpp b/Engine/Source/Editor/LevelEditor/Private/LevelEditorActions.cpp
--- a/Engine/Source/Editor/LevelEditor/Private/LevelEditorActions.cpp
+++ b/Engine/Source/Editor/LevelEditor/Private/LevelEditorActions.cpp
@@ -1,6 +1,7 @@
 // Copyright Epic Games, Inc. All Rights Reserved.
 
 #include "LevelEditorActions.h"
+#include "Algo/Unique.h"
 #include "SceneView.h"
 #include "Factories/Factory.h"
 #include "Animation/AnimSequence.h"
@@ -2649,6 +2650,64 @@ void FLevelEditorActionCallbacks::CopyActorFilePathtoClipboard_Clicked()
 	}
 }
 
+static TArray<FString> GetSelectedActorsPackageFullpath()
+{
+	TArray<FString> PackageFullpaths;
+
+	for (FSelectionIterator It(GEditor->GetSelectedActorIterator()); It; ++It)
+	{
+		const AActor* Actor = Cast<AActor>(*It);
+
+		if (Actor == nullptr)
+		{
+			continue;
+		}
+
+		const UPackage* Package = Actor->GetPackage();
+
+		if (Package == nullptr)
+		{
+			continue;
+		}
+
+		const FString LocalFullPath(Package->GetLoadedPath().GetLocalFullPath());
+
+		if (LocalFullPath.IsEmpty())
+		{
+			continue;
+		}
+
+		PackageFullpaths.Add(FPaths::ConvertRelativePathToFull(LocalFullPath));
+	}
+
+	return PackageFullpaths;
+}
+
+void FLevelEditorActionCallbacks::ShowActorHistory_Clicked()
+{
+	TArray<FString> PackageFullpaths = GetSelectedActorsPackageFullpath();
+
+	// Sort then remove consecutive identical elements to avoid displaying multiple times the same history.
+	PackageFullpaths.Sort();
+	PackageFullpaths.SetNum(Algo::Unique(PackageFullpaths));
+
+	FSourceControlWindows::DisplayRevisionHistory(PackageFullpaths);
+}
+
+bool FLevelEditorActionCallbacks::ShowActorHistory_CanExecute()
+{
+	ISourceControlProvider& SourceControlProvider = ISourceControlModule::Get().GetProvider();
+
+	if (!SourceControlProvider.IsEnabled())
+	{
+		return false;
+	}
+
+	TArray<FString> PackageFullpaths = GetSelectedActorsPackageFullpath();
+
+	return PackageFullpaths.Num() > 0;
+}
+
 void FLevelEditorActionCallbacks::OnEnableActorSnap()
 {
 	FSnappingUtils::EnableActorSnap( !FSnappingUtils::IsSnapToActorEnabled() );
@@ -3265,7 +3324,8 @@ void FLevelEditorCommands::RegisterCommands()
 	UI_COMMAND( SnapCameraToObject, "Snap View to Object", "Snaps the view to the selected object", EUserInterfaceActionType::Button, FInputChord() );
 	UI_COMMAND( SnapObjectToCamera, "Snap Object to View", "Snaps the selected object to the view", EUserInterfaceActionType::Button, FInputChord() );
 
-	UI_COMMAND( CopyActorFilePathtoClipboard, "Copy Actor File Path", "Copy the file path where this actor is saved", EUserInterfaceActionType::Button, FInputChord() );
+	UI_COMMAND(CopyActorFilePathtoClipboard, "Copy Actor File Path", "Copy the file path where this actor is saved", EUserInterfaceActionType::Button, FInputChord());
+	UI_COMMAND(ShowActorHistory, "Show Actor History", "Shows the history of the file containing the actor.", EUserInterfaceActionType::Button, FInputChord());
 
 	UI_COMMAND( GoToCodeForActor, "Go to C++ Code for Actor", "Opens a code editing IDE and navigates to the source file associated with the seleced actor", EUserInterfaceActionType::Button, FInputChord() );
 	UI_COMMAND( GoToDocsForActor, "Go to Documentation for Actor", "Opens documentation for the Actor in the default web browser", EUserInterfaceActionType::Button, FInputChord() );
diff --git a/Engine/Source/Editor/LevelEditor/Private/LevelEditorContextMenu.cpp b/Engine/Source/Editor/LevelEditor/Private/LevelEditorContextMenu.cpp
--- a/Engine/Source/Editor/LevelEditor/Private/LevelEditorContextMenu.cpp
+++ b/Engine/Source/Editor/LevelEditor/Private/LevelEditorContextMenu.cpp
@@ -136,6 +136,20 @@ public:
 	 */
 	static void FillMergeActorsMenu(UToolMenu* Menu);
 
+	/**
+	 * Adds the Source Control SubMenu to the provided Section.
+	 *
+	 * @param Section	The menu to add items to
+	 */
+	static void AddSourceControlMenu(FToolMenuSection& Section);
+
+	/**
+	 * Fills in menu options for the source control menu
+	 *
+	 * @param MenuBuilder	The menu to add items to
+	 */
+	static void FillSourceControlMenu(UToolMenu* Menu);
+
 private:
 	/**
 	 * Fills in menu options for the matinee selection menu
@@ -204,6 +218,8 @@ void FLevelEditorContextMenu::RegisterComponentContextMenu()
 			Section.AddMenuEntry(FLevelEditorCommands::Get().SnapCameraToObject);
 			Section.AddMenuEntry(FLevelEditorCommands::Get().SnapObjectToCamera);
 			Section.AddMenuEntry(FLevelEditorCommands::Get().CopyActorFilePathtoClipboard);
+
+			FLevelEditorContextMenuImpl::AddSourceControlMenu(Section);
 		}
 
 		FComponentEditorUtils::FillComponentContextMenuOptions(InMenu, SelectedComponents);
@@ -290,6 +306,8 @@ void FLevelEditorContextMenu::RegisterActorContextMenu()
 			Section.AddMenuEntry(FLevelEditorCommands::Get().SnapCameraToObject);
 			Section.AddMenuEntry(FLevelEditorCommands::Get().SnapObjectToCamera);
 			Section.AddMenuEntry(FLevelEditorCommands::Get().CopyActorFilePathtoClipboard);
+			
+			FLevelEditorContextMenuImpl::AddSourceControlMenu(Section);
 
 			if (SelectedActors.Num() == 1)
 			{
@@ -1406,6 +1424,20 @@ void FLevelEditorContextMenuImpl::FillMergeActorsMenu(UToolMenu* Menu)
 		FSlateIcon(FEditorStyle::GetStyleSetName(), "LevelEditor.GameSettings"));
 }
 
+void FLevelEditorContextMenuImpl::AddSourceControlMenu(FToolMenuSection& Section)
+{
+	Section.AddSubMenu(TEXT("SourceControlSubMenu"), LOCTEXT("SourceControlSubMenu", "Source Control"),
+					   LOCTEXT("SourceControlSubMenu_ToolTip", "Opens the Source Control sub menu"),
+					   FNewToolMenuDelegate::CreateStatic(&FLevelEditorContextMenuImpl::FillSourceControlMenu));
+}
+
+void FLevelEditorContextMenuImpl::FillSourceControlMenu(UToolMenu* Menu)
+{
+	FToolMenuSection& Section = Menu->AddSection(TEXT("Source Control"));
+
+	Section.AddMenuEntry(FLevelEditorCommands::Get().ShowActorHistory);
+}
+
 void FLevelScriptEventMenuHelper::FillLevelBlueprintEventsMenu(UToolMenu* Menu, const TArray<AActor*>& SelectedActors)
 {
 	AActor* SelectedActor = (1 == SelectedActors.Num()) ? SelectedActors[0] : NULL;
diff --git a/Engine/Source/Editor/LevelEditor/Public/LevelEditorActions.h b/Engine/Source/Editor/LevelEditor/Public/LevelEditorActions.h
--- a/Engine/Source/Editor/LevelEditor/Public/LevelEditorActions.h
+++ b/Engine/Source/Editor/LevelEditor/Public/LevelEditorActions.h
@@ -134,6 +134,12 @@ public:
 	/** Copy the file path where the actor is saved. */
 	TSharedPtr< FUICommandInfo > CopyActorFilePathtoClipboard;
 
+	/** Opens the source control panel. */
+	TSharedPtr< FUICommandInfo > OpenSourceControl;
+
+	/** Shows the history of the file containing the actor. */
+	TSharedPtr< FUICommandInfo > ShowActorHistory;
+
 	/** Goes to the source code for the selected actor's class. */
 	TSharedPtr< FUICommandInfo > GoToCodeForActor;
 
@@ -1263,6 +1269,17 @@ public:
 	*/
 	static void CopyActorFilePathtoClipboard_Clicked();
 
+	/**
+	* Shows the history of the file containing the actor.
+	*/
+	static void ShowActorHistory_Clicked();
+
+	/**
+	 * Checks whether ShowActorHistory can be executed on the selected actors.
+	 * @return 	True if ShowActorHistory can be executer on the selected actors.
+	 */
+	static bool ShowActorHistory_CanExecute();
+
 	/** 
 	 * Moves the selected elements to the last selected element.
 	 */
