PR #7885: Use ShouldShowMouseCursor() instead of direct variable access (Contributed by [REDACTED_PERSON])

diff --git a/Engine/Source/Editor/LevelEditor/Private/SLevelViewport.cpp b/Engine/Source/Editor/LevelEditor/Private/SLevelViewport.cpp
--- a/Engine/Source/Editor/LevelEditor/Private/SLevelViewport.cpp
+++ b/Engine/Source/Editor/LevelEditor/Private/SLevelViewport.cpp
@@ -4200,7 +4200,7 @@ EVisibility SLevelViewport::GetMouseCaptureLabelVisibility() const
 	{
 		// Show the label if the local player's PC isn't set to show the cursor
 		auto const TargetPlayer = GEngine->GetLocalPlayerFromControllerId(GEditor->PlayWorld, 0);
-		if (TargetPlayer && TargetPlayer->PlayerController && !TargetPlayer->PlayerController->bShowMouseCursor)
+		if (TargetPlayer && TargetPlayer->PlayerController && !TargetPlayer->PlayerController->ShouldShowMouseCursor())
 		{
 			return EVisibility::HitTestInvisible;
 		}
