Minor update to how the default modes icon is pulled in

diff --git a/Engine/Source/Editor/EditorFramework/Private/Tools/Modes.cpp b/Engine/Source/Editor/EditorFramework/Private/Tools/Modes.cpp
--- a/Engine/Source/Editor/EditorFramework/Private/Tools/Modes.cpp
+++ b/Engine/Source/Editor/EditorFramework/Private/Tools/Modes.cpp
@@ -24,10 +24,6 @@ FEditorModeInfo::FEditorModeInfo(
 	, PriorityOrder(InPriorityOrder)
 	, Visibility(InVisibility)
 {
-	if (!InIconBrush.IsSet())
-	{
-		IconBrush = FSlateIcon(FEditorStyle::GetStyleSetName(), "LevelEditor.EditorModes");
-	}
 }
 
 bool FEditorModeInfo::IsVisible() const
diff --git a/Engine/Source/Editor/LevelEditor/Private/LevelEditorModesActions.cpp b/Engine/Source/Editor/LevelEditor/Private/LevelEditorModesActions.cpp
--- a/Engine/Source/Editor/LevelEditor/Private/LevelEditorModesActions.cpp
+++ b/Engine/Source/Editor/LevelEditor/Private/LevelEditorModesActions.cpp
@@ -43,13 +43,14 @@ void FLevelEditorModesCommands::RegisterCommands()
 				++editorMode;
 			}
 
+			FSlateIcon ModeIcon = (Mode.IconBrush != FSlateIcon()) ? Mode.IconBrush : FSlateIcon(FAppStyle::Get().GetStyleSetName(), "LevelEditor.Tabs.EditorModes");
 			FUICommandInfo::MakeCommandInfo(
 				this->AsShared(),
 				EditorModeCommand,
 				EditorModeCommandName,
 				ModeName,
 				Tooltip,
-				Mode.IconBrush,
+				ModeIcon,
 				EUserInterfaceActionType::ToggleButton,
 				DefaultKeyBinding);
 
