Removed redundant element selection helpers

#rnx

diff --git a/Engine/Source/Editor/LevelEditor/Private/Elements/Actor/ActorElementLevelEditorCommonActionsCustomization.cpp b/Engine/Source/Editor/LevelEditor/Private/Elements/Actor/ActorElementLevelEditorCommonActionsCustomization.cpp
--- a/Engine/Source/Editor/LevelEditor/Private/Elements/Actor/ActorElementLevelEditorCommonActionsCustomization.cpp
+++ b/Engine/Source/Editor/LevelEditor/Private/Elements/Actor/ActorElementLevelEditorCommonActionsCustomization.cpp
@@ -5,7 +5,6 @@
 #include "GameFramework/Actor.h"
 
 #include "Elements/Framework/TypedElementSelectionSet.h"
-#include "Elements/Component/ComponentElementSelectionInterface.h"
 
 #include "EditorModeManager.h"
 #include "Toolkits/IToolkitHost.h"
@@ -14,7 +13,7 @@ void FActorElementLevelEditorCommonActionsCustomization::GetElementsForAction(co
 {
 	AActor* Actor = ActorElementDataUtil::GetActorFromHandleChecked(InElementWorldHandle);
 
-	if (UComponentElementSelectionInterface::HasSelectedComponents(InElementList))
+	if (InElementList->HasElementsOfType(NAME_Components))
 	{
 		// If we have components selected then we will perform the action on those rather than the actors
 		return;
diff --git a/Engine/Source/Editor/LevelEditor/Private/Elements/Actor/ActorElementLevelEditorSelectionCustomization.cpp b/Engine/Source/Editor/LevelEditor/Private/Elements/Actor/ActorElementLevelEditorSelectionCustomization.cpp
--- a/Engine/Source/Editor/LevelEditor/Private/Elements/Actor/ActorElementLevelEditorSelectionCustomization.cpp
+++ b/Engine/Source/Editor/LevelEditor/Private/Elements/Actor/ActorElementLevelEditorSelectionCustomization.cpp
@@ -8,9 +8,6 @@
 #include "Elements/Framework/TypedElementList.h"
 #include "Elements/Framework/EngineElementsLibrary.h"
 
-#include "Elements/Actor/ActorElementSelectionInterface.h"
-#include "Elements/Component/ComponentElementSelectionInterface.h"
-
 #include "Editor.h"
 #include "LevelUtils.h"
 #include "UnrealEdGlobals.h"
@@ -46,7 +43,7 @@ bool FActorElementLevelEditorSelectionCustomization::DeselectElement(const TType
 bool FActorElementLevelEditorSelectionCustomization::AllowSelectionModifiers(const TTypedElement<UTypedElementSelectionInterface>& InElementSelectionHandle, const UTypedElementList* InSelectionSet)
 {
 	// Ctrl or Shift clicking an actor is the same as regular clicking when components are selected
-	return !UComponentElementSelectionInterface::HasSelectedComponents(InSelectionSet);
+	return !InSelectionSet->HasElementsOfType(NAME_Components);
 }
 
 FTypedElementHandle FActorElementLevelEditorSelectionCustomization::GetSelectionElement(const TTypedElement<UTypedElementSelectionInterface>& InElementSelectionHandle, const UTypedElementList* InCurrentSelection, const ETypedElementSelectionMethod InSelectionMethod)
diff --git a/Engine/Source/Editor/LevelEditor/Private/Elements/Actor/ActorElementLevelEditorViewportInteractionCustomization.cpp b/Engine/Source/Editor/LevelEditor/Private/Elements/Actor/ActorElementLevelEditorViewportInteractionCustomization.cpp
--- a/Engine/Source/Editor/LevelEditor/Private/Elements/Actor/ActorElementLevelEditorViewportInteractionCustomization.cpp
+++ b/Engine/Source/Editor/LevelEditor/Private/Elements/Actor/ActorElementLevelEditorViewportInteractionCustomization.cpp
@@ -6,7 +6,6 @@
 
 #include "Elements/Framework/EngineElementsLibrary.h"
 #include "Elements/Framework/TypedElementSelectionSet.h"
-#include "Elements/Component/ComponentElementSelectionInterface.h"
 #include "Elements/Component/ComponentElementLevelEditorViewportInteractionCustomization.h"
 
 #include "Editor.h"
@@ -19,7 +18,7 @@ void FActorElementLevelEditorViewportInteractionCustomization::GetElementsToMove
 {
 	AActor* Actor = ActorElementDataUtil::GetActorFromHandleChecked(InElementWorldHandle);
 
-	if (UComponentElementSelectionInterface::HasSelectedComponents(InSelectionSet->GetElementList()))
+	if (InSelectionSet->GetElementList()->HasElementsOfType(NAME_Components))
 	{
 		// If we have components selected then we will move those rather than the actors
 		// The component may still choose to move its owner actor rather than itself
diff --git a/Engine/Source/Editor/LevelEditor/Private/Elements/Component/ComponentElementLevelEditorSelectionCustomization.cpp b/Engine/Source/Editor/LevelEditor/Private/Elements/Component/ComponentElementLevelEditorSelectionCustomization.cpp
--- a/Engine/Source/Editor/LevelEditor/Private/Elements/Component/ComponentElementLevelEditorSelectionCustomization.cpp
+++ b/Engine/Source/Editor/LevelEditor/Private/Elements/Component/ComponentElementLevelEditorSelectionCustomization.cpp
@@ -8,9 +8,6 @@
 #include "Elements/Framework/TypedElementList.h"
 #include "Elements/Framework/EngineElementsLibrary.h"
 
-#include "Elements/Actor/ActorElementSelectionInterface.h"
-#include "Elements/Component/ComponentElementSelectionInterface.h"
-
 #include "EditorViewportClient.h"
 #include "Kismet2/ComponentEditorUtils.h"
 #include "Kismet2/KismetEditorUtilities.h"
@@ -64,9 +61,9 @@ FTypedElementHandle FComponentElementLevelEditorSelectionCustomization::GetSelec
 			// 2. The actor selected is blueprintable
 			// 3. No components are already selected and the click was a double click
 			// 4. OR, a component is already selected and the click was NOT a double click
-			const bool bActorAlreadySelectedExclusively = InCurrentSelection->Contains(ConsideredActorHandle) && UActorElementSelectionInterface::GetNumSelectedActors(InCurrentSelection) == 1;
+			const bool bActorAlreadySelectedExclusively = InCurrentSelection->Contains(ConsideredActorHandle) && InCurrentSelection->CountElementsOfType(NAME_Actor) == 1;
 			const bool bActorIsBlueprintable = FKismetEditorUtilities::CanCreateBlueprintOfClass(ConsideredActor->GetClass());
-			const bool bComponentAlreadySelected = UComponentElementSelectionInterface::HasSelectedComponents(InCurrentSelection);
+			const bool bComponentAlreadySelected = InCurrentSelection->HasElementsOfType(NAME_Components);
 			const bool bWasDoubleClick = InSelectionMethod == ETypedElementSelectionMethod::Secondary;
 
 			const bool bSelectComponent = bActorAlreadySelectedExclusively && bActorIsBlueprintable && (bComponentAlreadySelected != bWasDoubleClick);
diff --git a/Engine/Source/Runtime/Engine/Private/Elements/Actor/ActorElementSelectionInterface.cpp b/Engine/Source/Runtime/Engine/Private/Elements/Actor/ActorElementSelectionInterface.cpp
--- a/Engine/Source/Runtime/Engine/Private/Elements/Actor/ActorElementSelectionInterface.cpp
+++ b/Engine/Source/Runtime/Engine/Private/Elements/Actor/ActorElementSelectionInterface.cpp
@@ -4,13 +4,3 @@
 #include "Elements/Actor/ActorElementData.h"
 #include "GameFramework/Actor.h"
 #include "Elements/Framework/TypedElementList.h"
-
-int32 UActorElementSelectionInterface::GetNumSelectedActors(const UTypedElementList* InCurrentSelection)
-{
-	return InCurrentSelection->CountElementsOfType(NAME_Actor);
-}
-
-bool UActorElementSelectionInterface::HasSelectedActors(const UTypedElementList* InCurrentSelection)
-{
-	return InCurrentSelection->HasElementsOfType(NAME_Actor);
-}
diff --git a/Engine/Source/Runtime/Engine/Private/Elements/Component/ComponentElementSelectionInterface.cpp b/Engine/Source/Runtime/Engine/Private/Elements/Component/ComponentElementSelectionInterface.cpp
--- a/Engine/Source/Runtime/Engine/Private/Elements/Component/ComponentElementSelectionInterface.cpp
+++ b/Engine/Source/Runtime/Engine/Private/Elements/Component/ComponentElementSelectionInterface.cpp
@@ -4,13 +4,3 @@
 #include "Elements/Component/ComponentElementData.h"
 #include "Components/ActorComponent.h"
 #include "Elements/Framework/TypedElementList.h"
-
-int32 UComponentElementSelectionInterface::GetNumSelectedComponents(const UTypedElementList* InCurrentSelection)
-{
-	return InCurrentSelection->CountElementsOfType(NAME_Components);
-}
-
-bool UComponentElementSelectionInterface::HasSelectedComponents(const UTypedElementList* InCurrentSelection)
-{
-	return InCurrentSelection->HasElementsOfType(NAME_Components);
-}
diff --git a/Engine/Source/Runtime/Engine/Private/Elements/SMInstance/SMInstanceElementSelectionInterface.cpp b/Engine/Source/Runtime/Engine/Private/Elements/SMInstance/SMInstanceElementSelectionInterface.cpp
--- a/Engine/Source/Runtime/Engine/Private/Elements/SMInstance/SMInstanceElementSelectionInterface.cpp
+++ b/Engine/Source/Runtime/Engine/Private/Elements/SMInstance/SMInstanceElementSelectionInterface.cpp
@@ -4,13 +4,3 @@
 #include "Elements/SMInstance/SMInstanceElementData.h"
 #include "Components/InstancedStaticMeshComponent.h"
 #include "Elements/Framework/TypedElementList.h"
-
-int32 USMInstanceElementSelectionInterface::GetNumSelectedSMInstances(const UTypedElementList* InCurrentSelection)
-{
-	return InCurrentSelection->CountElementsOfType(NAME_SMInstance);
-}
-
-bool USMInstanceElementSelectionInterface::HasSelectedSMInstances(const UTypedElementList* InCurrentSelection)
-{
-	return InCurrentSelection->HasElementsOfType(NAME_SMInstance);
-}
diff --git a/Engine/Source/Runtime/Engine/Public/Elements/Actor/ActorElementSelectionInterface.h b/Engine/Source/Runtime/Engine/Public/Elements/Actor/ActorElementSelectionInterface.h
--- a/Engine/Source/Runtime/Engine/Public/Elements/Actor/ActorElementSelectionInterface.h
+++ b/Engine/Source/Runtime/Engine/Public/Elements/Actor/ActorElementSelectionInterface.h
@@ -9,8 +9,4 @@ UCLASS()
 class ENGINE_API UActorElementSelectionInterface : public UTypedElementSelectionInterface
 {
 	GENERATED_BODY()
-
-public:
-	static int32 GetNumSelectedActors(const UTypedElementList* InCurrentSelection);
-	static bool HasSelectedActors(const UTypedElementList* InCurrentSelection);
 };
diff --git a/Engine/Source/Runtime/Engine/Public/Elements/Component/ComponentElementSelectionInterface.h b/Engine/Source/Runtime/Engine/Public/Elements/Component/ComponentElementSelectionInterface.h
--- a/Engine/Source/Runtime/Engine/Public/Elements/Component/ComponentElementSelectionInterface.h
+++ b/Engine/Source/Runtime/Engine/Public/Elements/Component/ComponentElementSelectionInterface.h
@@ -9,8 +9,4 @@ UCLASS()
 class ENGINE_API UComponentElementSelectionInterface : public UTypedElementSelectionInterface
 {
 	GENERATED_BODY()
-
-public:
-	static int32 GetNumSelectedComponents(const UTypedElementList* InCurrentSelection);
-	static bool HasSelectedComponents(const UTypedElementList* InCurrentSelection);
 };
diff --git a/Engine/Source/Runtime/Engine/Public/Elements/SMInstance/SMInstanceElementSelectionInterface.h b/Engine/Source/Runtime/Engine/Public/Elements/SMInstance/SMInstanceElementSelectionInterface.h
--- a/Engine/Source/Runtime/Engine/Public/Elements/SMInstance/SMInstanceElementSelectionInterface.h
+++ b/Engine/Source/Runtime/Engine/Public/Elements/SMInstance/SMInstanceElementSelectionInterface.h
@@ -9,8 +9,4 @@ UCLASS()
 class ENGINE_API USMInstanceElementSelectionInterface : public UTypedElementSelectionInterface
 {
 	GENERATED_BODY()
-
-public:
-	static int32 GetNumSelectedSMInstances(const UTypedElementList* InCurrentSelection);
-	static bool HasSelectedSMInstances(const UTypedElementList* InCurrentSelection);
 };
