Fixed warning when editing non-actor element types

The sub-object instance editor would warn that the given context wasn't an actor, so this limits it back to only actor contexts (like how the SCS editor used to work)

diff --git a/Engine/Source/Editor/LevelEditor/Private/SActorDetails.cpp b/Engine/Source/Editor/LevelEditor/Private/SActorDetails.cpp
--- a/Engine/Source/Editor/LevelEditor/Private/SActorDetails.cpp
+++ b/Engine/Source/Editor/LevelEditor/Private/SActorDetails.cpp
@@ -142,7 +142,7 @@ void SActorDetails::Construct(const FArguments& InArgs, UTypedElementSelectionSe
 	FModuleManager::LoadModuleChecked<FSubobjectEditorModule>("SubobjectEditor");
 	
 	SubobjectEditor = SNew(SSubobjectInstanceEditor)
-		.ObjectContext(this, &SActorDetails::GetObjectContext)
+		.ObjectContext(this, &SActorDetails::GetActorContextAsObject)
 		.AllowEditing(this, &SActorDetails::GetAllowComponentTreeEditing)
 		.OnSelectionUpdated(this, &SActorDetails::OnSubobjectEditorTreeViewSelectionChanged)
 		.OnItemDoubleClicked(this, &SActorDetails::OnSubobjectEditorTreeViewItemDoubleClicked)
@@ -479,14 +479,14 @@ bool SActorDetails::GetAllowComponentTreeEditing() const
 
 AActor* SActorDetails::GetActorContext() const
 {
-	return Cast<AActor>(GetObjectContext());
+	return TopLevelElements.Num() == 1
+		? Cast<AActor>(TopLevelElements[0]->GetObject())
+		: nullptr;
 }
 
-UObject* SActorDetails::GetObjectContext() const
+UObject* SActorDetails::GetActorContextAsObject() const
 {
-	return TopLevelElements.Num() == 1
-		? TopLevelElements[0]->GetObject()
-		: nullptr;
+	return GetActorContext();
 }
 
 void SActorDetails::OnSubobjectEditorTreeViewSelectionChanged(const TArray<FSubobjectEditorTreeNodePtrType>& SelectedNodes)
diff --git a/Engine/Source/Editor/LevelEditor/Private/SActorDetails.h b/Engine/Source/Editor/LevelEditor/Private/SActorDetails.h
--- a/Engine/Source/Editor/LevelEditor/Private/SActorDetails.h
+++ b/Engine/Source/Editor/LevelEditor/Private/SActorDetails.h
@@ -79,7 +79,7 @@ private:
 	void SetElementDetailsObjects(TArrayView<const TUniquePtr<ITypedElementDetailsObject>> InElementDetailsObjects, const bool bForceRefresh, const bool bOverrideLock);
 
 	AActor* GetActorContext() const;
-	UObject* GetObjectContext() const;
+	UObject* GetActorContextAsObject() const;
 	bool GetAllowComponentTreeEditing() const;
 
 	void OnComponentsEditedInWorld();
