Updates to the in-viewport menu styling, and turning the UI on by default. Adds SWindow::MakeStyledCursorDecorator for custom styling.

diff --git a/Engine/Source/Editor/EditorStyle/Private/StarshipStyle.cpp b/Engine/Source/Editor/EditorStyle/Private/StarshipStyle.cpp
--- a/Engine/Source/Editor/EditorStyle/Private/StarshipStyle.cpp
+++ b/Engine/Source/Editor/EditorStyle/Private/StarshipStyle.cpp
@@ -3046,22 +3046,45 @@ void FStarshipEditorStyle::FStyle::SetupPropertyEditorStyles()
 			.SetForegroundColor( DefaultForeground )
 			);
 
-	
-		
-		FLinearColor TransparentHeader = FStyleColors::Recessed.GetSpecifiedColor();
-		TransparentHeader.A = 0.95f;
+		FWindowStyle InViewportDecoratorWindow = FStarshipCoreStyle::GetCoreStyle().GetWidgetStyle<FWindowStyle>("Window");
+		InViewportDecoratorWindow.SetCornerRadius(20);
+
+		Set("InViewportDecoratorWindow", InViewportDecoratorWindow);
+		FLinearColor TransparentHeader = FStyleColors::Panel.GetSpecifiedColor();
+		TransparentHeader.A = 0.5f;
 		FLinearColor TransparentToolbar = FStyleColors::Header.GetSpecifiedColor();
-		TransparentToolbar.A = 0.95f;
+		TransparentToolbar.A = 0.5f;
 		FLinearColor TransparentBackground = FStyleColors::Background.GetSpecifiedColor();
-		TransparentBackground.A = 0.75f;
-		Set("PropertyTable.InViewport.Background", new FSlateColorBrush(FSlateColor(TransparentBackground)));
-		Set("PropertyTable.InViewport.Header", new FSlateColorBrush(FSlateColor(TransparentHeader)));
+		TransparentBackground.A = 0.35f;
+		Set("PropertyTable.InViewport.Header", new FSlateRoundedBoxBrush(FSlateColor(TransparentHeader), FVector4(10.0f, 10.0f, 0.0f, 0.0f)));
+		Set("PropertyTable.InViewport.Background", new FSlateRoundedBoxBrush(FSlateColor(TransparentBackground), 10.0f));
 		// InViewportToolbar
 		{
 			FToolBarStyle InViewportToolbar = FStarshipCoreStyle::GetCoreStyle().GetWidgetStyle<FToolBarStyle>("SlimToolBar");
 			InViewportToolbar.SetBackground(FSlateColorBrush(FSlateColor(TransparentToolbar)));
+			InViewportToolbar.SetBackgroundPadding(0.0f);
 			Set("InViewportToolbar", InViewportToolbar);
 		}
+		const FTableViewStyle InViewportViewStyle = FTableViewStyle()
+			.SetBackgroundBrush(FSlateNoResource());
+		Set("PropertyTable.InViewport.ListView", InViewportViewStyle);
+
+		Set("PropertyTable.InViewport.Row", FTableRowStyle(NormalTableRowStyle)
+			.SetEvenRowBackgroundBrush(FSlateNoResource())
+			.SetEvenRowBackgroundHoveredBrush(FSlateNoResource())
+			.SetOddRowBackgroundBrush(FSlateNoResource())
+			.SetOddRowBackgroundHoveredBrush(FSlateNoResource())
+			.SetSelectorFocusedBrush(FSlateNoResource())
+			.SetActiveBrush(FSlateNoResource())
+			.SetActiveHoveredBrush(FSlateNoResource())
+			.SetInactiveBrush(FSlateNoResource())
+			.SetInactiveHoveredBrush(FSlateNoResource())
+		);
+
+		const FSplitterStyle TransparentSplitterStyle = FSplitterStyle()
+			.SetHandleNormalBrush(FSlateNoResource())
+			.SetHandleHighlightBrush(FSlateNoResource());
+		Set("PropertyTable.InViewport.Splitter", TransparentSplitterStyle);
 
 		Set( "PropertyTable.Selection.Active",						new IMAGE_BRUSH( "Common/Selection", Icon8x8, SelectionColor ) );
 
@@ -4715,8 +4738,7 @@ void FStarshipEditorStyle::FStyle::SetupLevelEditorStyle()
 
 		Set( "LevelViewport.EjectActorPilot", new IMAGE_BRUSH( "Icons/icon_EjectActorPilot_16x", Icon16x16 ) );
 		Set( "LevelViewport.EjectActorPilot.Small", new IMAGE_BRUSH( "Icons/icon_EjectActorPilot_16x", Icon16x16 ) );
-		Set( "LevelViewport.PilotSelectedActor", new IMAGE_BRUSH( "Icons/icon_PilotSelectedActor_16x", Icon16x16 ) );
-		Set( "LevelViewport.PilotSelectedActor.Small", new IMAGE_BRUSH( "Icons/icon_PilotSelectedActor_16x", Icon16x16 ) );
+		Set( "LevelViewport.PilotSelectedActor", new IMAGE_BRUSH_SVG( "Starship/EditorViewport/pilot", Icon16x16 ) );
 		
 		Set( "LevelViewport.ToggleActorPilotCameraView",       new IMAGE_BRUSH_SVG( "Starship/Common/CameraShake", Icon16x16 ) );
 		Set( "LevelViewport.ToggleActorPilotCameraView.Small", new IMAGE_BRUSH_SVG( "Starship/Common/CameraShake", Icon16x16 ) );
diff --git a/Engine/Source/Editor/LevelEditor/Private/SInViewportDetails.cpp b/Engine/Source/Editor/LevelEditor/Private/SInViewportDetails.cpp
--- a/Engine/Source/Editor/LevelEditor/Private/SInViewportDetails.cpp
+++ b/Engine/Source/Editor/LevelEditor/Private/SInViewportDetails.cpp
@@ -40,6 +40,9 @@
 #include "ToolMenus.h"
 #include "Editor/EditorEngine.h"
 #include "LevelEditorMenuContext.h"
+#include "Styling/SlateIconFinder.h"
+#include "UObject/Class.h"
+#include "UObject/WeakObjectPtrTemplates.h"
 
 #define LOCTEXT_NAMESPACE "InViewportDetails"
 void FInViewportUIDragOperation::OnDrop(bool bDropWasHandled, const FPointerEvent& MouseEvent)
@@ -86,18 +89,11 @@ FInViewportUIDragOperation::FInViewportUIDragOperation(const TSharedRef<class SI
 
 	// Start the window off hidden.
 	const bool bShowImmediately = true;
-	CursorDecoratorWindow = FSlateApplication::Get().AddWindow(SWindow::MakeCursorDecorator(), bShowImmediately);
-	// Usually cursor decorators figure out their size automatically from content, but we will drive it
-	// here because the window will reshape itself to better reflect what will happen when the user drops the Tab.
-	CursorDecoratorWindow->SetSizingRule(ESizingRule::Autosized);
+	CursorDecoratorWindow = FSlateApplication::Get().AddWindow(SWindow::MakeStyledCursorDecorator(FAppStyle::Get().GetWidgetStyle<FWindowStyle>("InViewportDecoratorWindow")), bShowImmediately);
 	CursorDecoratorWindow->SetOpacity(0.45f);
 	CursorDecoratorWindow->SetContent
 	(
-		SNew(SBorder)
-		.BorderImage(FCoreStyle::Get().GetBrush("Docking.Background"))
-		[
-			InUIToBeDragged
-		]
+		InUIToBeDragged
 	);
 
 }
@@ -141,7 +137,7 @@ public:
 		{
  
 			RowWidget = SNew(SSplitter)
- 				.Style(FEditorStyle::Get(), "DetailsView.Splitter")
+				.Style(FEditorStyle::Get(), "PropertyTable.InViewport.Splitter")
 				.PhysicalSplitterHandleSize(1.0f)
  				.HitDetectionSplitterHandleSize(5.0f)
  				+ SSplitter::Slot()
@@ -151,7 +147,7 @@ public:
  				[
 					SNew(SBox)
 					.HAlign(HAlign_Right)
- 					.Padding(2.0f)
+ 					.Padding(5.0f)
  					[
 						NameWidget.ToSharedRef()
 					]
@@ -162,7 +158,7 @@ public:
  				.OnSlotResized(ColumnSizeData.OnValueColumnResized)
  				[
  					SNew(SBox)
- 					.Padding(2.0f)
+ 					.Padding(5.0f)
  					[
  						ValueWidget.ToSharedRef()
  					]
@@ -174,7 +170,7 @@ public:
 						+ SHorizontalBox::Slot()
 						.VAlign(VAlign_Center)
 						.HAlign(HAlign_Center)
-						.Padding(0.0f)
+						.Padding(2.0f)
 						[
 							ResetWidget.ToSharedRef()
 						]
@@ -184,7 +180,7 @@ public:
 		ChildSlot
 			[
 				SNew(SBox)
-				.MinDesiredWidth(300.0f)
+				.MinDesiredWidth(350.0f)
 				[
 					RowWidget.ToSharedRef()
 				]
@@ -192,7 +188,7 @@ public:
 
 		STableRow< TSharedPtr<IDetailTreeNode> >::ConstructInternal(
 			STableRow::FArguments()
-			.Style(FAppStyle::Get(), "DetailsView.TreeView.TableRow")
+			.Style(FAppStyle::Get(), "PropertyTable.InViewport.Row")
 			.ShowSelection(false),
 			InOwnerTableView
 		);
@@ -234,42 +230,81 @@ void SInViewportDetails::GenerateWidget()
 		{
 			NameString = SelectedActor->GetHumanReadableName();
 		}
-		
+
+		// Get the common base class of the selected objects
+		UClass* BaseClass = NULL;
+		for (FSelectionIterator It(GEditor->GetSelectedActorIterator()); It; ++It)
+		{
+			AActor* Actor = Cast<AActor>(*It);
+			if (Actor)
+			{
+				UClass* ActorClass = Actor->GetClass();
+
+				if (!BaseClass)
+				{
+					BaseClass = ActorClass;
+				}
+
+				while (!ActorClass->IsChildOf(BaseClass))
+				{
+					BaseClass = BaseClass->GetSuperClass();
+				}
+			}
+		}
+		const FSlateBrush* ActorIcon = FSlateIconFinder::FindIconBrushForClass(BaseClass);
 		ChildSlot
 			[
 				SNew(SBackgroundBlur)
-				.Visibility(this, &SInViewportDetails::GetHeaderVisibility)
-				.BlurStrength(1)
-				.BlurRadius(10)
+				.BlurStrength(20)
+				.Padding(0.0f)
+				.CornerRadius(FVector4(15.0f, 15.0f, 15.0f, 15.0f))
 				[
-					SNew(SVerticalBox)
-					+ SVerticalBox::Slot()
-					.AutoHeight()
+					SNew(SBorder)
+					.BorderImage(FAppStyle::Get().GetBrush("PropertyTable.InViewport.Background"))
+					.Visibility(this, &SInViewportDetails::GetHeaderVisibility)
+					.Padding(0.0f)
 					[
-						SNew(SInViewportDetailsHeader)
-						.Parent(SharedThis(this))
-						.Content()
-						[					
-							SNew(SBorder)
-							.BorderImage(FAppStyle::Get().GetBrush("PropertyTable.InViewport.Header"))
-							.Padding(5.0f)
-							[
-								SNew(STextBlock)
-								.Text(FText::FromString(NameString))
-								.TextStyle(FAppStyle::Get(), "DetailsView.CategoryTextStyle")
+						SNew(SVerticalBox)
+						+ SVerticalBox::Slot()
+						.AutoHeight()
+						.Padding(0.0f)
+						[
+							SNew(SInViewportDetailsHeader)
+							.Parent(SharedThis(this))
+							.Content()
+							[		
+								SNew(SHorizontalBox)
+								+SHorizontalBox::Slot()
+								.AutoWidth()
+								.HAlign(HAlign_Left)
+								.VAlign(VAlign_Center)
+								.Padding(0)
+								[
+									SNew(SImage)
+									.Image(ActorIcon)
+								]
+								+ SHorizontalBox::Slot()
+								.Padding(10.0f, 2.0f)
+								[
+									SNew(STextBlock)
+									.Text(FText::FromString(NameString))
+									.TextStyle(FAppStyle::Get(), "DetailsView.CategoryTextStyle")
+								]
 							]
 						]
-					]
-					+ SVerticalBox::Slot()
-					.AutoHeight()
-					[
-						SNew(SInViewportDetailsToolbar)
-						.Parent(SharedThis(this))
-					]
-					+ SVerticalBox::Slot()
-					.AutoHeight()
-					[
-						MakeDetailsWidget()
+						+ SVerticalBox::Slot()
+						.AutoHeight()
+						.Padding(0.0f)
+						[
+							SNew(SInViewportDetailsToolbar)
+							.Parent(SharedThis(this))
+						]
+						+ SVerticalBox::Slot()
+						.AutoHeight()
+						.Padding(0.0f)
+						[
+							MakeDetailsWidget()
+						]
 					]
 				]
 			];
@@ -320,16 +355,13 @@ TSharedRef<SWidget> SInViewportDetails::MakeDetailsWidget()
 	if (Nodes.Num())
 	{
 		NodeList = SNew(SListView< TSharedPtr<IDetailTreeNode> >)
+			.ListViewStyle(&FAppStyle::Get().GetWidgetStyle<FTableViewStyle>("PropertyTable.InViewport.ListView"))
 			.ItemHeight(24)
 			.ListItemsSource(&Nodes)
 			.OnGenerateRow(this, &SInViewportDetails::GenerateListRow);
 
 		
-		DetailWidget = SNew(SBorder)
-			.BorderImage(FAppStyle::Get().GetBrush("PropertyTable.InViewport.Background"))
-			[
-				NodeList.ToSharedRef()
-			];
+		DetailWidget = NodeList.ToSharedRef();
 	}
 	return DetailWidget;
 }
@@ -436,7 +468,12 @@ void SInViewportDetailsHeader::Construct(const FArguments& InArgs)
 	ParentPtr = InArgs._Parent;
 	ChildSlot
 		[
-			InArgs._Content.Widget
+			SNew(SBorder)
+			.BorderImage(FAppStyle::Get().GetBrush("PropertyTable.InViewport.Header"))
+			.Padding(5.0f)
+			[
+				InArgs._Content.Widget
+			]
 		];
 }
 
diff --git a/Engine/Source/Editor/LevelEditor/Private/SLevelViewport.cpp b/Engine/Source/Editor/LevelEditor/Private/SLevelViewport.cpp
--- a/Engine/Source/Editor/LevelEditor/Private/SLevelViewport.cpp
+++ b/Engine/Source/Editor/LevelEditor/Private/SLevelViewport.cpp
@@ -90,7 +90,7 @@
 
 static const FName LevelEditorName("LevelEditor");
 static FAutoConsoleCommand EnableInViewportMenu(TEXT("Editor.EnableInViewportMenu"), TEXT("Enables the new in-viewport property menu"), FConsoleCommandDelegate::CreateStatic(&SLevelViewport::EnableInViewportMenu));
-bool SLevelViewport::bInViewportMenuEnabled = false;
+bool SLevelViewport::bInViewportMenuEnabled = true;
 #define LOCTEXT_NAMESPACE "LevelViewport"
 
 // @todo Slate Hack: Disallow game UI to be used in play in viewport until GWorld problem is fixed
diff --git a/Engine/Source/Runtime/SlateCore/Private/Widgets/SWindow.cpp b/Engine/Source/Runtime/SlateCore/Private/Widgets/SWindow.cpp
--- a/Engine/Source/Runtime/SlateCore/Private/Widgets/SWindow.cpp
+++ b/Engine/Source/Runtime/SlateCore/Private/Widgets/SWindow.cpp
@@ -491,6 +491,22 @@ TSharedRef<SWindow> SWindow::MakeCursorDecorator()
 	return NewWindow;
 }
 
+TSharedRef<SWindow> SWindow::MakeStyledCursorDecorator(const FWindowStyle& InStyle)
+{
+	TSharedRef<SWindow> NewWindow = SNew(SWindow)
+		.Style(&InStyle)
+		.Type(EWindowType::CursorDecorator)
+		.IsPopupWindow(true)
+		.IsTopmostWindow(true)
+		.SizingRule(ESizingRule::Autosized)
+		.SupportsTransparency(EWindowTransparency::PerWindow)
+		.FocusWhenFirstShown(false)
+		.ActivationPolicy(EWindowActivationPolicy::Never);
+	NewWindow->Opacity = 1.0f;
+
+	return NewWindow;
+}
+
 FVector2D SWindow::ComputeWindowSizeForContent( FVector2D ContentSize )
 {
 	// @todo mainframe: This code should be updated to handle the case where we're spawning a window that doesn't have
diff --git a/Engine/Source/Runtime/SlateCore/Public/Widgets/SWindow.h b/Engine/Source/Runtime/SlateCore/Public/Widgets/SWindow.h
--- a/Engine/Source/Runtime/SlateCore/Public/Widgets/SWindow.h
+++ b/Engine/Source/Runtime/SlateCore/Public/Widgets/SWindow.h
@@ -296,6 +296,14 @@ public:
 	 */
 	static TSharedRef<SWindow> MakeCursorDecorator();
 
+	/**
+	 * Make cursor decorator window with a non-default style
+	 *
+	 * @param InStyle The style to use for the cursor decorator
+	 * @return The new SWindow
+	 */
+	static TSharedRef<SWindow> MakeStyledCursorDecorator(const FWindowStyle& InStyle);
+
 	/**
 	 * Make a notification window
 	 *
