Fixed the sub-object instance editor toolbar being visible, even when non-actor elements were selected

The button bar of the sub-object instance editor would always be visible, even if there was no valid actor context to operate on. This changes it to only appear for actor contexts (ie, we can't add components to static mesh instance elements).

diff --git a/Engine/Source/Editor/LevelEditor/Private/SActorDetails.cpp b/Engine/Source/Editor/LevelEditor/Private/SActorDetails.cpp
--- a/Engine/Source/Editor/LevelEditor/Private/SActorDetails.cpp
+++ b/Engine/Source/Editor/LevelEditor/Private/SActorDetails.cpp
@@ -137,7 +137,7 @@ void SActorDetails::Construct(const FArguments& InArgs, UTypedElementSelectionSe
 
 	ComponentsBox = SNew(SBox)
 		.Padding(FMargin(2.0f, 0.0f, 2.0f, 0.0f))
-		.Visibility(this, &SActorDetails::GetComponentsBoxVisibility);
+		.Visibility(this, &SActorDetails::GetComponentEditorVisibility);
 
 	FModuleManager::LoadModuleChecked<FSubobjectEditorModule>("SubobjectEditor");
 	
@@ -151,6 +151,7 @@ void SActorDetails::Construct(const FArguments& InArgs, UTypedElementSelectionSe
 	ComponentsBox->SetContent(SubobjectEditor.ToSharedRef());
 
 	TSharedRef<SWidget> ButtonBox = SubobjectEditor->GetToolButtonsBox().ToSharedRef();
+	ButtonBox->SetVisibility(MakeAttributeSP(this, &SActorDetails::GetComponentEditorVisibility));
 	DetailsView->SetNameAreaCustomContent( ButtonBox );
 
 	ChildSlot
@@ -709,7 +710,7 @@ void SActorDetails::OnNativeComponentWarningHyperlinkClicked(const FSlateHyperli
 	}
 }
 
-EVisibility SActorDetails::GetComponentsBoxVisibility() const
+EVisibility SActorDetails::GetComponentEditorVisibility() const
 {
 	return GetActorContext() ? EVisibility::Visible : EVisibility::Collapsed;
 }
diff --git a/Engine/Source/Editor/LevelEditor/Private/SActorDetails.h b/Engine/Source/Editor/LevelEditor/Private/SActorDetails.h
--- a/Engine/Source/Editor/LevelEditor/Private/SActorDetails.h
+++ b/Engine/Source/Editor/LevelEditor/Private/SActorDetails.h
@@ -90,7 +90,7 @@ private:
 
 	bool IsPropertyReadOnly(const struct FPropertyAndParent& PropertyAndParent) const;
 	bool IsPropertyEditingEnabled() const;
-	EVisibility GetComponentsBoxVisibility() const;
+	EVisibility GetComponentEditorVisibility() const;
 	EVisibility GetUCSComponentWarningVisibility() const;
 	EVisibility GetInheritedBlueprintComponentWarningVisibility() const;
 	EVisibility GetNativeComponentWarningVisibility() const;
