Add support for Play in new window to the Play From Here actor flow

diff --git a/Engine/Source/Editor/LevelEditor/Private/LevelEditor.cpp b/Engine/Source/Editor/LevelEditor/Private/LevelEditor.cpp
--- a/Engine/Source/Editor/LevelEditor/Private/LevelEditor.cpp
+++ b/Engine/Source/Editor/LevelEditor/Private/LevelEditor.cpp
@@ -719,15 +719,7 @@ void FLevelEditorModule::BindGlobalLevelEditorCommands()
 		Commands.GoHere,
 		FExecuteAction::CreateStatic( &FLevelEditorActionCallbacks::GoHere_Clicked, (const FVector*)nullptr )
 		);
-
-	ActionList.MapAction(
-		Commands.PlayFromHere,
-		FExecuteAction::CreateStatic( &FLevelEditorActionCallbacks::PlayFromHere_Clicked ), 
-		FCanExecuteAction(),
-		FIsActionChecked(),
-		FIsActionButtonVisible::CreateStatic( &FLevelEditorActionCallbacks::PlayFromHere_IsVisible )
-		);
-
+		
 	ActionList.MapAction( 
 		Commands.SnapCameraToObject,
 		FExecuteAction::CreateStatic( &FLevelEditorActionCallbacks::ExecuteExecCommand, FString( TEXT("CAMERA SNAP") ) )
diff --git a/Engine/Source/Editor/LevelEditor/Private/LevelEditorActions.cpp b/Engine/Source/Editor/LevelEditor/Private/LevelEditorActions.cpp
--- a/Engine/Source/Editor/LevelEditor/Private/LevelEditorActions.cpp
+++ b/Engine/Source/Editor/LevelEditor/Private/LevelEditorActions.cpp
@@ -1312,14 +1312,15 @@ void FLevelEditorActionCallbacks::CreateNewOutlinerFolder_Clicked()
 	FActorFolders::Get().CreateFolderContainingSelection(*GetWorld(), NewFolderName);
 }
 
-void FLevelEditorActionCallbacks::PlayFromHere_Clicked()
+void FLevelEditorActionCallbacks::PlayFromHere_Clicked(bool bFloatingWindow)
 {
 	if (GEditor->GetSelectedActorCount() == 1)
 	{
 		if (AActor* Actor = Cast<AActor>(*GEditor->GetSelectedActorIterator()))
 		{
 			Actor->GetWorld()->PersistentLevel->PlayFromHereActor = Actor;
-			FPlayWorldCommandCallbacks::StartPlayFromHere(Actor->GetActorLocation(), Actor->GetActorRotation());
+			FLevelEditorModule& LevelEditorModule = FModuleManager::GetModuleChecked<FLevelEditorModule>("LevelEditor");
+			FPlayWorldCommandCallbacks::StartPlayFromHere(Actor->GetActorLocation(), Actor->GetActorRotation(), bFloatingWindow ? nullptr : LevelEditorModule.GetFirstActiveViewport());
 		}
 	}
 }
@@ -3346,7 +3347,6 @@ void FLevelEditorCommands::RegisterCommands()
 	UI_COMMAND( EditAssetNoConfirmMultiple, "Edit Asset", "Edits the asset associated with the selected actor", EUserInterfaceActionType::Button, FInputChord( EKeys::E, EModifierKey::Control | EModifierKey::Shift ) );
 
 	UI_COMMAND( GoHere, "Go Here", "Moves the camera to the current mouse position", EUserInterfaceActionType::Button, FInputChord() );
-	UI_COMMAND( PlayFromHere, "Play From Here", "Start PIE session from this actor", EUserInterfaceActionType::Button, FInputChord()) ;
 
 	UI_COMMAND( SnapCameraToObject, "Snap View to Object", "Snaps the view to the selected object", EUserInterfaceActionType::Button, FInputChord() );
 	UI_COMMAND( SnapObjectToCamera, "Snap Object to View", "Snaps the selected object to the view", EUserInterfaceActionType::Button, FInputChord() );
diff --git a/Engine/Source/Editor/LevelEditor/Private/LevelEditorContextMenu.cpp b/Engine/Source/Editor/LevelEditor/Private/LevelEditorContextMenu.cpp
--- a/Engine/Source/Editor/LevelEditor/Private/LevelEditorContextMenu.cpp
+++ b/Engine/Source/Editor/LevelEditor/Private/LevelEditorContextMenu.cpp
@@ -217,7 +217,7 @@ void FLevelEditorContextMenu::RegisterComponentContextMenu()
 			Section.AddMenuEntry(FLevelEditorCommands::Get().GoHere);
 			Section.AddMenuEntry(FLevelEditorCommands::Get().SnapCameraToObject);
 			Section.AddMenuEntry(FLevelEditorCommands::Get().SnapObjectToCamera);
-			Section.AddMenuEntry(FLevelEditorCommands::Get().PlayFromHere);
+			AddPlayFromHereSubMenu(Section);						
 			Section.AddMenuEntry(FLevelEditorCommands::Get().CopyActorFilePathtoClipboard);
 
 			FLevelEditorContextMenuImpl::AddSourceControlMenu(Section);
@@ -227,6 +227,22 @@ void FLevelEditorContextMenu::RegisterComponentContextMenu()
 	}));
 }
 
+void FLevelEditorContextMenu::AddPlayFromHereSubMenu(FToolMenuSection& Section)
+{
+	if(FLevelEditorActionCallbacks::PlayFromHere_IsVisible())
+	{
+		Section.AddSubMenu("PlayFromHere", LOCTEXT("PlayFromHere", "Play From Here"), LOCTEXT("", ""), FNewToolMenuDelegate::CreateLambda([](UToolMenu* InMenu)
+		{
+			FToolMenuSection& NewSection = InMenu->AddSection("Section");
+			FUIAction PlayFromHere(FExecuteAction::CreateStatic(&FLevelEditorActionCallbacks::PlayFromHere_Clicked, false));
+			FUIAction PlayFromHereFloating(FExecuteAction::CreateStatic(&FLevelEditorActionCallbacks::PlayFromHere_Clicked, true));
+
+			NewSection.AddMenuEntry("PlayFromHereActiveViewport", LOCTEXT("PlayFromHereActiveViewport", "Selected Viewport"), LOCTEXT("PlayFromHereActiveViewportTooltip", "Play from this actor in the active level editor viewport"), FSlateIcon("EditorSTyle", "PlayWorld.PlayInViewport"), PlayFromHere);
+			NewSection.AddMenuEntry("PlayFromHereFloatingWindow", LOCTEXT("PlayFromHereFloatingWindow", "New Editor Window (PIE)"), LOCTEXT("PlayFromHereFloatingWindowTooltip", "Play from this actor in a new editor window"), FSlateIcon("EditorSTyle", "PlayWorld.PlayInEditorFloating"), PlayFromHereFloating);
+		}));
+	}
+}
+
 void FLevelEditorContextMenu::RegisterActorContextMenu()
 {
 	UToolMenus* ToolMenus = UToolMenus::Get();
@@ -306,7 +322,7 @@ void FLevelEditorContextMenu::RegisterActorContextMenu()
 			Section.AddMenuEntry(FLevelEditorCommands::Get().GoHere);
 			Section.AddMenuEntry(FLevelEditorCommands::Get().SnapCameraToObject);
 			Section.AddMenuEntry(FLevelEditorCommands::Get().SnapObjectToCamera);
-			Section.AddMenuEntry(FLevelEditorCommands::Get().PlayFromHere);
+			AddPlayFromHereSubMenu(Section);
 			Section.AddMenuEntry(FLevelEditorCommands::Get().CopyActorFilePathtoClipboard);
 			
 			FLevelEditorContextMenuImpl::AddSourceControlMenu(Section);
diff --git a/Engine/Source/Editor/LevelEditor/Private/LevelEditorContextMenu.h b/Engine/Source/Editor/LevelEditor/Private/LevelEditorContextMenu.h
--- a/Engine/Source/Editor/LevelEditor/Private/LevelEditorContextMenu.h
+++ b/Engine/Source/Editor/LevelEditor/Private/LevelEditorContextMenu.h
@@ -71,6 +71,8 @@ private:
 	static void RegisterSceneOutlinerContextMenu();
 	static void RegisterEmptySelectionContextMenu();
 
+	static void AddPlayFromHereSubMenu(FToolMenuSection& Section);
+
 	/**
 	 * Builds the actor group menu
 	 *
diff --git a/Engine/Source/Editor/LevelEditor/Public/LevelEditorActions.h b/Engine/Source/Editor/LevelEditor/Public/LevelEditorActions.h
--- a/Engine/Source/Editor/LevelEditor/Public/LevelEditorActions.h
+++ b/Engine/Source/Editor/LevelEditor/Public/LevelEditorActions.h
@@ -125,9 +125,6 @@ public:
 	/** Moves the camera to the current mouse position */
 	TSharedPtr< FUICommandInfo > GoHere;
 
-	/** Play from here on selected actor. */
-	TSharedPtr< FUICommandInfo > PlayFromHere;
-
 	/** Snaps the camera to the selected object. */
 	TSharedPtr< FUICommandInfo > SnapCameraToObject;
 
@@ -928,7 +925,7 @@ public:
 	static void GoHere_Clicked( const FVector* Point );
 
 	/** Called when selected actor can be used to start a play session */
-	static void PlayFromHere_Clicked();
+	static void PlayFromHere_Clicked(bool bFloatingWindow);
 	static bool PlayFromHere_IsVisible();
 
 	/** Called when 'Go to Code for Actor' is clicked */
diff --git a/Engine/Source/Editor/UnrealEd/Private/Kismet2/DebuggerCommands.cpp b/Engine/Source/Editor/UnrealEd/Private/Kismet2/DebuggerCommands.cpp
--- a/Engine/Source/Editor/UnrealEd/Private/Kismet2/DebuggerCommands.cpp
+++ b/Engine/Source/Editor/UnrealEd/Private/Kismet2/DebuggerCommands.cpp
@@ -886,10 +886,10 @@ void FPlayWorldCommandCallbacks::StartPlayFromHere()
 		StartRotation = ActiveLevelViewport->GetAssetViewportClient().GetViewRotation();
 	}
 
-	StartPlayFromHere(StartLocation, StartRotation);
+	StartPlayFromHere(StartLocation, StartRotation, ActiveLevelViewport);
 }
 
-void FPlayWorldCommandCallbacks::StartPlayFromHere(const TOptional<FVector>& Location, const TOptional<FRotator>& Rotation)
+void FPlayWorldCommandCallbacks::StartPlayFromHere(const TOptional<FVector>& Location, const TOptional<FRotator>& Rotation, const TSharedPtr<IAssetViewport>& ActiveLevelViewport)
 {
 	// Is a PIE session already running?  If so we close it first
 	if (GUnrealEd->PlayWorld != NULL)
@@ -903,8 +903,6 @@ void FPlayWorldCommandCallbacks::StartPlayFromHere(const TOptional<FVector>& Loc
 
 	FLevelEditorModule& LevelEditorModule = FModuleManager::GetModuleChecked<FLevelEditorModule>("LevelEditor");
 
-	TSharedPtr<IAssetViewport> ActiveLevelViewport = LevelEditorModule.GetFirstActiveViewport();
-
 	if (ActiveLevelViewport.IsValid() && ActiveLevelViewport->GetAssetViewportClient().IsPerspective())
 	{
 		// If there is no level viewport, a new window will be spawned to play in.
diff --git a/Engine/Source/Editor/UnrealEd/Public/Kismet2/DebuggerCommands.h b/Engine/Source/Editor/UnrealEd/Public/Kismet2/DebuggerCommands.h
--- a/Engine/Source/Editor/UnrealEd/Public/Kismet2/DebuggerCommands.h
+++ b/Engine/Source/Editor/UnrealEd/Public/Kismet2/DebuggerCommands.h
@@ -10,6 +10,7 @@
 #include "Framework/Commands/Commands.h"
 #include "Framework/Commands/UICommandList.h"
 
+class IAssetViewport;
 class FToolBarBuilder;
 struct FToolMenuSection;
 
@@ -138,7 +139,7 @@ public:
 	 * Called from the context menu to start previewing the game at the clicked location                   
 	 */
 	static void StartPlayFromHere();
-	static void StartPlayFromHere(const TOptional<FVector>& Location, const TOptional<FRotator>& Rotation);
+	static void StartPlayFromHere(const TOptional<FVector>& Location, const TOptional<FRotator>& Rotation, const TSharedPtr<IAssetViewport>& ActiveLevelViewport);
 
 	static void ResumePlaySession_Clicked();
 	static void PausePlaySession_Clicked();
