Correctly aquire the editor component handle from an actor component when the tree selection is updated in the SActorDetails panel

#rnx

diff --git a/Engine/Source/Editor/LevelEditor/Private/SActorDetails.cpp b/Engine/Source/Editor/LevelEditor/Private/SActorDetails.cpp
--- a/Engine/Source/Editor/LevelEditor/Private/SActorDetails.cpp
+++ b/Engine/Source/Editor/LevelEditor/Private/SActorDetails.cpp
@@ -544,7 +544,7 @@ void SActorDetails::OnSubobjectEditorTreeViewSelectionChanged(const TArray<FSubo
 			{
 				if(FSubobjectData* Data = SelectedNode->GetDataSource())
 				{
-					if(Data->IsDefaultSceneRoot())
+					if(Data->IsRootActor())
 					{
 						// If the actor node is selected then we ignore the component selection
 						NewEditorSelection.Reset();
@@ -552,7 +552,7 @@ void SActorDetails::OnSubobjectEditorTreeViewSelectionChanged(const TArray<FSubo
 						break;
 					}
 				
-					if (Data->IsInstancedComponent())
+					if (Data->IsComponent())
 					{
 						if (const UActorComponent* Component = Data->FindComponentInstanceInActor(ActorContext))
 						{
