Fixed some places leaving lingering element references from selection

diff --git a/Engine/Plugins/Developer/Concert/ConcertSync/ConcertSyncClient/Source/ConcertSyncClient/Private/ConcertSyncClientUtil.cpp b/Engine/Plugins/Developer/Concert/ConcertSync/ConcertSyncClient/Source/ConcertSyncClient/Private/ConcertSyncClientUtil.cpp
--- a/Engine/Plugins/Developer/Concert/ConcertSync/ConcertSyncClient/Source/ConcertSyncClient/Private/ConcertSyncClientUtil.cpp
+++ b/Engine/Plugins/Developer/Concert/ConcertSync/ConcertSyncClient/Source/ConcertSyncClient/Private/ConcertSyncClientUtil.cpp
@@ -35,6 +35,7 @@
 	#include "Editor/UnrealEdEngine.h"
 	#include "PackageTools.h"
 	#include "ObjectTools.h"
+	#include "Selection.h"
 	#include "Subsystems/AssetEditorSubsystem.h"
 	#include "GameMapsSettings.h"
 	#include "FileHelpers.h"
@@ -399,6 +400,7 @@ void PurgePackages(TArrayView<const FName> InPackageNames)
 		if (InObject->IsAsset() && GIsEditor)
 		{
 			GEditor->GetEditorSubsystem<UAssetEditorSubsystem>()->CloseAllEditorsForAsset(InObject);
+			GEditor->GetSelectedObjects()->Deselect(InObject);
 		}
 		ObjectsToPurge.Add(InObject);
 	};
diff --git a/Engine/Source/Editor/LevelEditor/Private/SLevelEditor.cpp b/Engine/Source/Editor/LevelEditor/Private/SLevelEditor.cpp
--- a/Engine/Source/Editor/LevelEditor/Private/SLevelEditor.cpp
+++ b/Engine/Source/Editor/LevelEditor/Private/SLevelEditor.cpp
@@ -210,6 +210,7 @@ void SLevelEditor::Construct( const SLevelEditor::FArguments& InArgs)
 		// Set the initial preview feature level.
 		World->ChangeFeatureLevel(GEditor->GetActiveFeatureLevelPreviewType());
 
+		LevelActorDeletedHandle = GEditor->OnLevelActorDeleted().AddSP(this, &SLevelEditor::OnLevelActorDeleted);
 		LevelActorOuterChangedHandle = GEditor->OnLevelActorOuterChanged().AddSP(this, &SLevelEditor::OnLevelActorOuterChanged);
 	}
 
@@ -373,6 +374,7 @@ SLevelEditor::~SLevelEditor()
 		{
 			AssetEditorSubsystem->OnEditorModesChanged().RemoveAll(this);
 		}
+		GEditor->OnLevelActorDeleted().Remove(LevelActorDeletedHandle);
 		GEditor->OnLevelActorOuterChanged().Remove(LevelActorOuterChangedHandle);
 		GEditor->GetEditorWorldContext(true).RemoveRef(World);
 	}
@@ -1864,6 +1866,24 @@ void SLevelEditor::OnOverridePropertyEditorSelection(const TArray<AActor*>& NewS
 	bNeedsRefresh = false;
 }
 
+void SLevelEditor::OnLevelActorDeleted(AActor* InActor)
+{
+	if (SelectedElements && InActor->GetPackage()->HasAnyPackageFlags(PKG_PlayInEditor))
+	{
+		// Deselect PIE actors when they are deleted, as these won't go through the normal editor flow that will make sure an element is deselected prior to being destroyed
+		if (FTypedElementHandle ActorElement = UEngineElementsLibrary::AcquireEditorActorElementHandle(InActor, /*bAllowCreate*/false))
+		{
+			const FTypedElementSelectionOptions SelectionOptions = FTypedElementSelectionOptions()
+				.SetAllowHidden(true)
+				.SetAllowGroups(false)
+				.SetAllowLegacyNotifications(false)
+				.SetWarnIfLocked(false);
+
+			SelectedElements->DeselectElement(ActorElement, SelectionOptions);
+		}
+	}
+}
+
 void SLevelEditor::OnLevelActorOuterChanged(AActor* InActor, UObject* InOldOuter)
 {
 	bNeedsRefresh = true;
diff --git a/Engine/Source/Editor/LevelEditor/Private/SLevelEditor.h b/Engine/Source/Editor/LevelEditor/Private/SLevelEditor.h
--- a/Engine/Source/Editor/LevelEditor/Private/SLevelEditor.h
+++ b/Engine/Source/Editor/LevelEditor/Private/SLevelEditor.h
@@ -246,6 +246,9 @@ private:
 	/** Called to set property editors to show the given actors, even if those actors aren't in the current selection set */
 	void OnOverridePropertyEditorSelection(const TArray<AActor*>& NewSelection, bool bForceRefresh = false);
 
+	/** Called when an actor is destroyed */
+	void OnLevelActorDeleted(AActor* InActor);
+
 	/** Called when an actor changes outer */
 	void OnLevelActorOuterChanged(AActor* InActor = nullptr, UObject* InOldOuter = nullptr);
 
@@ -313,6 +316,9 @@ private:
 	/** Handle to the registered OnPreviewFeatureLevelChanged delegate. */
 	FDelegateHandle PreviewFeatureLevelChangedHandle;
 
+	/** Handle to the registered OnLevelActorDeleted delegate */
+	FDelegateHandle LevelActorDeletedHandle;
+
 	/** Handle to the registered OnLevelActorOuterChanged delegate */
 	FDelegateHandle LevelActorOuterChangedHandle;
 
