Fix non-unity build

diff --git a/Engine/Source/Editor/LevelEditor/Private/LevelEditorContextMenu.h b/Engine/Source/Editor/LevelEditor/Private/LevelEditorContextMenu.h
--- a/Engine/Source/Editor/LevelEditor/Private/LevelEditorContextMenu.h
+++ b/Engine/Source/Editor/LevelEditor/Private/LevelEditorContextMenu.h
@@ -12,6 +12,7 @@ class UTypedElementSelectionSet;
 struct FToolMenuContext;
 class SLevelEditor;
 class SWidget;
+struct FToolMenuSection;
 
 /**
  * Context menu construction class 
