Fixed actor details not showing the properties of the root componenent if selected.
Updated the logic in SActorDetails::RefreshSubobjectTreeElements to match the logic in  SActorDetails::OnSubobjectEditorTreeViewSelectionChanged when deciding whether a component or actor is selected

diff --git a/Engine/Source/Editor/LevelEditor/Private/SActorDetails.cpp b/Engine/Source/Editor/LevelEditor/Private/SActorDetails.cpp
--- a/Engine/Source/Editor/LevelEditor/Private/SActorDetails.cpp
+++ b/Engine/Source/Editor/LevelEditor/Private/SActorDetails.cpp
@@ -353,7 +353,7 @@ void SActorDetails::RefreshSubobjectTreeElements(TArrayView<const FSubobjectEdit
 			{
 				if(const FSubobjectData* Data = SelectedNode->GetDataSource())
 				{
-					if(Data->IsDefaultSceneRoot())
+					if(Data->IsRootActor())
 					{
 						// If the actor node is selected then we ignore the component selection
                         Components.Reset();
