Matinee/Editor: Renamed usages of the word 'Matinee' to 'Cinematics' where applicable in code comments.

diff --git a/Engine/Source/Editor/LevelEditor/Private/LevelViewportActions.cpp b/Engine/Source/Editor/LevelEditor/Private/LevelViewportActions.cpp
--- a/Engine/Source/Editor/LevelEditor/Private/LevelViewportActions.cpp
+++ b/Engine/Source/Editor/LevelEditor/Private/LevelViewportActions.cpp
@@ -63,7 +63,7 @@ void FLevelViewportCommands::RegisterCommands()
 
 	UI_COMMAND( ApplyMaterialToActor, "Apply Material", "Attempts to apply a dropped material to this object", EUserInterfaceActionType::Button, FInputChord() );
 
-	UI_COMMAND( ToggleCinematicPreview, "Allow Cinematic Control", "If enabled, allows Matinee or Sequencer previews to play in this viewport", EUserInterfaceActionType::ToggleButton, FInputChord() );
+	UI_COMMAND( ToggleCinematicPreview, "Allow Cinematic Control", "If enabled, allows cinematic (Sequencer) previews to play in this viewport", EUserInterfaceActionType::ToggleButton, FInputChord() );
 
 	UI_COMMAND( FindInLevelScriptBlueprint, "Find In Level Script", "Finds references of a selected actor in the level script blueprint", EUserInterfaceActionType::Button, FInputChord(EModifierKey::Control, EKeys::K) );
 	UI_COMMAND( AdvancedSettings, "Advanced Settings...", "Opens the advanced viewport settings", EUserInterfaceActionType::Button, FInputChord());
diff --git a/Engine/Source/Editor/LevelEditor/Private/SLevelViewport.cpp b/Engine/Source/Editor/LevelEditor/Private/SLevelViewport.cpp
--- a/Engine/Source/Editor/LevelEditor/Private/SLevelViewport.cpp
+++ b/Engine/Source/Editor/LevelEditor/Private/SLevelViewport.cpp
@@ -2471,7 +2471,7 @@ void SLevelViewport::OnCompactBookmarks()
 
 void SLevelViewport::OnToggleAllowCinematicPreview()
 {
-	// Reset the FOV of Viewport for cases where we have been previewing the matinee with a changing FOV
+	// Reset the FOV of Viewport for cases where we have been previewing the cinematic with a changing FOV
 	LevelViewportClient->ViewFOV = LevelViewportClient->AllowsCinematicControl() ? LevelViewportClient->ViewFOV : LevelViewportClient->FOVAngle;
 
 	LevelViewportClient->SetAllowCinematicControl( !LevelViewportClient->AllowsCinematicControl() );
@@ -3627,7 +3627,7 @@ void SLevelViewport::PreviewActors( const TArray< AActor* >& InActorsToPreview,
 				// User should never be able to interact with this viewport
 				ActorPreviewLevelViewportClient->bDisableInput = true;
 
-				// Never allow Matinee to possess these views
+				// Never allow cinematics to possess these views
 				ActorPreviewLevelViewportClient->SetAllowCinematicControl( false );
 
 				// Our preview viewport is always visible if our owning SLevelViewport is visible, so we hook up
diff --git a/Engine/Source/Editor/LevelEditor/Public/SLevelViewport.h b/Engine/Source/Editor/LevelEditor/Public/SLevelViewport.h
--- a/Engine/Source/Editor/LevelEditor/Public/SLevelViewport.h
+++ b/Engine/Source/Editor/LevelEditor/Public/SLevelViewport.h
@@ -580,7 +580,7 @@ private:
 	void OnToggleAllowCinematicPreview();
 
 	/**
-	 * @return true if this viewport allows matinee to be previewed in it                   
+	 * @return true if this viewport allows cinematics to be previewed in it                   
 	 */
 	bool AllowsCinematicPreview() const;
 
diff --git a/Engine/Source/Editor/UnrealEd/Private/EditorViewportClient.cpp b/Engine/Source/Editor/UnrealEd/Private/EditorViewportClient.cpp
--- a/Engine/Source/Editor/UnrealEd/Private/EditorViewportClient.cpp
+++ b/Engine/Source/Editor/UnrealEd/Private/EditorViewportClient.cpp
@@ -77,7 +77,7 @@ static TAutoConsoleVariable<int32> CVarAlignedOrthoZoom(
 static TAutoConsoleVariable<int32> CVarEditorViewportTest(
 	TEXT("r.Test.EditorConstrainedView"),
 	0,
-	TEXT("Allows to test different viewport rectangle configuations (in game only) as they can happen when using Matinee/Editor.\n")
+	TEXT("Allows to test different viewport rectangle configuations (in game only) as they can happen when using cinematics/Editor.\n")
 	TEXT("0: off(default)\n")
 	TEXT("1..7: Various Configuations"),
 	ECVF_RenderThreadSafe);
@@ -5098,7 +5098,7 @@ void FEditorViewportClient::MoveViewportPerspectiveCamera( const FVector& InDrag
 	}
 	else
 	{
-		//when not constraining the pitch (matinee feature) we need to rotate differently to avoid a gimbal lock
+		//when not constraining the pitch we need to rotate differently to avoid a gimbal lock
 		const FRotator PitchRot(InRot.Pitch, 0, 0);
 		const FRotator LateralRot(0, InRot.Yaw, InRot.Roll);
 
diff --git a/Engine/Source/Editor/UnrealEd/Public/LevelEditorViewport.h b/Engine/Source/Editor/UnrealEd/Public/LevelEditorViewport.h
--- a/Engine/Source/Editor/UnrealEd/Public/LevelEditorViewport.h
+++ b/Engine/Source/Editor/UnrealEd/Public/LevelEditorViewport.h
@@ -962,7 +962,7 @@ private:
 	/** Caching for expensive FindViewComponentForActor. Invalidated once per Tick. */
 	static TMap<TObjectKey<AActor>, TWeakObjectPtr<UActorComponent>> ViewComponentForActorCache;
 
-	/** If true, we switched between two different cameras. Set by matinee, used by the motion blur to invalidate this frames motion vectors */
+	/** If true, we switched between two different cameras. Set by cinematics, used by the motion blur to invalidate this frames motion vectors */
 	bool					bEditorCameraCut;
 
 	/** Stores the previous frame's value of bEditorCameraCut in order to reset it back to false on the next frame */
