Add cursor world location to level editor tool menu context

diff --git a/Engine/Source/Editor/LevelEditor/Private/LevelEditorContextMenu.cpp b/Engine/Source/Editor/LevelEditor/Private/LevelEditorContextMenu.cpp
--- a/Engine/Source/Editor/LevelEditor/Private/LevelEditorContextMenu.cpp
+++ b/Engine/Source/Editor/LevelEditor/Private/LevelEditorContextMenu.cpp
@@ -616,6 +616,25 @@ FName FLevelEditorContextMenu::InitMenuContext(FToolMenuContext& Context, TWeakP
 	{
 		ContextObject->SelectedComponents.Add(CastChecked<UActorComponent>(*It));
 	}
+
+	// obtain the world location of the cursor
+	if (GCurrentLevelEditingViewportClient)
+	{
+		FHitResult HitResult;
+		FViewportCursorLocation CursorLocation = GCurrentLevelEditingViewportClient->GetCursorWorldLocationFromMousePos();
+		FCollisionQueryParams LineParams(SCENE_QUERY_STAT(FocusOnPoint), true);
+
+		if (GCurrentLevelEditingViewportClient->GetWorld()->LineTraceSingleByObjectType(
+			HitResult,
+			CursorLocation.GetOrigin(),
+			CursorLocation.GetOrigin() + CursorLocation.GetDirection() * HALF_WORLD_MAX,
+			FCollisionObjectQueryParams(FCollisionObjectQueryParams::InitType::AllObjects),
+			LineParams))
+		{
+			ContextObject->CursorWorldLocation = HitResult.ImpactPoint;
+		}
+	}
+	
 	Context.AddObject(ContextObject, [](UObject* InContext)
 	{
 		ULevelEditorContextMenuContext* CastContext = CastChecked<ULevelEditorContextMenuContext>(InContext);
diff --git a/Engine/Source/Editor/LevelEditor/Public/LevelEditorMenuContext.h b/Engine/Source/Editor/LevelEditor/Public/LevelEditorMenuContext.h
--- a/Engine/Source/Editor/LevelEditor/Public/LevelEditorMenuContext.h
+++ b/Engine/Source/Editor/LevelEditor/Public/LevelEditorMenuContext.h
@@ -25,7 +25,8 @@ public:
 
 
 /** Enum to describe what a level editor context menu should be built for */
-enum class ELevelEditorMenuContext
+UENUM()
+enum class ELevelEditorMenuContext : uint8
 {
 	/** This context menu is applicable to a viewport */
 	Viewport,
@@ -41,6 +42,7 @@ public:
 
 	TWeakPtr<SLevelEditor> LevelEditor;
 
+	UPROPERTY(VisibleAnywhere, BlueprintReadOnly, Category="LevelEditor|Menu")
 	ELevelEditorMenuContext ContextType;
 
 	UPROPERTY(VisibleAnywhere, BlueprintReadOnly, Category="LevelEditor|Menu")
@@ -50,6 +52,9 @@ public:
 	UPROPERTY(VisibleAnywhere, BlueprintReadOnly, Category="LevelEditor|Menu")
 	FTypedElementHandle HitProxyElement;
 
+	UPROPERTY(VisibleAnywhere, BlueprintReadOnly, Category="LevelEditor|Menu")
+	FVector CursorWorldLocation;
+
 	UPROPERTY()
 	TArray<TObjectPtr<UActorComponent>> SelectedComponents;
 
