Also deselect child elements when an actor is destroyed during PIE

diff --git a/Engine/Source/Editor/LevelEditor/Private/SLevelEditor.cpp b/Engine/Source/Editor/LevelEditor/Private/SLevelEditor.cpp
--- a/Engine/Source/Editor/LevelEditor/Private/SLevelEditor.cpp
+++ b/Engine/Source/Editor/LevelEditor/Private/SLevelEditor.cpp
@@ -1860,7 +1860,8 @@ void SLevelEditor::OnLevelActorDeleted(AActor* InActor)
 				.SetAllowHidden(true)
 				.SetAllowGroups(false)
 				.SetAllowLegacyNotifications(false)
-				.SetWarnIfLocked(false);
+				.SetWarnIfLocked(false)
+				.SetChildElementInclusionMethod(ETypedElementChildInclusionMethod::Recursive);
 
 			SelectedElements->DeselectElement(ActorElement, SelectionOptions);
 		}
