Update the LevelViewportLayout size every frame after the parent widget is updated.

diff --git a/Engine/Source/Editor/LevelEditor/Private/LevelViewportLayout.cpp b/Engine/Source/Editor/LevelEditor/Private/LevelViewportLayout.cpp
--- a/Engine/Source/Editor/LevelEditor/Private/LevelViewportLayout.cpp
+++ b/Engine/Source/Editor/LevelEditor/Private/LevelViewportLayout.cpp
@@ -352,15 +352,48 @@ void FLevelViewportLayout::MaximizeViewport( FName ViewportToMaximize, const boo
 					// have a single viewport widget in two places at once!
 					LayoutConfiguration->ReplaceWidget( MaximizedEntity->AsWidget(), ViewportReplacementWidget.ToSharedRef() );
 
-					TSharedRef< SCanvas > Canvas = SNew( SCanvas );
-					Canvas->AddSlot()
-						.Position( TAttribute< FVector2D >( this, &FLevelViewportLayout::GetMaximizedViewportPositionOnCanvas ) )
-						.Size( TAttribute< FVector2D >( this, &FLevelViewportLayout::GetMaximizedViewportSizeOnCanvas ) )
-						[
-							MaximizedEntity->AsWidget()
-						];
-
-					ViewportsOverlayWidget = Canvas;
+					// The attributes need the AllocatedSize of the parent.
+					// The size is updated in the Paint function and the attributes in the Prepass (too soon).
+					// Update the value manually Tick function (after the parent's Paint).
+					class SCanvasInternal : public SCanvas
+					{
+					public:
+						SCanvasInternal()
+						{
+							SetCanTick(true);
+						}
+
+						void Construct(const FArguments& Args, TSharedRef<FLevelViewportLayout> ViewportLayout, TSharedRef<SWidget> MaximizedEntity)
+						{
+							SCanvas::Construct(Args);
+							OwnerViewportLayout = ViewportLayout;
+							AddSlot()
+								.Expose(ViewportsOverlayWidgetSlot)
+								[
+									MaximizedEntity
+								];
+						}
+
+						virtual void Tick(const FGeometry& AllottedGeometry, const double InCurrentTime, const float InDeltaTime) override
+						{
+							SCanvas::Tick(AllottedGeometry, InCurrentTime, InDeltaTime);
+							if (TSharedPtr<FLevelViewportLayout> Owner = OwnerViewportLayout.Pin())
+							{
+								ViewportsOverlayWidgetSlot->SetPosition(Owner->GetMaximizedViewportPositionOnCanvas());
+								ViewportsOverlayWidgetSlot->SetSize(Owner->GetMaximizedViewportSizeOnCanvas());
+							}
+						}
+
+					private:
+						using SCanvas::AddSlot;
+						using SCanvas::RemoveSlot;
+
+					private:
+						SCanvas::FSlot* ViewportsOverlayWidgetSlot = nullptr;
+						TWeakPtr<FLevelViewportLayout> OwnerViewportLayout;
+					};
+
+					ViewportsOverlayWidget = SNew(SCanvasInternal, SharedThis(this), MaximizedEntity->AsWidget());
 				}
 			}
 
