Fixing preview rendering levels where you couldn't change to a platform that has the same Feature Level as the running Editor.

diff --git a/Engine/Source/Editor/LevelEditor/Private/LevelEditor.cpp b/Engine/Source/Editor/LevelEditor/Private/LevelEditor.cpp
--- a/Engine/Source/Editor/LevelEditor/Private/LevelEditor.cpp
+++ b/Engine/Source/Editor/LevelEditor/Private/LevelEditor.cpp
@@ -1740,10 +1740,11 @@ void FLevelEditorModule::BindGlobalLevelEditorCommands()
 		const FPreviewPlatformMenuItem& Item = MenuItems[Index];
 		EShaderPlatform ShaderPlatform = ShaderFormatToLegacyShaderPlatform(Item.ShaderFormat);
 		ERHIFeatureLevel::Type FeatureLevel = GetMaxSupportedFeatureLevel(ShaderPlatform);
-		const bool IsMaxFL = FeatureLevel == GMaxRHIFeatureLevel;
 
-		bool AllowSetOnSwitch = FeatureLevel != GMaxRHIFeatureLevel;
-		FPreviewPlatformInfo PreviewFeatureLevelInfo(FeatureLevel, IsMaxFL ? NAME_None : Item.PlatformName, IsMaxFL ? NAME_None : Item.ShaderFormat, IsMaxFL ? NAME_None : Item.DeviceProfileName, AllowSetOnSwitch);
+		const bool IsDefaultActive = ShaderPlatform == GMaxRHIShaderPlatform;
+		const bool AllowPreview = !IsDefaultActive;
+
+		FPreviewPlatformInfo PreviewFeatureLevelInfo(FeatureLevel, IsDefaultActive ? NAME_None : Item.PlatformName, IsDefaultActive ? NAME_None : Item.ShaderFormat, IsDefaultActive ? NAME_None : Item.DeviceProfileName, AllowPreview);
 
 		ActionList.MapAction(
 			Commands.PreviewPlatformOverrides[Index],
