Fixes for in-viewport details crash when selecting an actor during PIE

diff --git a/Engine/Source/Editor/LevelEditor/Private/SInViewportDetails.cpp b/Engine/Source/Editor/LevelEditor/Private/SInViewportDetails.cpp
--- a/Engine/Source/Editor/LevelEditor/Private/SInViewportDetails.cpp
+++ b/Engine/Source/Editor/LevelEditor/Private/SInViewportDetails.cpp
@@ -221,6 +221,12 @@ void SInViewportDetails::GenerateWidget()
 {
 	if (AActor* SelectedActor = GetSelectedActorInEditor())
 	{
+		// Don't show this menu in PIE 
+		if (SelectedActor->GetWorld()->WorldType != EWorldType::Editor)
+		{
+			return;
+		}
+
 		FString NameString;
 		if (GEditor->GetSelectedActors()->Num() > 1)
 		{
@@ -334,11 +340,11 @@ TSharedRef<SWidget> SInViewportDetails::MakeDetailsWidget()
 				bool bShowChild = false;
 				{
 					TSharedPtr<IPropertyHandle> NodePropertyHandle = Child->CreatePropertyHandle();
-					if (NodePropertyHandle.IsValid() && NodePropertyHandle->GetProperty()->HasAllPropertyFlags(CPF_DisableEditOnInstance))
+					if (NodePropertyHandle && NodePropertyHandle->GetProperty() && NodePropertyHandle->GetProperty()->HasAllPropertyFlags(CPF_DisableEditOnInstance))
 					{
 						continue;
 					}
-					if (NodePropertyHandle.IsValid() && NodePropertyHandle->GetProperty()->GetBoolMetaData(Name_ShouldShowInViewport))
+					if (NodePropertyHandle && NodePropertyHandle->GetProperty() && NodePropertyHandle->GetProperty()->GetBoolMetaData(Name_ShouldShowInViewport))
 					{
 						bShowChild = true;
 					}
@@ -407,6 +413,10 @@ void SInViewportDetails::PostRedo(bool bSuccess)
 
 void SInViewportDetails::OnEditorSelectionChanged(UObject* Object)
 {
+	if (Object->GetWorld()->WorldType != EWorldType::Editor)
+	{
+		return;
+	}
 	TArray<UObject*> SelectedActors;
 	for (FSelectionIterator It(GEditor->GetSelectedActorIterator()); It; ++It)
 	{
