FAssetEditorModeUILayer now serves as a layer between a given asset editor and the mode toolkits, so that the mode toolkits can request UI panels and the asset editor determines where they are located in the asset editor layout. The first implementation is in the level editor, and all default mode UI has been moved to FModeToolkit.

diff --git a/Engine/Plugins/2D/Paper2D/Source/Paper2DEditor/Private/TileMapEditing/TileMapEdModeToolkit.cpp b/Engine/Plugins/2D/Paper2D/Source/Paper2DEditor/Private/TileMapEditing/TileMapEdModeToolkit.cpp
--- a/Engine/Plugins/2D/Paper2D/Source/Paper2DEditor/Private/TileMapEditing/TileMapEdModeToolkit.cpp
+++ b/Engine/Plugins/2D/Paper2D/Source/Paper2DEditor/Private/TileMapEditing/TileMapEdModeToolkit.cpp
@@ -28,13 +28,6 @@ FTileMapEdModeToolkit::FTileMapEdModeToolkit(class FEdModeTileMap* InOwningMode)
 	TileMapEditor = InOwningMode;
 }
 
-void FTileMapEdModeToolkit::RegisterTabSpawners(const TSharedRef<class FTabManager>& TabManager)
-{
-}
-
-void FTileMapEdModeToolkit::UnregisterTabSpawners(const TSharedRef<class FTabManager>& TabManager)
-{
-}
 
 FName FTileMapEdModeToolkit::GetToolkitFName() const
 {
diff --git a/Engine/Plugins/2D/Paper2D/Source/Paper2DEditor/Private/TileMapEditing/TileMapEdModeToolkit.h b/Engine/Plugins/2D/Paper2D/Source/Paper2DEditor/Private/TileMapEditing/TileMapEdModeToolkit.h
--- a/Engine/Plugins/2D/Paper2D/Source/Paper2DEditor/Private/TileMapEditing/TileMapEdModeToolkit.h
+++ b/Engine/Plugins/2D/Paper2D/Source/Paper2DEditor/Private/TileMapEditing/TileMapEdModeToolkit.h
@@ -23,8 +23,6 @@ public:
 	FTileMapEdModeToolkit(class FEdModeTileMap* InOwningMode);
 
 	// IToolkit interface
-	virtual void RegisterTabSpawners(const TSharedRef<class FTabManager>& TabManager) override;
-	virtual void UnregisterTabSpawners(const TSharedRef<class FTabManager>& TabManager) override;
 	virtual FName GetToolkitFName() const override;
 	virtual FText GetBaseToolkitName() const override;
 	virtual FText GetToolkitName() const override;
diff --git a/Engine/Plugins/Animation/ControlRig/Source/ControlRigEditor/Private/EditMode/ControlRigEditMode.cpp b/Engine/Plugins/Animation/ControlRig/Source/ControlRigEditor/Private/EditMode/ControlRigEditMode.cpp
--- a/Engine/Plugins/Animation/ControlRig/Source/ControlRigEditor/Private/EditMode/ControlRigEditMode.cpp
+++ b/Engine/Plugins/Animation/ControlRig/Source/ControlRigEditor/Private/EditMode/ControlRigEditMode.cpp
@@ -370,6 +370,7 @@ void FControlRigEditMode::Exit()
 	if (Toolkit.IsValid())
 	{
 		FToolkitManager::Get().CloseToolkit(Toolkit.ToSharedRef());
+		Toolkit.Reset();
 	}
 
 	DestroyGizmosActors();
diff --git a/Engine/Plugins/Animation/ControlRig/Source/ControlRigEditor/Private/EditMode/ControlRigEditModeToolkit.cpp b/Engine/Plugins/Animation/ControlRig/Source/ControlRigEditor/Private/EditMode/ControlRigEditModeToolkit.cpp
--- a/Engine/Plugins/Animation/ControlRig/Source/ControlRigEditor/Private/EditMode/ControlRigEditModeToolkit.cpp
+++ b/Engine/Plugins/Animation/ControlRig/Source/ControlRigEditor/Private/EditMode/ControlRigEditModeToolkit.cpp
@@ -80,15 +80,5 @@ FText FControlRigEditModeToolkit::GetActiveToolMessage() const
 	return ModeTools->GetActiveToolMessage();
 }
 
-//todo may need these later
-void FControlRigEditModeToolkit::RegisterTabSpawners(const TSharedRef<class FTabManager>& TabManager)
-{
-
-}
-
-void FControlRigEditModeToolkit::UnregisterTabSpawners(const TSharedRef<class FTabManager>& TabManager)
-{
-
-}
 
 #undef LOCTEXT_NAMESPACE
\ No newline at end of file
diff --git a/Engine/Plugins/Animation/ControlRig/Source/ControlRigEditor/Private/EditMode/ControlRigEditModeToolkit.h b/Engine/Plugins/Animation/ControlRig/Source/ControlRigEditor/Private/EditMode/ControlRigEditModeToolkit.h
--- a/Engine/Plugins/Animation/ControlRig/Source/ControlRigEditor/Private/EditMode/ControlRigEditModeToolkit.h
+++ b/Engine/Plugins/Animation/ControlRig/Source/ControlRigEditor/Private/EditMode/ControlRigEditModeToolkit.h
@@ -46,9 +46,6 @@ public:
 	virtual FText GetActiveToolMessage() const override;
 	virtual void OnToolPaletteChanged(FName PaletteName) override;
 
-	virtual void RegisterTabSpawners(const TSharedRef<class FTabManager>& TabManager) override;
-	virtual void UnregisterTabSpawners(const TSharedRef<class FTabManager>& TabManager) override;
-
 private:
 	/** The edit mode we are bound to */
 	FControlRigEditMode& EditMode;
diff --git a/Engine/Plugins/Developer/TraceDataFiltering/Source/TraceDataFiltering/Private/TraceDataFilteringModule.cpp b/Engine/Plugins/Developer/TraceDataFiltering/Source/TraceDataFiltering/Private/TraceDataFilteringModule.cpp
--- a/Engine/Plugins/Developer/TraceDataFiltering/Source/TraceDataFiltering/Private/TraceDataFilteringModule.cpp
+++ b/Engine/Plugins/Developer/TraceDataFiltering/Source/TraceDataFiltering/Private/TraceDataFilteringModule.cpp
@@ -76,7 +76,7 @@ void FTraceFilteringModule::RegisterTimingProfilerLayoutExtensions(FInsightsMajo
 {
 	TSharedRef<FWorkspaceItem> Category = InOutExtender.GetTabManager()->AddLocalWorkspaceMenuCategory(NSLOCTEXT("FTraceInsightsModule", "FilteringCategoryLabel", "Filtering"));
 
-	FInsightsMinorTabConfig& MinorTabConfig = InOutExtender.AddMinorTabConfig();
+	FMinorTabConfig& MinorTabConfig = InOutExtender.AddMinorTabConfig();
 	MinorTabConfig.TabId = FTraceFilteringModule::InsightsFilterTabName;
 	MinorTabConfig.TabLabel = NSLOCTEXT("FTraceInsightsModule", "FilteringTabTitle", "Trace Data Filtering");
 	MinorTabConfig.TabTooltip = NSLOCTEXT("FTraceInsightsModule", "FilteringTabTooltip", "Opens the Trace Data Filtering tab, allows for setting Trace Channel states");
diff --git a/Engine/Plugins/Developer/TraceSourceFiltering/Source/SourceFilteringEditor/Private/SourceFilteringEditorModule.cpp b/Engine/Plugins/Developer/TraceSourceFiltering/Source/SourceFilteringEditor/Private/SourceFilteringEditorModule.cpp
--- a/Engine/Plugins/Developer/TraceSourceFiltering/Source/SourceFilteringEditor/Private/SourceFilteringEditorModule.cpp
+++ b/Engine/Plugins/Developer/TraceSourceFiltering/Source/SourceFilteringEditor/Private/SourceFilteringEditorModule.cpp
@@ -67,7 +67,7 @@ void FSourceFilteringEditorModule::RegisterLayoutExtensions(FInsightsMajorTabExt
 
 	TSharedRef<FWorkspaceItem> Category = InOutExtender.GetTabManager()->AddLocalWorkspaceMenuCategory(NSLOCTEXT("FInsightsSourceFilteringModule", "SourceFilteringGroupName", "Filtering"));
 
-	FInsightsMinorTabConfig& MinorTabConfig = InOutExtender.AddMinorTabConfig();
+	FMinorTabConfig& MinorTabConfig = InOutExtender.AddMinorTabConfig();
 	MinorTabConfig.TabId = InsightsSourceFilteringTabName;
 	MinorTabConfig.TabLabel = NSLOCTEXT("SourceFilteringEditorModule", "SourceFilteringTab", "Trace Source Filtering");
 	MinorTabConfig.TabTooltip = NSLOCTEXT("SourceFilteringEditorModule", "SourceFilteringTabTooltip", "Opens the Trace Source Filtering tab, allows for filtering UWorld and AActor instances to not output Trace data");
diff --git a/Engine/Plugins/Editor/GeometryMode/Source/GeometryMode/Private/GeometryMode.cpp b/Engine/Plugins/Editor/GeometryMode/Source/GeometryMode/Private/GeometryMode.cpp
--- a/Engine/Plugins/Editor/GeometryMode/Source/GeometryMode/Private/GeometryMode.cpp
+++ b/Engine/Plugins/Editor/GeometryMode/Source/GeometryMode/Private/GeometryMode.cpp
@@ -249,17 +249,6 @@ TSharedRef<SButton> SGeometryModeControls::CreateSingleModifierButton(UGeomModif
 	return Widget;
 }
 
-
-void FGeometryModeToolkit::RegisterTabSpawners(const TSharedRef<class FTabManager>& TabManager)
-{
-
-}
-
-void FGeometryModeToolkit::UnregisterTabSpawners(const TSharedRef<class FTabManager>& TabManager)
-{
-
-}
-
 void FGeometryModeToolkit::Init(const TSharedPtr< class IToolkitHost >& InitToolkitHost)
 {
 	GeomWidget = SNew(SGeometryModeControls, SharedThis(this));
diff --git a/Engine/Plugins/Editor/GeometryMode/Source/GeometryMode/Private/GeometryMode.h b/Engine/Plugins/Editor/GeometryMode/Source/GeometryMode/Private/GeometryMode.h
--- a/Engine/Plugins/Editor/GeometryMode/Source/GeometryMode/Private/GeometryMode.h
+++ b/Engine/Plugins/Editor/GeometryMode/Source/GeometryMode/Private/GeometryMode.h
@@ -31,8 +31,6 @@ class FGeometryModeToolkit : public FModeToolkit
 	friend class SGeometryModeControls;
 
 public:
-	virtual void RegisterTabSpawners(const TSharedRef<class FTabManager>& TabManager) override;
-	virtual void UnregisterTabSpawners(const TSharedRef<class FTabManager>& TabManager) override;
 
 	/** Initializes the geometry mode toolkit */
 	virtual void Init(const TSharedPtr< class IToolkitHost >& InitToolkitHost) override;
diff --git a/Engine/Plugins/Editor/GeometryMode/Source/TextureAlignMode/Private/TextureAlignMode.cpp b/Engine/Plugins/Editor/GeometryMode/Source/TextureAlignMode/Private/TextureAlignMode.cpp
--- a/Engine/Plugins/Editor/GeometryMode/Source/TextureAlignMode/Private/TextureAlignMode.cpp
+++ b/Engine/Plugins/Editor/GeometryMode/Source/TextureAlignMode/Private/TextureAlignMode.cpp
@@ -8,15 +8,6 @@
 
 #define LOCTEXT_NAMESPACE "TextureAlignMode"
 
-void FTextureAlignMode::RegisterTabSpawners(const TSharedRef<class FTabManager>& TabManager)
-{
-
-}
-
-void FTextureAlignMode::UnregisterTabSpawners(const TSharedRef<class FTabManager>& TabManager)
-{
-
-}
 
 FName FTextureAlignMode::GetToolkitFName() const
 {
diff --git a/Engine/Plugins/Editor/GeometryMode/Source/TextureAlignMode/Public/TextureAlignMode.h b/Engine/Plugins/Editor/GeometryMode/Source/TextureAlignMode/Public/TextureAlignMode.h
--- a/Engine/Plugins/Editor/GeometryMode/Source/TextureAlignMode/Public/TextureAlignMode.h
+++ b/Engine/Plugins/Editor/GeometryMode/Source/TextureAlignMode/Public/TextureAlignMode.h
@@ -11,9 +11,6 @@
 class FTextureAlignMode : public FModeToolkit
 {
 public:
-	virtual void RegisterTabSpawners(const TSharedRef<class FTabManager>& TabManager) override;
-	virtual void UnregisterTabSpawners(const TSharedRef<class FTabManager>& TabManager) override;
-
 	/** IToolkit interface */
 	virtual FName GetToolkitFName() const override;
 	virtual FText GetBaseToolkitName() const override;
diff --git a/Engine/Plugins/Editor/ModelingToolsEditorMode/Source/ModelingToolsEditorMode/Private/ModelingToolsEditorModeToolkit.cpp b/Engine/Plugins/Editor/ModelingToolsEditorMode/Source/ModelingToolsEditorMode/Private/ModelingToolsEditorModeToolkit.cpp
--- a/Engine/Plugins/Editor/ModelingToolsEditorMode/Source/ModelingToolsEditorMode/Private/ModelingToolsEditorModeToolkit.cpp
+++ b/Engine/Plugins/Editor/ModelingToolsEditorMode/Source/ModelingToolsEditorMode/Private/ModelingToolsEditorModeToolkit.cpp
@@ -14,6 +14,7 @@
 #include "IDetailRootObjectCustomization.h"
 #include "ISettingsModule.h"
 #include "EditorModeManager.h"
+#include "Toolkits/AssetEditorModeUILayer.h"
 
 #include "Widgets/Input/SButton.h"
 #include "Widgets/Text/STextBlock.h"
@@ -408,16 +409,22 @@ void FModelingToolsEditorModeToolkit::PostNotification(const FText& Message)
 
 	ActiveToolMessage = Message;
 
-	const FName LevelEditorStatusBarName = "LevelEditor.StatusBar";
-	ActiveToolMessageHandle = GEditor->GetEditorSubsystem<UStatusBarSubsystem>()->PushStatusBarMessage(LevelEditorStatusBarName, ActiveToolMessage);
+	if (ModeUILayer.IsValid())
+	{
+		TSharedPtr<FAssetEditorModeUILayer> ModeUILayerPtr = ModeUILayer.Pin();
+		ActiveToolMessageHandle = GEditor->GetEditorSubsystem<UStatusBarSubsystem>()->PushStatusBarMessage(ModeUILayerPtr->GetStatusBarName(), ActiveToolMessage);
+	}
 }
 
 void FModelingToolsEditorModeToolkit::ClearNotification()
 {
 	ActiveToolMessage = FText::GetEmpty();
 
-	const FName LevelEditorStatusBarName = "LevelEditor.StatusBar";
-	GEditor->GetEditorSubsystem<UStatusBarSubsystem>()->PopStatusBarMessage(LevelEditorStatusBarName, ActiveToolMessageHandle);
+	if (ModeUILayer.IsValid())
+	{
+		TSharedPtr<FAssetEditorModeUILayer> ModeUILayerPtr = ModeUILayer.Pin();
+		GEditor->GetEditorSubsystem<UStatusBarSubsystem>()->PopStatusBarMessage(ModeUILayerPtr->GetStatusBarName(), ActiveToolMessageHandle);
+	}
 	ActiveToolMessageHandle.Reset();
 }
 
diff --git a/Engine/Plugins/Experimental/ChaosEditor/Source/FractureEditor/Private/FractureEditorMode.cpp b/Engine/Plugins/Experimental/ChaosEditor/Source/FractureEditor/Private/FractureEditorMode.cpp
--- a/Engine/Plugins/Experimental/ChaosEditor/Source/FractureEditor/Private/FractureEditorMode.cpp
+++ b/Engine/Plugins/Experimental/ChaosEditor/Source/FractureEditor/Private/FractureEditorMode.cpp
@@ -92,8 +92,6 @@ void UFractureEditorMode::Exit()
 	if (Toolkit.IsValid())
 	{
 		static_cast<FFractureEditorModeToolkit*>(Toolkit.Get())->Shutdown();
-		FToolkitManager::Get().CloseToolkit(Toolkit.ToSharedRef());
-		Toolkit.Reset();
 	}
 
 	FLevelEditorModule* LevelEditor = FModuleManager::GetModulePtr<FLevelEditorModule>("LevelEditor");
diff --git a/Engine/Plugins/Runtime/Steam/SteamAudio/Source/SteamAudioEditor/Private/SteamAudioEdModeToolkit.cpp b/Engine/Plugins/Runtime/Steam/SteamAudio/Source/SteamAudioEditor/Private/SteamAudioEdModeToolkit.cpp
--- a/Engine/Plugins/Runtime/Steam/SteamAudio/Source/SteamAudioEditor/Private/SteamAudioEdModeToolkit.cpp
+++ b/Engine/Plugins/Runtime/Steam/SteamAudio/Source/SteamAudioEditor/Private/SteamAudioEdModeToolkit.cpp
@@ -45,14 +45,6 @@ namespace SteamAudio
 	{
 	}
 
-	void FSteamAudioEdModeToolkit::RegisterTabSpawners(const TSharedRef<FTabManager>& TabManager)
-	{
-	}
-
-	void FSteamAudioEdModeToolkit::UnregisterTabSpawners(const TSharedRef<FTabManager>& TabManager)
-	{
-	}
-
 	void FSteamAudioEdModeToolkit::Init(const TSharedPtr<IToolkitHost>& InitToolkitHost)
 	{
 		SAssignNew(ToolkitWidget, SScrollBox)
diff --git a/Engine/Plugins/Runtime/Steam/SteamAudio/Source/SteamAudioEditor/Private/SteamAudioEdModeToolkit.h b/Engine/Plugins/Runtime/Steam/SteamAudio/Source/SteamAudioEditor/Private/SteamAudioEdModeToolkit.h
--- a/Engine/Plugins/Runtime/Steam/SteamAudio/Source/SteamAudioEditor/Private/SteamAudioEdModeToolkit.h
+++ b/Engine/Plugins/Runtime/Steam/SteamAudio/Source/SteamAudioEditor/Private/SteamAudioEdModeToolkit.h
@@ -21,9 +21,6 @@ namespace SteamAudio
 		FSteamAudioEdModeToolkit();
 		~FSteamAudioEdModeToolkit();
 
-		virtual void RegisterTabSpawners(const TSharedRef<FTabManager>& TabManager) override;
-		virtual void UnregisterTabSpawners(const TSharedRef<FTabManager>& TabManager) override;
-
 		/** Initializes the geometry mode toolkit */
 		virtual void Init(const TSharedPtr<IToolkitHost>& InitToolkitHost) override;
 
diff --git a/Engine/Plugins/Slate/SlateInsights/Source/SlateInsights/Private/SlateInsightsModule.cpp b/Engine/Plugins/Slate/SlateInsights/Source/SlateInsights/Private/SlateInsightsModule.cpp
--- a/Engine/Plugins/Slate/SlateInsights/Source/SlateInsights/Private/SlateInsightsModule.cpp
+++ b/Engine/Plugins/Slate/SlateInsights/Source/SlateInsights/Private/SlateInsightsModule.cpp
@@ -75,7 +75,7 @@ void FSlateInsightsModule::RegisterTimingProfilerLayoutExtensions(FInsightsMajor
 	//TSharedRef<FWorkspaceItem> Category = InOutExtender.GetTabManager()->AddLocalWorkspaceMenuCategory(LOCTEXT("SlateCategoryLabel", "Slate"));
 	//TSharedRef<FWorkspaceItem> Category = ;
 
-	FInsightsMinorTabConfig& MinorTabConfig = InOutExtender.AddMinorTabConfig();
+	FMinorTabConfig& MinorTabConfig = InOutExtender.AddMinorTabConfig();
 	MinorTabConfig.TabId = Private::FrameViewTab;
 	MinorTabConfig.TabLabel = LOCTEXT("SlateTabTitle", "Slate Frame View");
 	MinorTabConfig.TabTooltip = LOCTEXT("SlateTabTitleTooltip", "Opens the Slate Frame View tab, allows for diagnostics  of a Slate frame.");
diff --git a/Engine/Plugins/VirtualProduction/Takes/Source/TakeRecorder/Private/Widgets/STakePresetAssetEditor.cpp b/Engine/Plugins/VirtualProduction/Takes/Source/TakeRecorder/Private/Widgets/STakePresetAssetEditor.cpp
--- a/Engine/Plugins/VirtualProduction/Takes/Source/TakeRecorder/Private/Widgets/STakePresetAssetEditor.cpp
+++ b/Engine/Plugins/VirtualProduction/Takes/Source/TakeRecorder/Private/Widgets/STakePresetAssetEditor.cpp
@@ -28,6 +28,7 @@
 STakePresetAssetEditor::~STakePresetAssetEditor()
 {
 	FToolkitManager::Get().CloseToolkit(Toolkit.ToSharedRef());
+	Toolkit.Reset();
 }
 
 PRAGMA_DISABLE_OPTIMIZATION
diff --git a/Engine/Source/Developer/TraceInsights/Private/Insights/TaskGraphProfiler/TaskGraphProfilerManager.cpp b/Engine/Source/Developer/TraceInsights/Private/Insights/TaskGraphProfiler/TaskGraphProfilerManager.cpp
--- a/Engine/Source/Developer/TraceInsights/Private/Insights/TaskGraphProfiler/TaskGraphProfilerManager.cpp
+++ b/Engine/Source/Developer/TraceInsights/Private/Insights/TaskGraphProfiler/TaskGraphProfilerManager.cpp
@@ -18,6 +18,7 @@
 #include "Insights/ViewModels/ThreadTrackEvent.h"
 #include "Insights/Widgets/STimingProfilerWindow.h"
 #include "Insights/Widgets/STimingView.h"
+#include "Framework/Docking/TabManager.h"
 
 ////////////////////////////////////////////////////////////////////////////////////////////////////
 
@@ -208,7 +209,7 @@ void FTaskGraphProfilerManager::RegisterTimingProfilerLayoutExtensions(FInsights
 {
 	TimingTabManager = InOutExtender.GetTabManager();
 
-	FInsightsMinorTabConfig& MinorTabConfig = InOutExtender.AddMinorTabConfig();
+	FMinorTabConfig& MinorTabConfig = InOutExtender.AddMinorTabConfig();
 	MinorTabConfig.TabId = FTaskGraphProfilerTabs::TaskTableTreeViewTabID;
 	MinorTabConfig.TabLabel = LOCTEXT("TaskTableTreeViewTabTitle", "Tasks");
 	MinorTabConfig.TabTooltip = LOCTEXT("TaskTableTreeViewTabTitleTooltip", "Opens the Task Table Tree View tab, that allows Task Graph profilling.");
diff --git a/Engine/Source/Developer/TraceInsights/Private/Insights/Widgets/STimingProfilerWindow.cpp b/Engine/Source/Developer/TraceInsights/Private/Insights/Widgets/STimingProfilerWindow.cpp
--- a/Engine/Source/Developer/TraceInsights/Private/Insights/Widgets/STimingProfilerWindow.cpp
+++ b/Engine/Source/Developer/TraceInsights/Private/Insights/Widgets/STimingProfilerWindow.cpp
@@ -39,6 +39,7 @@
 #include "Insights/Widgets/STimerTreeView.h"
 #include "Insights/Widgets/STimingProfilerToolbar.h"
 #include "Insights/Widgets/STimingView.h"
+#include "Framework/Docking/TabManager.h"
 
 ////////////////////////////////////////////////////////////////////////////////////////////////////
 
@@ -436,7 +437,7 @@ void STimingProfilerWindow::Construct(const FArguments& InArgs, const TSharedRef
 	}
 
 	// Register any new minor tabs.
-	for (const FInsightsMinorTabConfig& MinorTabConfig : Extension->GetMinorTabs())
+	for (const FMinorTabConfig& MinorTabConfig : Extension->GetMinorTabs())
 	{
 		FTabSpawnerEntry& TabSpawnerEntry = TabManager->RegisterTabSpawner(MinorTabConfig.TabId, MinorTabConfig.OnSpawnTab, MinorTabConfig.CanSpawnTab);
 
diff --git a/Engine/Source/Developer/TraceInsights/Public/Insights/IUnrealInsightsModule.h b/Engine/Source/Developer/TraceInsights/Public/Insights/IUnrealInsightsModule.h
--- a/Engine/Source/Developer/TraceInsights/Public/Insights/IUnrealInsightsModule.h
+++ b/Engine/Source/Developer/TraceInsights/Public/Insights/IUnrealInsightsModule.h
@@ -56,28 +56,6 @@ struct TRACEINSIGHTS_API FTimingProfilerTabs
 
 ////////////////////////////////////////////////////////////////////////////////////////////////////
 
-/** Configuration for an Insights minor tab. This is used to augment the standard supplied tabs from plugins. */
-struct TRACEINSIGHTS_API FInsightsMinorTabConfig
-{
-	FName TabId;
-
-	FText TabLabel;
-
-	FText TabTooltip;
-
-	FSlateIcon TabIcon;
-
-	FOnSpawnTab OnSpawnTab;
-
-	FCanSpawnTab CanSpawnTab;
-
-	FOnFindTabToReuse OnFindTabToReuse;
-
-	TSharedPtr<FWorkspaceItem> WorkspaceGroup;
-};
-
-////////////////////////////////////////////////////////////////////////////////////////////////////
-
 /** Configuration for an Insights major tab */
 struct TRACEINSIGHTS_API FInsightsMajorTabConfig
 {
@@ -133,17 +111,17 @@ struct TRACEINSIGHTS_API FInsightsMajorTabExtender
 
 	TSharedPtr<FExtender>& GetMenuExtender() { return MenuExtender; }
 	FLayoutExtender& GetLayoutExtender() { return LayoutExtender; }
-	FInsightsMinorTabConfig& AddMinorTabConfig() { return MinorTabs.AddDefaulted_GetRef(); }
+	FMinorTabConfig& AddMinorTabConfig() { return MinorTabs.AddDefaulted_GetRef(); }
 	TSharedPtr<FTabManager> GetTabManager() const { return TabManager; }
-	const TArray<FInsightsMinorTabConfig>& GetMinorTabs() const { return MinorTabs; }
 	TSharedRef<FWorkspaceItem> GetWorkspaceGroup() { return WorkspaceGroup; }
+	const TArray<FMinorTabConfig>& GetMinorTabs() const { return MinorTabs; }
 
 protected:
 	/** Extender used to add to the menu for this tab */
 	TSharedPtr<FExtender> MenuExtender;
 
 	/** Any additional minor tabs to add */
-	TArray<FInsightsMinorTabConfig> MinorTabs;
+	TArray<FMinorTabConfig> MinorTabs;
 
 	/** Extender used when creating the layout for this tab */
 	FLayoutExtender LayoutExtender;
diff --git a/Engine/Source/Editor/EditorFramework/EditorFramework.Build.cs b/Engine/Source/Editor/EditorFramework/EditorFramework.Build.cs
--- a/Engine/Source/Editor/EditorFramework/EditorFramework.Build.cs
+++ b/Engine/Source/Editor/EditorFramework/EditorFramework.Build.cs
@@ -26,6 +26,7 @@ public class EditorFramework : ModuleRules
 				"RenderCore",
 				"Slate",
 				"SlateCore",
+				"WorkspaceMenuStructure"
 			}
 		);
 
diff --git a/Engine/Source/Editor/EditorFramework/Private/Toolkits/AssetEditorModeUILayer.cpp b/Engine/Source/Editor/EditorFramework/Private/Toolkits/AssetEditorModeUILayer.cpp
new file mode 100644
--- /dev/null
+++ b/Engine/Source/Editor/EditorFramework/Private/Toolkits/AssetEditorModeUILayer.cpp
@@ -0,0 +1,167 @@
+// Copyright Epic Games, Inc. All Rights Reserved.
+
+#include "Toolkits/AssetEditorModeUILayer.h"
+#include "Framework/MultiBox/MultiBoxDefs.h"
+#include "Widgets/Text/STextBlock.h"
+#include "Framework/MultiBox/MultiBoxBuilder.h"
+#include "Framework/Commands/InputBindingManager.h"
+#include "Subsystems/AssetEditorSubsystem.h"
+#include "Modules/ModuleManager.h"
+#include "Widgets/Layout/SBorder.h"
+#include "EditorStyleSet.h"
+#include "Editor/EditorPerProjectUserSettings.h"
+#include "LevelEditor.h"
+#include "LevelEditorActions.h"
+#include "LevelEditorModesActions.h"
+#include "Classes/EditorStyleSettings.h"
+#include "EdMode.h"
+#include "Widgets/Docking/SDockTab.h"
+#include "Widgets/Input/SCheckBox.h"
+#include "Widgets/Layout/SWidgetSwitcher.h"
+#include "Widgets/Layout/SUniformWrapPanel.h"
+#include "StatusBarSubsystem.h"
+#include "Widgets/Layout/SScrollBox.h"
+#include "Widgets/Layout/SExpandableArea.h"
+#include "WorkspaceMenuStructure.h"
+#include "WorkspaceMenuStructureModule.h"
+
+#define LOCTEXT_NAMESPACE "SLevelEditorToolBox"
+
+const FName FAssetEditorModeUILayer::TopLeftTabID(TEXT("TopLeftModeTab"));
+const FName FAssetEditorModeUILayer::BottomLeftTabID(TEXT("BottomLeftModeTab"));
+const FName FAssetEditorModeUILayer::TopRightTabID(TEXT("TopRightModeTab"));
+const FName FAssetEditorModeUILayer::BottomRightTabID(TEXT("BottomRightModeTab"));
+const FName FAssetEditorModeUILayer::VerticalToolbarID = TEXT("VerticalModeToolbar");
+
+FAssetEditorModeUILayer::FAssetEditorModeUILayer(const IToolkitHost* InToolkitHost)
+	: ToolkitHost(InToolkitHost)
+{
+	RequestedTabInfo.Add(VerticalToolbarID, FMinorTabConfig(VerticalToolbarID));
+	RequestedTabInfo.Add(TopLeftTabID, FMinorTabConfig(TopLeftTabID));
+	RequestedTabInfo.Add(BottomLeftTabID, FMinorTabConfig(BottomLeftTabID));
+	RequestedTabInfo.Add(TopRightTabID, FMinorTabConfig(TopRightTabID));
+	RequestedTabInfo.Add(BottomRightTabID, FMinorTabConfig(BottomRightTabID));
+}
+
+void FAssetEditorModeUILayer::OnToolkitHostingStarted(const TSharedRef<IToolkit>& Toolkit)
+{
+	GetTabManager()->UnregisterTabSpawner(TopLeftTabID);
+	GetTabManager()->UnregisterTabSpawner(BottomLeftTabID);
+	GetTabManager()->UnregisterTabSpawner(VerticalToolbarID);
+	GetTabManager()->UnregisterTabSpawner(TopRightTabID);
+	GetTabManager()->UnregisterTabSpawner(BottomRightTabID);
+
+}
+
+void FAssetEditorModeUILayer::OnToolkitHostingFinished(const TSharedRef<IToolkit>& Toolkit)
+{
+	if (HostedToolkit.IsValid() && HostedToolkit.Pin() == Toolkit)
+	{
+		for (TPair<FName, FMinorTabConfig>& TabSpawnerInfo : RequestedTabInfo)
+		{
+			TabSpawnerInfo.Value = FMinorTabConfig(TabSpawnerInfo.Key);
+		}
+		for (TPair<FName, TWeakPtr<SDockTab>>& SpawnedTab : SpawnedTabs)
+		{
+			if (SpawnedTab.Value.IsValid())
+			{
+				SpawnedTab.Value.Pin()->SetContent(SNullWidget::NullWidget);
+				SpawnedTab.Value.Pin()->RequestCloseTab();
+			}
+		}
+
+		SpawnedTabs.Empty();
+	}
+}
+
+TSharedPtr<FTabManager> FAssetEditorModeUILayer::GetTabManager()
+{
+	return ToolkitHost ? ToolkitHost->GetTabManager() : TSharedPtr<FTabManager>();
+}
+
+
+TSharedPtr<FWorkspaceItem> FAssetEditorModeUILayer::GetModeMenuCategory()
+{
+	return TSharedPtr<FWorkspaceItem>();
+}
+
+void FAssetEditorModeUILayer::RegisterModeTabSpawners()
+{
+	RegisterModeTabSpawner(TopLeftTabID);
+	RegisterModeTabSpawner(BottomLeftTabID);
+	RegisterModeTabSpawner(VerticalToolbarID);
+	RegisterModeTabSpawner(TopRightTabID);
+	RegisterModeTabSpawner(BottomRightTabID);
+
+}
+
+void FAssetEditorModeUILayer::RegisterModeTabSpawner(const FName TabID)
+{
+	TSharedRef<FWorkspaceItem> MenuGroup = GetModeMenuCategory().ToSharedRef();
+	bool bShowMenuOption = GetStoredSpawner(TabID).IsBound();
+	GetTabManager()->RegisterTabSpawner(TabID,
+		FOnSpawnTab::CreateSP(this, &FAssetEditorModeUILayer::SpawnStoredTab, TabID),
+		FCanSpawnTab::CreateSP(this, &FAssetEditorModeUILayer::CanSpawnStoredTab, TabID))
+		.SetDisplayNameAttribute(MakeAttributeSP(this, &FAssetEditorModeUILayer::GetTabSpawnerName, TabID))
+		.SetTooltipTextAttribute(MakeAttributeSP(this, &FAssetEditorModeUILayer::GetTabSpawnerTooltip, TabID))
+		.SetAutoGenerateMenuEntry(bShowMenuOption)
+		.SetGroup(MenuGroup);
+}
+
+void FAssetEditorModeUILayer::SetModePanelInfo(const FName InTabSpawnerID, const FMinorTabConfig& InTabInfo)
+{
+	RequestedTabInfo.Emplace(InTabSpawnerID, InTabInfo);
+}
+
+TMap<FName, TWeakPtr<SDockTab>> FAssetEditorModeUILayer::GetSpawnedTabs()
+{
+	return SpawnedTabs;
+}
+
+const FOnSpawnTab& FAssetEditorModeUILayer::GetStoredSpawner(const FName TabID)
+{
+	return RequestedTabInfo[TabID].OnSpawnTab;
+}
+
+TSharedRef<SDockTab> FAssetEditorModeUILayer::SpawnStoredTab(const FSpawnTabArgs& Args, const FName TabID)
+{
+	TSharedRef<SDockTab> SpawnedTab = SNew(SDockTab);
+	if (GetStoredSpawner(TabID).IsBound())
+	{
+		FOnSpawnTab& StoredTabSpawner = RequestedTabInfo[TabID].OnSpawnTab;
+		SpawnedTab = StoredTabSpawner.Execute(Args);
+	}
+	SpawnedTabs.Emplace(TabID, SpawnedTab);
+	return SpawnedTab;
+}
+
+bool FAssetEditorModeUILayer::CanSpawnStoredTab(const FSpawnTabArgs& Args, const FName TabID)
+{
+	bool bCanSpawnTab = GetStoredSpawner(TabID).IsBound();
+	if (RequestedTabInfo[TabID].CanSpawnTab.IsBound())
+	{
+		bCanSpawnTab |= RequestedTabInfo[TabID].CanSpawnTab.Execute(Args);
+	}
+	return bCanSpawnTab;
+}
+
+FText FAssetEditorModeUILayer::GetTabSpawnerName(const FName TabID) const
+{
+	if (!RequestedTabInfo[TabID].TabLabel.IsEmpty())
+	{
+		return RequestedTabInfo[TabID].TabLabel;
+	}
+	return FText::GetEmpty();
+}
+
+FText FAssetEditorModeUILayer::GetTabSpawnerTooltip(const FName TabID) const
+{
+	if (!RequestedTabInfo[TabID].TabTooltip.IsEmpty())
+	{
+		return RequestedTabInfo[TabID].TabTooltip;
+	}
+	return FText::GetEmpty();
+}
+
+#undef LOCTEXT_NAMESPACE
+
diff --git a/Engine/Source/Editor/EditorFramework/Public/Toolkits/AssetEditorModeUILayer.h b/Engine/Source/Editor/EditorFramework/Public/Toolkits/AssetEditorModeUILayer.h
new file mode 100644
--- /dev/null
+++ b/Engine/Source/Editor/EditorFramework/Public/Toolkits/AssetEditorModeUILayer.h
@@ -0,0 +1,73 @@
+// Copyright Epic Games, Inc. All Rights Reserved.
+
+
+#pragma once
+
+#include "CoreMinimal.h"
+#include "Textures/SlateIcon.h"
+#include "Layout/Visibility.h"
+#include "Widgets/DeclarativeSyntaxSupport.h"
+#include "Widgets/SCompoundWidget.h"
+#include "Editor.h"
+#include "ILevelEditor.h"
+#include "Misc/NotifyHook.h"
+#include "StatusBarSubsystem.h"
+#include "Framework/Docking/LayoutExtender.h"
+#include "Framework/Docking/TabManager.h"
+
+class FExtender;
+class SBorder;
+class IToolkit;
+class SDockTab;
+class ILevelEditor;
+
+
+class EDITORFRAMEWORK_API FAssetEditorModeUILayer : public TSharedFromThis<FAssetEditorModeUILayer>
+{
+public:
+	FAssetEditorModeUILayer(const IToolkitHost* InToolkitHost);
+	FAssetEditorModeUILayer() {};
+	virtual ~FAssetEditorModeUILayer() {};
+	/** Called by SLevelEditor to notify the toolbox about a new toolkit being hosted */
+	virtual void OnToolkitHostingStarted(const TSharedRef<IToolkit>& Toolkit);
+
+	/** Called by SLevelEditor to notify the toolbox about an existing toolkit no longer being hosted */
+	virtual void OnToolkitHostingFinished(const TSharedRef<IToolkit>& Toolkit);
+	virtual TSharedPtr<FTabManager> GetTabManager();
+	virtual TSharedPtr<FWorkspaceItem> GetModeMenuCategory();
+	virtual void SetModePanelInfo(const FName InTabSpawnerID, const FMinorTabConfig& InTabInfo);
+	virtual TMap<FName, TWeakPtr<SDockTab>> GetSpawnedTabs();
+	virtual FSimpleDelegate& ToolkitHostReadyForUI()
+	{
+		return OnToolkitHostReadyForUI;
+	};
+	virtual const FName GetStatusBarName() const
+	{
+		return NAME_None;
+	}
+public:
+	static const FName VerticalToolbarID;
+	static const FName TopLeftTabID;
+	static const FName BottomLeftTabID;
+	static const FName TopRightTabID;
+	static const FName BottomRightTabID;
+
+protected:
+	const FOnSpawnTab& GetStoredSpawner(const FName TabID);
+	void RegisterModeTabSpawners();
+	void RegisterModeTabSpawner(const FName TabID);
+	TSharedRef<SDockTab> SpawnStoredTab(const FSpawnTabArgs& Args, const FName TabID);
+	bool CanSpawnStoredTab(const FSpawnTabArgs& Args, const FName TabID);
+	FText GetTabSpawnerName(const FName TabID) const;
+	FText GetTabSpawnerTooltip(const FName TabID) const;
+	virtual void RegisterLayoutExtensions(FLayoutExtender& Extender) {};
+
+protected:
+	/** The host of the toolkits created by modes */
+	const IToolkitHost* ToolkitHost;
+	TArray<FName> ModeTabIDs;
+	TWeakPtr<IToolkit> HostedToolkit;
+	TMap<FName, FMinorTabConfig> RequestedTabInfo;
+	TMap<FName, TWeakPtr<SDockTab>> SpawnedTabs;
+	FSimpleDelegate OnToolkitHostReadyForUI;
+};
diff --git a/Engine/Source/Editor/EditorFramework/Public/Toolkits/IToolkit.h b/Engine/Source/Editor/EditorFramework/Public/Toolkits/IToolkit.h
--- a/Engine/Source/Editor/EditorFramework/Public/Toolkits/IToolkit.h
+++ b/Engine/Source/Editor/EditorFramework/Public/Toolkits/IToolkit.h
@@ -10,6 +10,7 @@ class FEdMode;
 class UEdMode;
 class IToolkitHost;
 class FEditorModeTools;
+class FAssetEditorModeUILayer;
 
 /** Some toolkits can be spawned as either standalone tools or within an existing level editing UI */
 namespace EToolkitMode
@@ -122,24 +123,6 @@ public:
 	/** Call this function to bring all of this toolkit's tabs to the foreground in their respective stacks.  Also causes the toolkit's host window to be foregrounded, too! */
 	virtual void BringToolkitToFront() = 0;
 
-	/* Gets the mode manager for this toolkit */
-	virtual FEditorModeTools& GetEditorModeManager() const = 0;
-
-	/** @returns the editor mode this toolkit is used for, or null if not relevant. */
-	virtual FEdMode* GetEditorMode() const = 0;
-
-	/** @returns the scriptable editor mode this toolkit is used for, or null if not relevant. */
-	virtual TWeakObjectPtr<UEdMode> GetScriptableEditorMode() const = 0;
-
-	/** @return the editor mode display name if there is an editor mode associated with this toolkit.  An empty FText is returned otherwise */
-	virtual FText GetEditorModeDisplayName() const = 0;
-
-	/** @return the editor mode icon if there is an editor mode associated with this toolkit.  An empty icon is returned otherwise */
-	virtual FSlateIcon GetEditorModeIcon() const = 0;
-
-	/** @returns the inline content that this toolkit returns if it is an editor mode */
-	virtual TSharedPtr<SWidget> GetInlineContent() const = 0;
-
 	/** Returns if this is a IBlueprintEditor derivation */
 	virtual bool IsBlueprintEditor() const = 0;
 
@@ -148,6 +131,15 @@ public:
 
 	/* Returns true if this is the default generic asset editor used by multiple classes */
 	virtual bool IsSimpleAssetEditor() const { return false; }
+
+	/* Gets the mode manager for this toolkit */
+	virtual FEditorModeTools& GetEditorModeManager() const = 0;
+
+	/** @returns the inline content that this toolkit returns */
+	virtual TSharedPtr<SWidget> GetInlineContent() const = 0;
+
+	virtual void SetModeUILayer(const TSharedPtr<FAssetEditorModeUILayer> InLayer) {};
 };
 
 
+
diff --git a/Engine/Source/Editor/FoliageEdit/Private/FoliageEdModeToolkit.cpp b/Engine/Source/Editor/FoliageEdit/Private/FoliageEdModeToolkit.cpp
--- a/Engine/Source/Editor/FoliageEdit/Private/FoliageEdModeToolkit.cpp
+++ b/Engine/Source/Editor/FoliageEdit/Private/FoliageEdModeToolkit.cpp
@@ -18,16 +18,6 @@ namespace
 	const TArray<FName> FoliagePaletteNames = { FoliageName };
 }
 
-void FFoliageEdModeToolkit::RegisterTabSpawners(const TSharedRef<class FTabManager>& TabManager)
-{
-
-}
-
-void FFoliageEdModeToolkit::UnregisterTabSpawners(const TSharedRef<class FTabManager>& TabManager)
-{
-
-}
-
 void FFoliageEdModeToolkit::Init(const TSharedPtr< class IToolkitHost >& InitToolkitHost)
 {
 	FoliageEdWidget = SNew(SFoliageEdit);
diff --git a/Engine/Source/Editor/FoliageEdit/Private/FoliageEdModeToolkit.h b/Engine/Source/Editor/FoliageEdit/Private/FoliageEdModeToolkit.h
--- a/Engine/Source/Editor/FoliageEdit/Private/FoliageEdModeToolkit.h
+++ b/Engine/Source/Editor/FoliageEdit/Private/FoliageEdModeToolkit.h
@@ -15,9 +15,6 @@ class SFoliageEdit;
 class FFoliageEdModeToolkit : public FModeToolkit
 {
 public:
-	virtual void RegisterTabSpawners(const TSharedRef<class FTabManager>& TabManager) override;
-	virtual void UnregisterTabSpawners(const TSharedRef<class FTabManager>& TabManager) override;
-
 	/** Initializes the foliage mode toolkit */
 	virtual void Init(const TSharedPtr< class IToolkitHost >& InitToolkitHost) override;
 
diff --git a/Engine/Source/Editor/LandscapeEditor/Private/SLandscapeEditor.cpp b/Engine/Source/Editor/LandscapeEditor/Private/SLandscapeEditor.cpp
--- a/Engine/Source/Editor/LandscapeEditor/Private/SLandscapeEditor.cpp
+++ b/Engine/Source/Editor/LandscapeEditor/Private/SLandscapeEditor.cpp
@@ -68,13 +68,7 @@ void SLandscapeAssetThumbnail::SetAsset(UObject* Asset)
 
 //////////////////////////////////////////////////////////////////////////
 
-void FLandscapeToolKit::RegisterTabSpawners(const TSharedRef<FTabManager>& TabManager)
-{
-}
 
-void FLandscapeToolKit::UnregisterTabSpawners(const TSharedRef<FTabManager>& TabManager)
-{
-}
 
 void FLandscapeToolKit::Init(const TSharedPtr<IToolkitHost>& InitToolkitHost)
 {
diff --git a/Engine/Source/Editor/LandscapeEditor/Private/SLandscapeEditor.h b/Engine/Source/Editor/LandscapeEditor/Private/SLandscapeEditor.h
--- a/Engine/Source/Editor/LandscapeEditor/Private/SLandscapeEditor.h
+++ b/Engine/Source/Editor/LandscapeEditor/Private/SLandscapeEditor.h
@@ -52,9 +52,6 @@ namespace LandscapeEditorNames
 class FLandscapeToolKit : public FModeToolkit
 {
 public:
-	virtual void RegisterTabSpawners(const TSharedRef<class FTabManager>& TabManager) override;
-	virtual void UnregisterTabSpawners(const TSharedRef<class FTabManager>& TabManager) override;
-
 	/** Initializes the geometry mode toolkit */
 	virtual void Init(const TSharedPtr< class IToolkitHost >& InitToolkitHost) override;
 
diff --git a/Engine/Source/Editor/LevelEditor/Private/SLevelEditor.cpp b/Engine/Source/Editor/LevelEditor/Private/SLevelEditor.cpp
--- a/Engine/Source/Editor/LevelEditor/Private/SLevelEditor.cpp
+++ b/Engine/Source/Editor/LevelEditor/Private/SLevelEditor.cpp
@@ -28,7 +28,6 @@
 #include "LevelEditorContextMenu.h"
 #include "LevelEditorToolBar.h"
 #include "SLevelEditorToolBox.h"
-#include "SLevelEditorModeContent.h"
 #include "SLevelEditorBuildAndSubmit.h"
 #include "Kismet2/DebuggerCommands.h"
 #include "SceneOutlinerPublicTypes.h"
@@ -200,7 +199,6 @@ void SLevelEditor::Construct( const SLevelEditor::FArguments& InArgs)
 			AssetEditorSubsystem->OnEditorModesChanged().AddRaw(this, &SLevelEditor::EditorModeCommandsChanged);
 		}
 	}
-	GetEditorModeManager().OnEditorModeIDChanged().AddSP(this, &SLevelEditor::OnEditorModeIdChanged);
 
 	// @todo This is a hack to get this working for now. This won't work with multiple worlds
 	if (GEditor != nullptr)
@@ -235,6 +233,8 @@ void SLevelEditor::Initialize( const TSharedRef<SDockTab>& OwnerTab, const TShar
 {
 	SelectedElements = NewObject<UTypedElementSelectionSet>(GetTransientPackage(), NAME_None, RF_Transactional);
 	SelectedElements->AddToRoot();
+	
+	ModeUILayer = MakeShareable(new FLevelEditorModeUILayer(this));
 
 	// Register the level editor specific selection behavior
 	{
@@ -559,56 +559,15 @@ void SLevelEditor::OnToolkitHostingStarted( const TSharedRef< class IToolkit >&
 	//   Otherwise, it's going to be a huge cluster trying to distinguish tabs for different assets of the same type
 	//   of editor
 	
-	TSharedPtr<FTabManager> LevelEditorTabManager = GetTabManager();
-
 	HostedToolkits.Add( Toolkit );
-
-	Toolkit->RegisterTabSpawners( LevelEditorTabManager.ToSharedRef() );
-
-	// @todo toolkit minor: We should clean out old invalid array entries from time to time
-
-	// Tell all of the toolkit area widgets about the new toolkit
-	for( auto ToolBoxIt = ToolBoxTabs.CreateIterator(); ToolBoxIt; ++ToolBoxIt )
-	{
-		if( ToolBoxIt->IsValid() )
-		{
-			ToolBoxIt->Pin()->OnToolkitHostingStarted( Toolkit );
-		}
-	}
-
-	// Tell all of the toolkit area widgets about the new toolkit
-	for( auto ToolBoxIt = ModesTabs.CreateIterator(); ToolBoxIt; ++ToolBoxIt )
-	{
-		if( ToolBoxIt->IsValid() )
-		{
-			ToolBoxIt->Pin()->OnToolkitHostingStarted( Toolkit );
-		}
-	}
+	ModeUILayer->OnToolkitHostingStarted(Toolkit);
 }
 
 void SLevelEditor::OnToolkitHostingFinished( const TSharedRef< class IToolkit >& Toolkit )
 {
 	TSharedPtr<FTabManager> LevelEditorTabManager = GetTabManager();
 
-	Toolkit->UnregisterTabSpawners(LevelEditorTabManager.ToSharedRef());
-
-	// Tell all of the toolkit area widgets that our toolkit was removed
-	for( auto ToolBoxIt = ToolBoxTabs.CreateIterator(); ToolBoxIt; ++ToolBoxIt )
-	{
-		if( ToolBoxIt->IsValid() )
-		{
-			ToolBoxIt->Pin()->OnToolkitHostingFinished( Toolkit );
-		}
-	}
-
-	// Tell all of the toolkit area widgets that our toolkit was removed
-	for( auto ToolBoxIt = ModesTabs.CreateIterator(); ToolBoxIt; ++ToolBoxIt )
-	{
-		if( ToolBoxIt->IsValid() )
-		{
-			ToolBoxIt->Pin()->OnToolkitHostingFinished( Toolkit );
-		}
-	}
+	ModeUILayer->OnToolkitHostingFinished(Toolkit);
 
 	HostedToolkits.Remove( Toolkit );
 
@@ -732,35 +691,12 @@ TSharedRef<SDockTab> SLevelEditor::SpawnLevelEditorTab( const FSpawnTabArgs& Arg
 	{
 		return this->BuildViewportTab( NSLOCTEXT("LevelViewportTypes", "LevelEditorViewport_Clone3", "Viewport 4"), TEXT("Viewport 4"), InitializationPayload );
 	}
-	else if (TabIdentifier == FEditorModeTools::EditorModeToolbarTabName)
-	{
-		return GetEditorModeManager().MakeModeToolbarTab();
-	}
 	else if( TabIdentifier == TEXT("LevelEditorSelectionDetails") || TabIdentifier == TEXT("LevelEditorSelectionDetails2") || TabIdentifier == TEXT("LevelEditorSelectionDetails3") || TabIdentifier == TEXT("LevelEditorSelectionDetails4") )
 	{
 		TSharedRef<SDockTab> DetailsPanel = SummonDetailsPanel( TabIdentifier );
 		GUnrealEd->UpdateFloatingPropertyWindows();
 		return DetailsPanel;
 	}
-	else if( TabIdentifier == LevelEditorTabIds::LevelEditorToolBox )
-	{
-		TSharedRef<SLevelEditorToolBox> NewToolBox = StaticCastSharedRef<SLevelEditorToolBox>( CreateToolBox() );
-
-		TSharedRef<SDockTab> DockTab = SNew( SDockTab )
-			.Label( NSLOCTEXT( "LevelEditor", "ToolsTabTitle", "Toolbox" ) )
-			.OnTabClosed(this, &SLevelEditor::OnToolboxTabClosed)
-			[
-				SNew( SBox )
-				.AddMetaData<FTutorialMetaData>(FTutorialMetaData(TEXT("ToolsPanel"), TEXT("LevelEditorToolBox")))
-				[
-					NewToolBox
-				]
-			];
-
-		NewToolBox->SetParentTab(DockTab);
-
-		return DockTab;
-	}
 	else if (TabIdentifier == LevelEditorTabIds::PlacementBrowser)
 	{
 		TSharedRef<SDockTab> DockTab = SNew(SDockTab)
@@ -983,26 +919,14 @@ TSharedRef<SDockTab> SLevelEditor::SpawnLevelEditorTab( const FSpawnTabArgs& Arg
 	return SNew(SDockTab);
 }
 
-bool SLevelEditor::CanSpawnEditorModeToolbarTab(const FSpawnTabArgs& Args) const
-{
-	return GetEditorModeManager().ShouldShowModeToolbar();
-}
 
-bool SLevelEditor::CanSpawnEditorModeToolboxTab(const FSpawnTabArgs& Args) const
+bool SLevelEditor::CanSpawnLevelEditorTab(const FSpawnTabArgs& Args, FName TabIdentifier)
 {
-	return HasAnyHostedEditorModeToolkit();
-}
-
-bool SLevelEditor::HasAnyHostedEditorModeToolkit() const
-{
-	for (TSharedPtr<IToolkit> Toolkit : HostedToolkits)
+	if (TabIdentifier == LevelEditorTabIds::Sequencer && !IsModeActive("EM_SequencerMode"))
 	{
-		if (Toolkit->GetScriptableEditorMode().IsValid() || Toolkit->GetEditorMode())
-		{
-			return true;
-		}
+		return false;
 	}
-	return false;
+	return true;
 }
 
 TSharedPtr<SDockTab> SLevelEditor::TryInvokeTab( FName TabID )
@@ -1255,27 +1179,6 @@ TSharedRef<SWidget> SLevelEditor::RestoreContentArea( const TSharedRef<SDockTab>
 				.SetIcon( DetailsIcon );
 		}
 
-		{
-
-			FCanSpawnTab CanSpawnTabDelegate = FCanSpawnTab::CreateSP(this, &SLevelEditor::CanSpawnEditorModeToolboxTab);
-			const FSlateIcon ToolsIcon(FEditorStyle::GetStyleSetName(), "LevelEditor.Tabs.Modes");
-			LevelEditorTabManager->RegisterTabSpawner(LevelEditorTabIds::LevelEditorToolBox, FOnSpawnTab::CreateSP<SLevelEditor, FName, FString>(this, &SLevelEditor::SpawnLevelEditorTab, LevelEditorTabIds::LevelEditorToolBox, FString()), CanSpawnTabDelegate)
-				.SetDisplayName(NSLOCTEXT("LevelEditorTabs", "LevelEditorModesToolboxTab", "Active Mode Toolbox"))
-				.SetTooltipText(NSLOCTEXT("LevelEditorTabs", "LevelEditorModesToolboxTabTooltipText", "Open the Modes tab, which contains the active editor mode's settings."))
-				.SetGroup(MenuStructure.GetLevelEditorModesCategory())
-				.SetIcon(ToolsIcon);
-		}
-
-		{
-			const FSlateIcon ToolsIcon(FEditorStyle::GetStyleSetName(), "LevelEditor.Tabs.PlacementBrowser");
-			FCanSpawnTab CanSpawnTabDelegate = FCanSpawnTab::CreateSP(this, &SLevelEditor::CanSpawnEditorModeToolbarTab);
-			LevelEditorTabManager->RegisterTabSpawner(FEditorModeTools::EditorModeToolbarTabName, FOnSpawnTab::CreateSP<SLevelEditor, FName, FString>(this, &SLevelEditor::SpawnLevelEditorTab, FEditorModeTools::EditorModeToolbarTabName, FString()), CanSpawnTabDelegate)
-				.SetDisplayName(NSLOCTEXT("LevelEditorTabs", "LevelEditorModesToolbarTab", "Active Mode Toolbar"))
-				.SetTooltipText(NSLOCTEXT("LevelEditorTabs", "LevelEditorModesToolbarTabTooltipText", "Opens a toolbar for the active editor mode"))
-				.SetGroup(MenuStructure.GetLevelEditorModesCategory())
-				.SetIcon(ToolsIcon);
-		}
-
 		{
 			const FSlateIcon ToolsIcon(FEditorStyle::GetStyleSetName(), "LevelEditor.Tabs.PlacementBrowser");
 			LevelEditorTabManager->RegisterTabSpawner(LevelEditorTabIds::PlacementBrowser, FOnSpawnTab::CreateSP<SLevelEditor, FName, FString>(this, &SLevelEditor::SpawnLevelEditorTab, LevelEditorTabIds::PlacementBrowser, FString()))
@@ -1358,7 +1261,9 @@ TSharedRef<SWidget> SLevelEditor::RestoreContentArea( const TSharedRef<SDockTab>
 
 		{
 			// @todo remove when world-centric mode is added
-			LevelEditorTabManager->RegisterTabSpawner(LevelEditorTabIds::Sequencer, FOnSpawnTab::CreateSP<SLevelEditor, FName, FString>(this, &SLevelEditor::SpawnLevelEditorTab, LevelEditorTabIds::Sequencer, FString()) )
+			LevelEditorTabManager->RegisterTabSpawner(LevelEditorTabIds::Sequencer, 
+				FOnSpawnTab::CreateSP<SLevelEditor, FName, FString>(this, &SLevelEditor::SpawnLevelEditorTab, LevelEditorTabIds::Sequencer, FString()),
+				FCanSpawnTab::CreateSP<SLevelEditor, FName>(this, &SLevelEditor::CanSpawnLevelEditorTab, LevelEditorTabIds::Sequencer))
 				.SetDisplayName(NSLOCTEXT("LevelEditorTabs", "Sequencer", "Sequencer"))
 				.SetGroup( MenuStructure.GetLevelEditorCinematicsCategory() )
 				.SetIcon( FSlateIcon(FAppStyle::GetAppStyleSetName(), "LevelEditor.Tabs.Cinematics") );
@@ -1397,12 +1302,12 @@ TSharedRef<SWidget> SLevelEditor::RestoreContentArea( const TSharedRef<SDockTab>
 	}
 
 
-	// IMPORTANT: If you want to change the default value of "LevelEditor_Layout_v1.1" or "UnrealEd_Layout_v1.4" (even if you only change their version numbers), these are the steps to follow:
+	// IMPORTANT: If you want to change the default value of "LevelEditor_Layout_v1.8" or "UnrealEd_Layout_v1.4" (even if you only change their version numbers), these are the steps to follow:
 	// 1. Check out Engine\Config\Layouts\DefaultLayout.ini in Perforce.
 	// 2. Change the code below as you wish and compile the code.
 	// 3. (Optional:) Save your current layout so you can load it later.
 	// 4. Close the editor.
-	// 5. Manually remove Engine\Saved\Config\Windows\EditorLayout.ini
+	// 5. Manually remove Engine\Saved\Config\WindowsEditor\EditorLayout.ini and Users\[UserName]\AppData\Local\UnrealEngine\Editor\EditorLayout.json
 	// 6. Open the Editor, which will auto-regenerate a default EditorLayout.ini that uses your new code below.
 	// 7. "Window" --> "Save Layout" --> "Save Layout As..."
 	//     - Name: Default Editor Layout
@@ -1414,7 +1319,7 @@ TSharedRef<SWidget> SLevelEditor::RestoreContentArea( const TSharedRef<SDockTab>
 	//     9.3. Etc
 	// 10. Push the new "DefaultLayout.ini" together with your new code.
 	// 11. Also update these instructions if you change the version number (e.g., from "UnrealEd_Layout_v1.6" to "UnrealEd_Layout_v1.7").
-	const FName LayoutName = TEXT("LevelEditor_Layout_v1.7");
+	const FName LayoutName = TEXT("LevelEditor_Layout_v1.8");
 	const TSharedRef<FTabManager::FLayout> DefaultLayout = FLayoutSaveRestore::LoadFromConfig(GEditorLayoutIni,
 		FTabManager::NewLayout( LayoutName )
 		->AddArea
@@ -1435,16 +1340,27 @@ TSharedRef<SWidget> SLevelEditor::RestoreContentArea( const TSharedRef<SDockTab>
 					->Split
 					(
 						FTabManager::NewStack()
-						->SetSizeCoefficient( .18f )
+						->SetSizeCoefficient( .15f )
 						->SetHideTabWell(true)
-						->AddTab(FEditorModeTools::EditorModeToolbarTabName, ETabState::ClosedTab)
-					)
+						->SetExtensionId("VerticalToolbar")
+ 					)
 					->Split
 					(
-						FTabManager::NewStack()
-						->SetSizeCoefficient( 0.3f )
-						->AddTab(LevelEditorTabIds::PlacementBrowser, ETabState::ClosedTab)
-						->AddTab(LevelEditorTabIds::LevelEditorToolBox, ETabState::ClosedTab)
+						FTabManager::NewSplitter()
+						->SetSizeCoefficient(.3f)
+						->SetOrientation(Orient_Vertical)
+						->Split
+						(
+							FTabManager::NewStack()
+							->SetSizeCoefficient( 0.5f )
+							->AddTab(LevelEditorTabIds::PlacementBrowser, ETabState::ClosedTab)
+						)
+						->Split
+						(
+							FTabManager::NewStack()
+							->SetSizeCoefficient(0.5f)
+							->SetExtensionId("BottomLeftPanel")
+						)
 					)
 					->Split
 					(
@@ -1460,6 +1376,7 @@ TSharedRef<SWidget> SLevelEditor::RestoreContentArea( const TSharedRef<SDockTab>
 					->SetSizeCoefficient(.4)
 					->SetHideTabWell(false)
 					->AddTab("ContentBrowserTab1", ETabState::ClosedTab)
+					->AddTab(LevelEditorTabIds::Sequencer, ETabState::ClosedTab)
 					->AddTab(LevelEditorTabIds::OutputLog, ETabState::ClosedTab)
 				)
 			)
@@ -1501,7 +1418,7 @@ TSharedRef<SWidget> SLevelEditor::RestoreContentArea( const TSharedRef<SDockTab>
 		const FText WarningText = FText::Format(LOCTEXT("LevelEditorVersionErrorBody", "The expected Unreal Level Editor layout version is \"{0}\", while only version \"{1}\" was found. I.e., the current layout was created with a previous version of Unreal that is deprecated and no longer compatible.\n\nUnreal will continue with the default layout for its current version, the deprecated one has been removed.\n\nYou can create and save your custom layouts with \"Window\"->\"Save Layout\"->\"Save Layout As...\"."),
 			FText::FromString(LayoutName.ToString()), FText::FromString(RemovedOlderLayoutVersions[0]));
 		UE_LOG(LogSlate, Warning, TEXT("%s"), *WarningText.ToString());
-		// If user is trying to load a specific layout with "Load", also warn him with a message dialog
+		// If user is trying to load a specific layout with "Load", also warn them with a message dialog
 		if (bIsBeingRecreated)
 		{
 			const FText TextTitle = LOCTEXT("LevelEditorVersionErrorTitle", "Unreal Level Editor Layout Version Mismatch");
@@ -1557,10 +1474,6 @@ void SLevelEditor::ToggleEditorMode( FEditorModeID ModeID )
 	// *Important* - activate the mode first since FEditorModeTools::DeactivateMode will
 	// activate the default mode when the stack becomes empty, resulting in multiple active visible modes.
 	GetEditorModeManager().ActivateMode( ModeID );
-
-	// Invoke the tab regardless if if the mode was activated or already activated so that it the mode tools come into the foreground
-	LevelEditorTabManager->TryInvokeTab(LevelEditorTabIds::LevelEditorToolBox);
-
 }
 
 bool SLevelEditor::IsModeActive( FEditorModeID ModeID )
@@ -1592,28 +1505,6 @@ void SLevelEditor::EditorModeCommandsChanged()
 	RefreshEditorModeCommands();
 }
 
-void SLevelEditor::OnEditorModeIdChanged(const FEditorModeID& ModeChangedID, bool bIsEnteringMode)
-{
-	if(bIsEnteringMode)
-	{
-		FLevelEditorModule& LevelEditorModule = FModuleManager::GetModuleChecked<FLevelEditorModule>("LevelEditor");
-		TSharedPtr<FTabManager> LevelEditorTabManager = LevelEditorModule.GetLevelEditorTabManager();
-
-		if (!HasAnyHostedEditorModeToolkit())
-		{
-			TSharedPtr<SDockTab> ToolboxTab = LevelEditorTabManager->FindExistingLiveTab(LevelEditorTabIds::LevelEditorToolBox);
-			if (ToolboxTab.IsValid())
-			{
-				ToolboxTab->RequestCloseTab();
-			}
-		}
-		else
-		{
-			LevelEditorTabManager->TryInvokeTab(LevelEditorTabIds::LevelEditorToolBox);
-		}
-	}
-}
-
 void SLevelEditor::RefreshEditorModeCommands()
 {
 	FLevelEditorModule& LevelEditorModule = FModuleManager::GetModuleChecked<FLevelEditorModule>( "LevelEditor" );
@@ -1653,15 +1544,6 @@ void SLevelEditor::RefreshEditorModeCommands()
 			CommandIndex++;
 		}
 	}
-
-	for( const auto& ToolBoxTab : ToolBoxTabs )
-	{
-		auto Tab = ToolBoxTab.Pin();
-		if( Tab.IsValid() )
-		{
-			Tab->OnEditorModeCommandsChanged();
-		}
-	}
 }
 
 FReply SLevelEditor::OnKeyDown( const FGeometry& MyGeometry, const FKeyEvent& InKeyEvent )
@@ -2015,17 +1897,6 @@ void SLevelEditor::SetActorDetailsSCSEditorUICustomization(TSharedPtr<ISCSEditor
 	}
 }
 
-TSharedRef<SWidget> SLevelEditor::CreateToolBox()
-{
-	TSharedRef<SLevelEditorToolBox> NewToolBox =
-		SNew( SLevelEditorToolBox, SharedThis( this ) )
-		.IsEnabled( FSlateApplication::Get().GetNormalExecutionAttribute() );
-
-	ToolBoxTabs.Add( NewToolBox );
-
-	return NewToolBox;
-}
-
 FEditorModeTools& SLevelEditor::GetEditorModeManager() const
 {
 	return GLevelEditorModeTools();
diff --git a/Engine/Source/Editor/LevelEditor/Private/SLevelEditor.h b/Engine/Source/Editor/LevelEditor/Private/SLevelEditor.h
--- a/Engine/Source/Editor/LevelEditor/Private/SLevelEditor.h
+++ b/Engine/Source/Editor/LevelEditor/Private/SLevelEditor.h
@@ -15,6 +15,7 @@
 #include "AssetThumbnail.h"
 #include "ILevelEditor.h"
 #include "LevelViewportTabContent.h"
+#include "SLevelEditorToolBox.h"
 
 class IAssetEditorInstance;
 class IDetailsView;
@@ -127,7 +128,6 @@ public:
 	virtual TSharedRef<SWidget> CreateActorDetails( const FName TabIdentifier ) override;
 	virtual void SetActorDetailsRootCustomization(TSharedPtr<FDetailsViewObjectFilter> InActorDetailsObjectFilter, TSharedPtr<IDetailRootObjectCustomization> InActorDetailsRootCustomization) override;
 	virtual void SetActorDetailsSCSEditorUICustomization(TSharedPtr<ISCSEditorUICustomization> InActorDetailsSCSEditorUICustomization) override;
-	virtual TSharedRef<SWidget> CreateToolBox() override;
 	virtual FEditorModeTools& GetEditorModeManager() const override;
 	virtual UTypedElementCommonActions* GetCommonActions() const override;
 	virtual FOnActiveViewportChanged& OnActiveViewportChanged() { return OnActiveViewportChangedDelegate; }
@@ -154,10 +154,7 @@ public:
 private:
 	
 	TSharedRef<SDockTab> SpawnLevelEditorTab(const FSpawnTabArgs& Args, FName TabIdentifier, FString InitializationPayload);
-	bool CanSpawnEditorModeToolbarTab(const FSpawnTabArgs& Args) const;
-	bool CanSpawnEditorModeToolboxTab(const FSpawnTabArgs& Args) const;
-	bool HasAnyHostedEditorModeToolkit() const;
-
+	bool CanSpawnLevelEditorTab(const FSpawnTabArgs& Args, FName TabIdentifier);
 	//TSharedRef<SDockTab> SpawnLevelEditorModeTab(const FSpawnTabArgs& Args, FEdMode* EditorMode);
 	TSharedRef<SDockTab> SummonDetailsPanel( FName Identifier );
 
@@ -184,9 +181,6 @@ private:
 	/** Editor mode has been added or removed, clears cached command list so it will be rebuilt */
 	void EditorModeCommandsChanged();
 
-	/** Called when a level editor mode is toggled */
-	void OnEditorModeIdChanged(const FEditorModeID& ModeChangedID, bool bIsEnteringMode);
-
 	/** Gets the tabId mapping to an editor mode */
 	static FName GetEditorModeTabId( FEditorModeID ModeID );
 
@@ -278,8 +272,6 @@ private:
 	// Weak reference to all toolbox panels this level editor has spawned.  May contain invalid entries for tabs that were closed.
 	TArray< TWeakPtr< class SLevelEditorToolBox > > ToolBoxTabs;
 
-	TArray< TWeakPtr< class SLevelEditorModeContent > > ModesTabs;
-
 	// List of all of the toolkits we're currently hosting.
 	TArray< TSharedPtr< class IToolkit > > HostedToolkits;
 
@@ -336,4 +328,7 @@ private:
 
 	/** If this flag is raised we will force refresh on next selection update. */
 	bool bNeedsRefresh : 1;
+
+	TSharedPtr<FLevelEditorModeUILayer> ModeUILayer;
+
 };
diff --git a/Engine/Source/Editor/LevelEditor/Private/SLevelEditorModeContent.cpp b/Engine/Source/Editor/LevelEditor/Private/SLevelEditorModeContent.cpp
deleted file mode 100644
diff --git a/Engine/Source/Editor/LevelEditor/Private/SLevelEditorModeContent.h b/Engine/Source/Editor/LevelEditor/Private/SLevelEditorModeContent.h
deleted file mode 100644
diff --git a/Engine/Source/Editor/LevelEditor/Private/SLevelEditorToolBox.cpp b/Engine/Source/Editor/LevelEditor/Private/SLevelEditorToolBox.cpp
--- a/Engine/Source/Editor/LevelEditor/Private/SLevelEditorToolBox.cpp
+++ b/Engine/Source/Editor/LevelEditor/Private/SLevelEditorToolBox.cpp
@@ -20,320 +20,70 @@
 #include "Widgets/Layout/SWidgetSwitcher.h"
 #include "Widgets/Layout/SUniformWrapPanel.h"
 #include "StatusBarSubsystem.h"
+#include "Widgets/Layout/SScrollBox.h"
+#include "Widgets/Layout/SExpandableArea.h"
+#include "WorkspaceMenuStructure.h"
+#include "WorkspaceMenuStructureModule.h"
+#include <Misc/Attribute.h>
 
 #define LOCTEXT_NAMESPACE "SLevelEditorToolBox"
 
-SLevelEditorToolBox::~SLevelEditorToolBox()
+FLevelEditorModeUILayer::FLevelEditorModeUILayer(const IToolkitHost* InToolkitHost)
+	: FAssetEditorModeUILayer(InToolkitHost)
 {
-	if (UObjectInitialized())
-	{
-		GetMutableDefault<UEditorPerProjectUserSettings>()->OnUserSettingChanged().RemoveAll(this);
-	}
+	FLevelEditorModule& LevelEditorModule = FModuleManager::GetModuleChecked<FLevelEditorModule>("LevelEditor");
+	LevelEditorModule.OnRegisterLayoutExtensions().AddRaw(this, &FLevelEditorModeUILayer::RegisterLayoutExtensions);
 }
 
-void SLevelEditorToolBox::Construct( const FArguments& InArgs, const TSharedRef< class ILevelEditor >& OwningLevelEditor )
-{
-	TabIcon = FEditorStyle::Get().GetBrush("LevelEditor.Tabs.Modes");
-	LevelEditor = OwningLevelEditor;
-
-	// Important: We use a raw binding here because we are releasing our binding in our destructor (where a weak pointer would be invalid)
-	// It's imperative that our delegate is removed in the destructor for the level editor module to play nicely with reloading.
-
-	GetMutableDefault<UEditorPerProjectUserSettings>()->OnUserSettingChanged().AddRaw( this, &SLevelEditorToolBox::HandleUserSettingsChange );
-
-	ChildSlot
-	[
-		SNew( SBorder )
-		.BorderImage( FEditorStyle::GetBrush( "ToolPanel.GroupBorder" ) )
-		.Padding(0.0f)
-		[
-			SNew( SVerticalBox )
-			+ SVerticalBox::Slot()
-			.AutoHeight()
-			.HAlign( HAlign_Left )
-			[
-				SAssignNew( ModeToolBarContainer, SBorder )
-				.BorderImage( FEditorStyle::GetBrush( "ToolPanel.GroupBorder" ) )
-				.Padding( FMargin(4, 0, 0, 0) )
-			]
-
-			+ SVerticalBox::Slot()
-			.FillHeight( 1.0f )
-			[
-				SNew( SVerticalBox )
-
-				+SVerticalBox::Slot()
-				.Padding(0.0, 8.0, 0.0, 0.0)
-				.AutoHeight()
-				[
-					SAssignNew(ModeToolHeader, SBorder)
-					.BorderImage( FEditorStyle::GetBrush( "ToolPanel.GroupBorder" ) )
-				]
-
-				+ SVerticalBox::Slot()
-				.FillHeight(1)
-				[
-					SAssignNew(InlineContentHolder, SBorder)
-					.BorderImage( FEditorStyle::GetBrush( "ToolPanel.GroupBorder" ) )
-					.Visibility( this, &SLevelEditorToolBox::GetInlineContentHolderVisibility )
-				]
-			]
-		]
-	];
-
-	UpdateModeLegacyToolBar();
-}
-
-void SLevelEditorToolBox::HandleUserSettingsChange( FName PropertyName )
-{
-	UpdateModeLegacyToolBar();
-}
-
-void SLevelEditorToolBox::OnEditorModeCommandsChanged()
-{
-	UpdateModeLegacyToolBar();
-}
 
-void SLevelEditorToolBox::SetParentTab(TSharedRef<SDockTab>& InDockTab)
+FLevelEditorModeUILayer::~FLevelEditorModeUILayer()
 {
-	ParentTab = InDockTab;
-	InDockTab->SetLabel(TabName);
-	InDockTab->SetTabIcon(TabIcon);
+	FLevelEditorModule& LevelEditorModule = FModuleManager::GetModuleChecked<FLevelEditorModule>("LevelEditor");
+	LevelEditorModule.OnRegisterLayoutExtensions().RemoveAll(this);
 }
 
-void SLevelEditorToolBox::UpdateModeLegacyToolBar()
+void FLevelEditorModeUILayer::OnToolkitHostingStarted(const TSharedRef<IToolkit>& Toolkit)
 {
-	FLevelEditorModule& LevelEditorModule = FModuleManager::GetModuleChecked<FLevelEditorModule>( "LevelEditor");
-	const TSharedPtr< const FUICommandList > CommandList = LevelEditorModule.GetGlobalLevelEditorActions();
-	const TSharedPtr<FExtender> ModeBarExtenders = LevelEditorModule.GetModeBarExtensibilityManager()->GetAllExtenders();
-
-	FToolBarBuilder EditorModeTools( CommandList, FMultiBoxCustomization::None, ModeBarExtenders );
+	if (!Toolkit->IsAssetEditor())
 	{
-		EditorModeTools.SetLabelVisibility( EVisibility::Collapsed );
-
-		const FLevelEditorModesCommands& Commands = LevelEditorModule.GetLevelEditorModesCommands();
-
-		for ( const FEditorModeInfo& Mode : GEditor->GetEditorSubsystem<UAssetEditorSubsystem>()->GetEditorModeInfoOrderedByPriority())
-		{
-			// If the mode isn't visible don't create a menu option for it.
-			if ( !Mode.IsVisible() )
-			{
-				continue;
-			}
-
-			FName EditorModeCommandName = FName( *( FString( "EditorMode." ) + Mode.ID.ToString() ) );
-
-			TSharedPtr<FUICommandInfo> EditorModeCommand =
-				FInputBindingManager::Get().FindCommandInContext( Commands.GetContextName(), EditorModeCommandName );
-
-			// If a command isn't yet registered for this mode, we need to register one.
-			if ( !EditorModeCommand.IsValid() )
-			{
-				continue;
-			}
-
-			const FUIAction* UIAction = EditorModeTools.GetTopCommandList()->GetActionForCommand( EditorModeCommand );
-			if ( ensure( UIAction ) )
-			{
-				EditorModeTools.AddToolBarButton( EditorModeCommand, Mode.ID, Mode.Name, Mode.Name, Mode.IconBrush, Mode.ID );// , EUserInterfaceActionType::ToggleButton );
-			}
-		}
+		FAssetEditorModeUILayer::OnToolkitHostingStarted(Toolkit);
+		HostedToolkit = Toolkit;
+		Toolkit->SetModeUILayer(SharedThis(this));
+		Toolkit->RegisterTabSpawners(ToolkitHost->GetTabManager().ToSharedRef());
+		RegisterModeTabSpawners();
+		OnToolkitHostReadyForUI.ExecuteIfBound();
 	}
 
-	ModeToolBarContainer->SetVisibility(EVisibility::Collapsed);
+}
 
-	const TArray<TSharedPtr<IToolkit>>& HostedToolkits = LevelEditor.Pin()->GetHostedToolkits();
-	for(const TSharedPtr<IToolkit>& HostedToolkitIt : HostedToolkits)
+void FLevelEditorModeUILayer::OnToolkitHostingFinished(const TSharedRef<IToolkit>& Toolkit)
+{
+	if (HostedToolkit.IsValid() && HostedToolkit.Pin() == Toolkit)
 	{
-		UpdateInlineContent(HostedToolkitIt, HostedToolkitIt->GetInlineContent());
+		FAssetEditorModeUILayer::OnToolkitHostingFinished(Toolkit);
 	}
 }
 
-EVisibility SLevelEditorToolBox::GetInlineContentHolderVisibility() const
-{
-	return InlineContentHolder->GetContent() == SNullWidget::NullWidget ? EVisibility::Collapsed : EVisibility::Visible;
-}
 
-EVisibility SLevelEditorToolBox::GetNoToolSelectedTextVisibility() const
-{
-	return InlineContentHolder->GetContent() == SNullWidget::NullWidget ? EVisibility::Visible : EVisibility::Collapsed;
-}
 
-void SLevelEditorToolBox::UpdateInlineContent(const TSharedPtr<IToolkit>& Toolkit, TSharedPtr<SWidget> InlineContent) 
+void FLevelEditorModeUILayer::RegisterLayoutExtensions(FLayoutExtender& Extender)
 {
-	static const FName LevelEditorStatusBarName = "LevelEditor.StatusBar";
-
-	if (Toolkit.IsValid())
-	{
-		if(Toolkit->GetEditorMode() || Toolkit->GetScriptableEditorMode().IsValid())
-		{
-			TabName = Toolkit->GetEditorModeDisplayName();
-			TabIcon = Toolkit->GetEditorModeIcon().GetSmallIcon();
-
-			TWeakPtr<FModeToolkit> ModeToolkit = StaticCastSharedPtr<FModeToolkit>(Toolkit);
-
-			if (ModeToolkit.Pin()->HasIntegratedToolPalettes())
-			{
-				TSharedRef< SUniformWrapPanel> PaletteTabBox = SNew(SUniformWrapPanel)
-					.SlotPadding( FMargin(1.f, 2.f) )
-					.HAlign(HAlign_Center);
-
-				// Only show if there is more than one child in the switcher
-				PaletteTabBox->SetVisibility(TAttribute<EVisibility>::Create(TAttribute<EVisibility>::FGetter::CreateLambda([PaletteTabBox]() -> EVisibility { return PaletteTabBox->GetChildren()->Num() > 1 ? EVisibility::Visible : EVisibility::Collapsed; }) ));
-
-				// Also build the toolkit here 
-				TArray<FName> PaletteNames;
-				ModeToolkit.Pin()->GetToolPaletteNames(PaletteNames);
-
-				TSharedPtr<FUICommandList> CommandList;
-				CommandList = ModeToolkit.Pin()->GetToolkitCommands();
-
-				TSharedRef< SWidgetSwitcher > PaletteSwitcher = SNew(SWidgetSwitcher)
-				.WidgetIndex_Lambda( [PaletteNames, ModeToolkit] () -> int32 { 
-					int32 FoundIndex;
-					TSharedPtr<FModeToolkit> PinnedToolkit = ModeToolkit.Pin();
-					if (PinnedToolkit.IsValid() && PaletteNames.Find(PinnedToolkit->GetCurrentPalette(), FoundIndex))
-					{
-						return FoundIndex;	
-					}
-					return 0;
-				} );
-				
-				
-				for(auto Palette : PaletteNames)
-				{
-					FName ToolbarCustomizationName = ModeToolkit.Pin()->GetEditorMode() ? ModeToolkit.Pin()->GetEditorMode()->GetModeInfo().ToolbarCustomizationName : ModeToolkit.Pin()->GetScriptableEditorMode()->GetModeInfo().ToolbarCustomizationName;
-					FUniformToolBarBuilder ModeToolbarBuilder(CommandList, FMultiBoxCustomization(ToolbarCustomizationName));
-					ModeToolbarBuilder.SetStyle(&FEditorStyle::Get(), "PaletteToolBar");
-
-					ModeToolkit.Pin()->BuildToolPalette(Palette, ModeToolbarBuilder);
-
-					TSharedRef<SWidget> PaletteWidget = ModeToolbarBuilder.MakeWidget();
-
-					PaletteTabBox->AddSlot()
-					[
-						SNew(SCheckBox)
-						.Padding(FMargin(8.f, 4.f, 8.f, 5.f))
-						.Style( FEditorStyle::Get(),  "PaletteToolBar.Tab" )
-						.OnCheckStateChanged_Lambda([ModeToolkit, Palette](const ECheckBoxState) {
-								TSharedPtr<FModeToolkit> PinnedToolkit = ModeToolkit.Pin();
-								if (PinnedToolkit.IsValid())
-								{
-									PinnedToolkit->SetCurrentPalette(Palette);
-								}
-							} 
-						)
-						// .IsChecked_Lambda( [PaletteSwitcher, PaletteWidget] () -> ECheckBoxState { return PaletteSwitcher->GetActiveWidget() == PaletteWidget ? ECheckBoxState::Checked : ECheckBoxState::Unchecked; })
-						.IsChecked_Lambda([ModeToolkit, Palette]() -> ECheckBoxState {
-							TSharedPtr<FModeToolkit> PinnedToolkit = ModeToolkit.Pin();
-							if (PinnedToolkit.IsValid() && (PinnedToolkit->GetCurrentPalette() == Palette))
-							{
-								return ECheckBoxState::Checked;
-							}
-							return ECheckBoxState::Unchecked;
-						})
-						[
-							SNew(STextBlock)
-							.TextStyle(FAppStyle::Get(), "NormalText")
-							.Text(ModeToolkit.Pin()->GetToolPaletteDisplayName(Palette))
-							.Justification(ETextJustify::Center)
-						]
-					];
-
-
-					PaletteSwitcher->AddSlot()
-					[
-						PaletteWidget
-					]; 
-				}
-
-
-				ModeToolHeader->SetContent(
-					SNew(SVerticalBox)
-
-					+SVerticalBox::Slot()
-					.Padding(8.f, 0.f, 0.f, 8.f)
-					.AutoHeight()
-					[
-						PaletteTabBox
-					]
-
-					+SVerticalBox::Slot()
-					.AutoHeight()
-					[
-						PaletteSwitcher
-					]
-				);
-
-			}
-
-			if (StatusBarMessageHandle.IsValid())
-			{
-				GEditor->GetEditorSubsystem<UStatusBarSubsystem>()->PopStatusBarMessage(LevelEditorStatusBarName, StatusBarMessageHandle);
-			}
-
-			StatusBarMessageHandle = GEditor->GetEditorSubsystem<UStatusBarSubsystem>()->PushStatusBarMessage(LevelEditorStatusBarName, MakeAttributeLambda([ModeToolkit]()
-				{
-					TSharedPtr<FModeToolkit> PinnedToolkit = ModeToolkit.Pin();
-					if (PinnedToolkit.IsValid())
-					{
-						return PinnedToolkit->GetActiveToolDisplayName();
-					}
-					else
-					{
-						return FText();
-					}
-				})
-			);
-		}
-	}
-	else
-	{
-		GEditor->GetEditorSubsystem<UStatusBarSubsystem>()->PopStatusBarMessage(LevelEditorStatusBarName, StatusBarMessageHandle);
-		StatusBarMessageHandle.Reset();
+	Extender.ExtendLayout(LevelEditorTabIds::PlacementBrowser, ELayoutExtensionPosition::Before, FTabManager::FTab(TopLeftTabID, ETabState::ClosedTab));
+	Extender.ExtendStack("BottomLeftPanel", ELayoutExtensionPosition::Before, FTabManager::FTab(BottomLeftTabID, ETabState::ClosedTab));
+	Extender.ExtendStack("VerticalToolbar", ELayoutExtensionPosition::Before, FTabManager::FTab(VerticalToolbarID, ETabState::ClosedTab));
+	Extender.ExtendLayout(LevelEditorTabIds::LevelEditorSceneOutliner, ELayoutExtensionPosition::Before, FTabManager::FTab(TopRightTabID, ETabState::ClosedTab));
+	Extender.ExtendLayout(LevelEditorTabIds::LevelEditorSelectionDetails, ELayoutExtensionPosition::Before, FTabManager::FTab(BottomRightTabID, ETabState::ClosedTab));
+}
 
-		TabName = NSLOCTEXT("LevelEditor", "ToolsTabTitle", "Toolbox");
-		TabIcon = FEditorStyle::Get().GetBrush("LevelEditor.Tabs.Modes");
 
-		ModeToolHeader->SetContent(SNullWidget::NullWidget);
-	}
 
-	TSharedPtr<SDockTab> ParentTabPinned = ParentTab.Pin();
-	if (InlineContent.IsValid() && InlineContentHolder.IsValid())
-	{
-		InlineContentHolder->SetContent(InlineContent.ToSharedRef());
-	}
-
-	if (ParentTabPinned.IsValid())
-	{
-		ParentTabPinned->SetLabel(TabName);
-		ParentTabPinned->SetTabIcon(TabIcon);
-	}
-}
 
-void SLevelEditorToolBox::OnToolkitHostingStarted(const TSharedRef<IToolkit>& Toolkit)
+TSharedPtr<FWorkspaceItem> FLevelEditorModeUILayer::GetModeMenuCategory()
 {
-	UpdateInlineContent(Toolkit, Toolkit->GetInlineContent());
+	const IWorkspaceMenuStructure& MenuStructure = WorkspaceMenu::GetMenuStructure();
+	return MenuStructure.GetLevelEditorModesCategory();
 }
 
-void SLevelEditorToolBox::OnToolkitHostingFinished(const TSharedRef<IToolkit>& Toolkit)
-{
-	bool FoundAnotherToolkit = false;
-	const TArray<TSharedPtr<IToolkit>>& HostedToolkits = LevelEditor.Pin()->GetHostedToolkits();
-	for (const TSharedPtr<IToolkit>& HostedToolkitIt : HostedToolkits)
-	{
-		if (HostedToolkitIt != Toolkit)
-		{
-			UpdateInlineContent(HostedToolkitIt, HostedToolkitIt->GetInlineContent());
-			FoundAnotherToolkit = true;
-			break;
-		}
-	}
-
-	if (!FoundAnotherToolkit)
-	{
-		UpdateInlineContent(nullptr, SNullWidget::NullWidget );
-	}
-}
 
 #undef LOCTEXT_NAMESPACE
+
diff --git a/Engine/Source/Editor/LevelEditor/Private/SLevelEditorToolBox.h b/Engine/Source/Editor/LevelEditor/Private/SLevelEditorToolBox.h
--- a/Engine/Source/Editor/LevelEditor/Private/SLevelEditorToolBox.h
+++ b/Engine/Source/Editor/LevelEditor/Private/SLevelEditorToolBox.h
@@ -4,6 +4,7 @@
 #pragma once
 
 #include "CoreMinimal.h"
+#include "Toolkits/AssetEditorModeUILayer.h"
 #include "Textures/SlateIcon.h"
 #include "Layout/Visibility.h"
 #include "Widgets/DeclarativeSyntaxSupport.h"
@@ -12,6 +13,8 @@
 #include "ILevelEditor.h"
 #include "Misc/NotifyHook.h"
 #include "StatusBarSubsystem.h"
+#include "Framework/Docking/LayoutExtender.h"
+#include "Framework/Docking/TabManager.h"
 
 class FExtender;
 class SBorder;
@@ -19,70 +22,21 @@ class IToolkit;
 class SDockTab;
 class ILevelEditor;
 
-/**
- * Tools for the level editor                   
- */
-class SLevelEditorToolBox : public SCompoundWidget, public FNotifyHook
+class FLevelEditorModeUILayer : public FAssetEditorModeUILayer
 {
 public:
-	SLATE_BEGIN_ARGS( SLevelEditorToolBox ){}
-	SLATE_END_ARGS()
-
-	~SLevelEditorToolBox();
-
-	void Construct(const FArguments& InArgs, const TSharedRef<ILevelEditor>& OwningLevelEditor);
-
-	/** Called by SLevelEditor to notify the toolbox about a new toolkit being hosted */
-	void OnToolkitHostingStarted( const TSharedRef<IToolkit>& Toolkit);
-
-	/** Called by SLevelEditor to notify the toolbox about an existing toolkit no longer being hosted */
-	void OnToolkitHostingFinished(const TSharedRef<IToolkit>& Toolkit);
-
-	/** Handles updating the mode toolbar when the registered mode commands change */
-	void OnEditorModeCommandsChanged();
-
-	/** Sets the parent tab of this toolbox */
-	void SetParentTab(TSharedRef<SDockTab>& InDockTab);
-
-private:
-	void OnContainingTabActivated(TSharedRef<SDockTab> DockTab, ETabActivationCause InActivationCause);
-	/** Gets the visibility for the SBorder showing toolbox editor mode inline content */
-	EVisibility GetInlineContentHolderVisibility() const;
-
-	/** Gets the visibility for the message suggesting the user select a tool */
-	EVisibility GetNoToolSelectedTextVisibility() const;
-
-	/** Updates the widget for showing toolbox editor mode inline content */
-	void UpdateInlineContent(const TSharedPtr<IToolkit>& Toolkit, TSharedPtr<SWidget> InlineContent);
-
-	/** Creates and sets the mode toolbar */
-	void UpdateModeLegacyToolBar();
-
-	/** Handles updating the mode toolbar when the user settings change */
-	void HandleUserSettingsChange( FName PropertyName );
-
-private:
-	/** Parent tab where this toolbox is hosted */
-	TWeakPtr<SDockTab> ParentTab;
-
-	/** Level editor that we're associated with */
-	TWeakPtr<ILevelEditor> LevelEditor;
-
-	/** Inline content area for editor modes */
-	TSharedPtr<SBorder> InlineContentHolder;
-
-	/** The container holding the mode toolbar */
-	TSharedPtr<SBorder> ModeToolBarContainer;
-
-	/** The active tool header area **/
-	TSharedPtr<SBorder> ModeToolHeader;
-
-	/** The display name that the parent tab should have as its label */
-	FText TabName;
-
-	/** The icon that should be displayed in the parent tab */
-	const FSlateBrush* TabIcon;
-
-	/** Handle to status bar messages for removing them later */
-	FStatusBarMessageHandle StatusBarMessageHandle;
+	FLevelEditorModeUILayer(const IToolkitHost* InToolkitHost);
+	FLevelEditorModeUILayer() {};
+	virtual ~FLevelEditorModeUILayer() override;
+	virtual void OnToolkitHostingStarted(const TSharedRef<IToolkit>& Toolkit) override;
+	virtual void OnToolkitHostingFinished(const TSharedRef<IToolkit>& Toolkit) override;
+
+	virtual TSharedPtr<FWorkspaceItem> GetModeMenuCategory() override;
+	virtual const FName GetStatusBarName() const override
+	{
+		static const FName LevelEditorStatusBarName = "LevelEditor.StatusBar";
+		return LevelEditorStatusBarName;
+	}
+protected:
+	virtual void RegisterLayoutExtensions(FLayoutExtender& Extender) override;
 };
diff --git a/Engine/Source/Editor/LevelEditor/Public/ILevelEditor.h b/Engine/Source/Editor/LevelEditor/Public/ILevelEditor.h
--- a/Engine/Source/Editor/LevelEditor/Public/ILevelEditor.h
+++ b/Engine/Source/Editor/LevelEditor/Public/ILevelEditor.h
@@ -85,9 +85,6 @@ public:
 	/** Sets the UI customization of the SCSEditor inside the level editor details panel. */
 	virtual void SetActorDetailsSCSEditorUICustomization(TSharedPtr<class ISCSEditorUICustomization> ActorDetailsSCSEditorUICustomization) = 0;
 
-	/** Spawns a level editor ToolBox widget (aka. "Modes") */
-	virtual TSharedRef<SWidget> CreateToolBox() = 0;
-
 	virtual TSharedPtr<ISceneOutliner> GetSceneOutliner() const = 0;
 };
 
diff --git a/Engine/Source/Editor/SkeletalMeshEditor/Private/SSkeletalMeshEditorToolbox.cpp b/Engine/Source/Editor/SkeletalMeshEditor/Private/SSkeletalMeshEditorToolbox.cpp
--- a/Engine/Source/Editor/SkeletalMeshEditor/Private/SSkeletalMeshEditorToolbox.cpp
+++ b/Engine/Source/Editor/SkeletalMeshEditor/Private/SSkeletalMeshEditorToolbox.cpp
@@ -99,26 +99,20 @@ void SSkeletalMeshEditorToolbox::UpdateInlineContent(const TSharedPtr<IToolkit>&
 		GEditor->GetEditorSubsystem<UStatusBarSubsystem>()->PopStatusBarMessage(SkeletalMeshEditorStatusBarName, StatusBarMessageHandle);
 		StatusBarMessageHandle.Reset();
 	}
-
-	if (Toolkit.IsValid())
+	TWeakPtr<FModeToolkit> ModeToolkit = StaticCastSharedPtr<FModeToolkit>(Toolkit);
+	if (ModeToolkit.IsValid())
 	{
-		TabName = Toolkit->GetEditorModeDisplayName();
-		TabIcon = Toolkit->GetEditorModeIcon().GetSmallIcon();
-
-		TWeakPtr<FModeToolkit> ModeToolkit = StaticCastSharedPtr<FModeToolkit>(Toolkit);
+		TabName = ModeToolkit.Pin()->GetEditorModeDisplayName();
+		TabIcon = ModeToolkit.Pin()->GetEditorModeIcon().GetSmallIcon();
+		const TSharedRef<FModeToolkit> ModeToolkitPinned = ModeToolkit.Pin().ToSharedRef();
 
-		if (ModeToolkit.IsValid())
-		{
-			TSharedRef<FModeToolkit> ModeToolkitPinned = ModeToolkit.Pin().ToSharedRef();
+		UpdatePalette(ModeToolkitPinned);
 
-			UpdatePalette(ModeToolkitPinned);
-
-			// Show the name of the active tool in the statusbar.
-			// FIXME: We should also be showing Ctrl/Shift/Alt LMB/RMB shortcuts.
-			StatusBarMessageHandle = GEditor->GetEditorSubsystem<UStatusBarSubsystem>()->PushStatusBarMessage(
-					SkeletalMeshEditorStatusBarName,
-					TAttribute<FText>::Create(TAttribute<FText>::FGetter::CreateSP(ModeToolkitPinned, &FModeToolkit::GetActiveToolDisplayName)));
-		}
+		// Show the name of the active tool in the statusbar.
+		// FIXME: We should also be showing Ctrl/Shift/Alt LMB/RMB shortcuts.
+		StatusBarMessageHandle = GEditor->GetEditorSubsystem<UStatusBarSubsystem>()->PushStatusBarMessage(
+			SkeletalMeshEditorStatusBarName,
+			TAttribute<FText>::Create(TAttribute<FText>::FGetter::CreateSP(ModeToolkitPinned, &FModeToolkit::GetActiveToolDisplayName)));
 	}
 	else
 	{
diff --git a/Engine/Source/Editor/UnrealEd/Private/EditorModeManager.cpp b/Engine/Source/Editor/UnrealEd/Private/EditorModeManager.cpp
--- a/Engine/Source/Editor/UnrealEd/Private/EditorModeManager.cpp
+++ b/Engine/Source/Editor/UnrealEd/Private/EditorModeManager.cpp
@@ -30,8 +30,6 @@
 #include "Subsystems/BrushEditingSubsystem.h"
 #include "Tools/UEdMode.h"
 #include "Widgets/Images/SImage.h"
-#include "Widgets/Layout/SExpandableArea.h"
-#include "Widgets/Layout/SScrollBox.h"
 #include "InputRouter.h"
 #include "InteractiveGizmoManager.h"
 #include "EdModeInteractiveToolsContext.h"
@@ -51,7 +49,7 @@
 	The master class that handles tracking of the current mode.
 ------------------------------------------------------------------------------*/
 
-const FName FEditorModeTools::EditorModeToolbarTabName = TEXT("EditorModeToolbar");
+
 
 FEditorModeTools::FEditorModeTools()
 	: PivotShown(false)
@@ -672,16 +670,6 @@ void FEditorModeTools::DeactivateScriptableModeAtIndex(int32 InIndex)
 	UEdMode* Mode = ActiveScriptableModes[InIndex];
 	ActiveScriptableModes.RemoveAt(InIndex);
 
-	// Remove the toolbar widget to remove any toolkit references that may get removed when the mode exits.
-	ActiveToolBarRows.RemoveAll(
-		[&Mode](FEdModeToolbarRow& Row)
-		{
-			return Row.ModeID == Mode->GetID();
-		}
-	);
-
-	RebuildModeToolBar();
-
 	Mode->Exit();
 
 	const bool bIsEnteringMode = false;
@@ -695,187 +683,8 @@ void FEditorModeTools::OnModeUnregistered(FEditorModeID ModeID)
 	DestroyMode(ModeID);
 }
 
-void FEditorModeTools::RebuildModeToolBar()
-{
-	// If the tab or box is not valid the toolbar has not been opened or has been closed by the user
-	TSharedPtr<SVerticalBox> ModeToolbarBoxPinned = ModeToolbarBox.Pin();
-	if (ModeToolbarTab.IsValid() && ModeToolbarBoxPinned.IsValid())
-	{
-		ModeToolbarBoxPinned->ClearChildren();
-
-		bool bExclusivePalettes = true;
-		TSharedRef<SVerticalBox> ToolBoxVBox = SNew(SVerticalBox);
-
-		TSharedRef< SUniformWrapPanel> PaletteTabBox = SNew(SUniformWrapPanel)
-			.SlotPadding(FMargin(1.f, 2.f))
-			.HAlign(HAlign_Left);
-		TSharedRef< SWidgetSwitcher > PaletteSwitcher = SNew(SWidgetSwitcher);
-
-		int32 PaletteCount = ActiveToolBarRows.Num();
-		if(PaletteCount > 0)
-		{
-			for(int32 RowIdx = 0; RowIdx < PaletteCount; ++RowIdx)
-			{
-				const FEdModeToolbarRow& Row = ActiveToolBarRows[RowIdx];
-				if (ensure(Row.ToolbarWidget.IsValid()))
-				{
-					TSharedRef<SWidget> PaletteWidget = Row.ToolbarWidget.ToSharedRef();
-					
-					TSharedPtr<FModeToolkit> RowToolkit;
-					if (FEdMode* Mode = GetActiveMode(Row.ModeID))
-					{
-						RowToolkit = Mode->GetToolkit();
-					}
-					else if (UEdMode* ScriptableMode = GetActiveScriptableMode(Row.ModeID))
-					{
-						RowToolkit = ScriptableMode->GetToolkit().Pin();
-					}
-
-					bExclusivePalettes = RowToolkit->HasExclusiveToolPalettes();
-
-					if (!bExclusivePalettes)
-					{
-						ToolBoxVBox->AddSlot()
-						.AutoHeight()
-						.Padding(FMargin(2.0, 2.0))
-						[
-							SNew(SExpandableArea)
-							.AreaTitle(Row.DisplayName)
-							.AreaTitleFont(FAppStyle::Get().GetFontStyle("NormalFont"))
-							.BorderImage(FAppStyle::Get().GetBrush("PaletteToolbar.ExpandableAreaHeader"))
-							.BodyBorderImage(FAppStyle::Get().GetBrush("PaletteToolbar.ExpandableAreaBody"))
-							.HeaderPadding(FMargin(4.f))
-							.Padding(FMargin(4.0, 0.0))
-							.BodyContent()
-							[
-								PaletteWidget	
-							]
-						];
-					}
-					else 
-					{
-						// Don't show Palette Tabs if there is only one
-						if (PaletteCount > 1)
-						{
-							PaletteTabBox->AddSlot()
-							[
-								SNew(SCheckBox)
-								.Style( FEditorStyle::Get(),  "ToolPalette.DockingTab" )
-								.OnCheckStateChanged_Lambda( [PaletteSwitcher, Row, RowToolkit] (const ECheckBoxState) { 
-										PaletteSwitcher->SetActiveWidget(Row.ToolbarWidget.ToSharedRef());
-										RowToolkit->SetCurrentPalette(Row.PaletteName);
-									} 
-								)
-								.IsChecked_Lambda( [PaletteSwitcher, PaletteWidget] () -> ECheckBoxState { return PaletteSwitcher->GetActiveWidget() == PaletteWidget ? ECheckBoxState::Checked : ECheckBoxState::Unchecked; })
-								[
-									SNew(STextBlock)
-									.Text(Row.DisplayName)
-								]
-							];
-						}
 
-						PaletteSwitcher->AddSlot()
-						[
-							PaletteWidget
-						]; 
-					}
-				}
-			}
 
-			ModeToolbarBoxPinned->AddSlot()
-			.AutoHeight()
-			[
-				SNew(SOverlay)
-				+SOverlay::Slot()
-				[
-					SNew(SImage)
-					.Image(FEditorStyle::GetBrush("ToolPalette.DockingWell"))
-				]
-
-				+SOverlay::Slot()
-				[
-					PaletteTabBox
-				]
-			];
-
-			ModeToolbarBoxPinned->AddSlot()
-			.AutoHeight()
-			.Padding(1.f)
-			[
-				SNew(SBox)
-				.HeightOverride(PaletteSwitcher->GetNumWidgets() > 0 ? 45.f : 0.f)
-				[
-					PaletteSwitcher 
-				]
-			];
-
-			ModeToolbarBoxPinned->AddSlot()
-			[
-				SNew(SScrollBox)
-
-				+SScrollBox::Slot()
-				[
-					ToolBoxVBox	
-				]
-			];
-
-			ModeToolbarPaletteSwitcher = PaletteSwitcher;
-		}
-		else
-		{
-			ModeToolbarTab.Pin()->RequestCloseTab();
-		}
-	}
-}
-
-void FEditorModeTools::SpawnOrUpdateModeToolbar()
-{
-	if (ShouldShowModeToolbar())
-	{
-		if (ModeToolbarTab.IsValid())
-		{
-			RebuildModeToolBar();
-		}
-		else if (ToolkitHost.IsValid())
-		{
-			ToolkitHost.Pin()->GetTabManager()->TryInvokeTab(EditorModeToolbarTabName);
-		}
-	}
-}
-
-void FEditorModeTools::InvokeToolPaletteTab(FEditorModeID InModeID, FName InPaletteName)
-{
-	if (!ModeToolbarPaletteSwitcher.Pin()) 
-	{
-		return;
-	}
-
-	for (auto Row: ActiveToolBarRows)
-	{
-		if (Row.ModeID == InModeID && Row.PaletteName == InPaletteName)
-		{
-			TSharedRef<SWidget> PaletteWidget = Row.ToolbarWidget.ToSharedRef();
-
-			TSharedPtr<FModeToolkit> RowToolkit;
-			if (FEdMode* Mode = GetActiveMode(InModeID))
-			{
-				RowToolkit = Mode->GetToolkit();
-			}
-			else if (UEdMode* ScriptableMode = GetActiveScriptableMode(InModeID))
-			{
-				RowToolkit = ScriptableMode->GetToolkit().Pin();
-			}
-
-			TSharedPtr<SWidget> ActiveWidget = ModeToolbarPaletteSwitcher.Pin()->GetActiveWidget();
-			if (RowToolkit && ActiveWidget.Get() != Row.ToolbarWidget.Get())
-			{
-				ModeToolbarPaletteSwitcher.Pin()->SetActiveWidget(Row.ToolbarWidget.ToSharedRef());
-				RowToolkit->OnToolPaletteChanged(Row.PaletteName);
-			}
-			break;	
-		}
-	}	
-}
 
 void FEditorModeTools::OnWorldCleanup(UWorld* InWorld, bool bSessionEnded, bool bCleanupResources)
 {
@@ -959,29 +768,7 @@ void FEditorModeTools::DestroyMode( FEditorModeID InID )
 	RecycledScriptableModes.Remove(InID);
 }
 
-TSharedRef<SDockTab> FEditorModeTools::MakeModeToolbarTab()
-{
-	TSharedRef<SDockTab> ToolbarTabRef = 	
-		SNew(SDockTab)
-		.Label(NSLOCTEXT("EditorModes", "EditorModesToolbarTitle", "Mode Toolbar"))
-		.ContentPadding(0.0f)
-		[
-			SAssignNew(ModeToolbarBox, SVerticalBox)	
-		];
-
-	ModeToolbarTab = ToolbarTabRef;
-
-	// Rebuild the toolbar with existing mode tools that may be active
-	RebuildModeToolBar();
-
-	return ToolbarTabRef;
-
-}
 
-bool FEditorModeTools::ShouldShowModeToolbar() const
-{
-	return ActiveToolBarRows.Num() > 0;
-}
 
 bool FEditorModeTools::ShouldShowModeToolbox() const
 {
@@ -1064,27 +851,6 @@ void FEditorModeTools::ActivateMode(FEditorModeID InID, bool bToggle)
 	const bool bIsEnteringMode = true;
 	BroadcastEditorModeIDChanged(ScriptableMode->GetID(), bIsEnteringMode);
 
-	// Ask the mode to build the toolbar.
-	TSharedPtr<FUICommandList> CommandList;
-	const TSharedPtr<FModeToolkit> Toolkit = ScriptableMode->GetToolkit().Pin();
-	if (Toolkit && !Toolkit->HasIntegratedToolPalettes())
-	{
-		CommandList = Toolkit->GetToolkitCommands();
-
-		// Also build the toolkit here 
-		TArray<FName> PaletteNames;
-		Toolkit->GetToolPaletteNames(PaletteNames);
-		for (auto Palette : PaletteNames)
-		{
-			FUniformToolBarBuilder ModeToolbarBuilder(CommandList, FMultiBoxCustomization(ScriptableMode->GetModeInfo().ToolbarCustomizationName), TSharedPtr<FExtender>(), false);
-			ModeToolbarBuilder.SetStyle(&FEditorStyle::Get(), "PaletteToolBar");
-			Toolkit->BuildToolPalette(Palette, ModeToolbarBuilder);
-			ActiveToolBarRows.Emplace(ScriptableMode->GetID(), Palette, Toolkit->GetToolPaletteDisplayName(Palette), ModeToolbarBuilder.MakeWidget());
-		}
-	}
-
-	SpawnOrUpdateModeToolbar();
-
 	RecycledScriptableModes.Remove(InID);
 
 	// Update the editor UI
diff --git a/Engine/Source/Editor/UnrealEd/Private/Toolkits/AssetEditorToolkit.cpp b/Engine/Source/Editor/UnrealEd/Private/Toolkits/AssetEditorToolkit.cpp
--- a/Engine/Source/Editor/UnrealEd/Private/Toolkits/AssetEditorToolkit.cpp
+++ b/Engine/Source/Editor/UnrealEd/Private/Toolkits/AssetEditorToolkit.cpp
@@ -414,26 +414,6 @@ FText FAssetEditorToolkit::GetToolTipTextForObject(const UObject* InObject)
 	return FText::FromString(ToolTipString);
 }
 
-class FEdMode* FAssetEditorToolkit::GetEditorMode() const
-{
-	return nullptr;
-}
-
-TWeakObjectPtr<UEdMode> FAssetEditorToolkit::GetScriptableEditorMode() const
-{
-	return nullptr;
-}
-
-FText FAssetEditorToolkit::GetEditorModeDisplayName() const
-{
-	return FText::GetEmpty();
-}
-
-FSlateIcon FAssetEditorToolkit::GetEditorModeIcon() const
-{
-	return FSlateIcon();
-}
-
 FEditorModeTools& FAssetEditorToolkit::GetEditorModeManager() const
 {
 	if (IsWorldCentricAssetEditor() && IsHosted())
diff --git a/Engine/Source/Editor/UnrealEd/Private/Toolkits/BaseToolkit.cpp b/Engine/Source/Editor/UnrealEd/Private/Toolkits/BaseToolkit.cpp
--- a/Engine/Source/Editor/UnrealEd/Private/Toolkits/BaseToolkit.cpp
+++ b/Engine/Source/Editor/UnrealEd/Private/Toolkits/BaseToolkit.cpp
@@ -13,7 +13,12 @@
 #include "Framework/MultiBox/MultiBoxBuilder.h"
 #include "EdMode.h"
 #include "EditorModeManager.h"
-
+#include "Widgets/Layout/SScrollBox.h"
+#include "Widgets/Layout/SExpandableArea.h"
+#include "Toolkits/AssetEditorModeUILayer.h"
+#include "Widgets/Docking/SDockTab.h"
+#include "Widgets/Input/SCheckBox.h"
+#include "Widgets/Images/SImage.h"
 #define LOCTEXT_NAMESPACE "BaseToolkit"
 
 FBaseToolkit::FBaseToolkit()
@@ -108,9 +113,6 @@ bool FBaseToolkit::IsBlueprintEditor() const
 	return false;
 }
 
-#undef LOCTEXT_NAMESPACE
-
-
 void FModeToolkit::Init(const TSharedPtr< class IToolkitHost >& InitToolkitHost)
 {
 	Init(InitToolkitHost, TWeakObjectPtr<UEdMode>());
@@ -123,7 +125,8 @@ void FModeToolkit::Init(const TSharedPtr< class IToolkitHost >& InitToolkitHost,
 	ToolkitMode = EToolkitMode::Type::Standalone;
 	ToolkitHost = InitToolkitHost;
 	OwningEditorMode = InOwningMode;
-
+	PrimaryTabInfo = FMinorTabConfig();
+	ToolbarInfo = FMinorTabConfig();
 	UInteractiveToolManager* ToolManager = GetEditorModeManager().GetInteractiveToolsContext()->ToolManager;
 	ToolManager->OnToolStarted.AddSP(this, &FModeToolkit::OnToolStarted);
 	ToolManager->OnToolEnded.AddSP(this, &FModeToolkit::OnToolEnded);
@@ -170,10 +173,36 @@ FModeToolkit::~FModeToolkit()
 		ToolManager->OnToolStarted.RemoveAll(this);
 		ToolManager->OnToolEnded.RemoveAll(this);
 	}
-
+	if (ModeToolBarContainer)
+	{
+		ModeToolBarContainer->SetContent(SNullWidget::NullWidget);
+	}
+	if (ModeToolHeader)
+	{
+		ModeToolHeader->SetContent(SNullWidget::NullWidget);
+	}
+	if (InlineContentHolder)
+	{
+		InlineContentHolder->SetContent(SNullWidget::NullWidget);
+	}
 	OwningEditorMode.Reset();
 }
 
+void FModeToolkit::SetModeUILayer(const TSharedPtr<FAssetEditorModeUILayer> InLayer)
+{
+	ModeUILayer = InLayer;
+	ModeUILayer.Pin()->ToolkitHostReadyForUI().BindSP(this, &FModeToolkit::InvokeUI);
+}
+
+void FModeToolkit::RegisterTabSpawners(const TSharedRef<FTabManager>& TabManager)
+{
+	RequestModeUITabs();
+}
+
+void FModeToolkit::UnregisterTabSpawners(const TSharedRef<FTabManager>& TabManager)
+{
+}
+
 FName FModeToolkit::GetToolkitFName() const
 {
 	return FName("EditorModeToolkit");
@@ -364,3 +393,407 @@ void FModeToolkit::SetModeSettingsObject(UObject* InSettingsObject)
 {
 	ModeDetailsView->SetObject(InSettingsObject);
 }
+
+void FModeToolkit::InvokeUI()
+{
+	if (ModeUILayer.IsValid())
+	{
+		TSharedPtr<FAssetEditorModeUILayer> ModeUILayerPtr = ModeUILayer.Pin();
+		TSharedPtr<SDockTab> CreatedTab = ModeUILayerPtr->GetTabManager()->TryInvokeTab(FAssetEditorModeUILayer::TopLeftTabID);
+		UpdatePrimaryModePanel();
+		TSharedPtr<FUICommandList> CommandList;
+		if (!HasIntegratedToolPalettes())
+		{
+			if (GetScriptableEditorMode().IsValid())
+			{
+				UEdMode* ScriptableMode = GetScriptableEditorMode().Get();
+				CommandList = GetToolkitCommands();
+
+				// Also build the toolkit here 
+				TArray<FName> PaletteNames;
+				GetToolPaletteNames(PaletteNames);
+				for (auto Palette : PaletteNames)
+				{
+					FUniformToolBarBuilder ModeToolbarBuilder(CommandList, FMultiBoxCustomization(ScriptableMode->GetModeInfo().ToolbarCustomizationName), TSharedPtr<FExtender>(), false);
+					ModeToolbarBuilder.SetStyle(&FEditorStyle::Get(), "PaletteToolBar");
+					BuildToolPalette(Palette, ModeToolbarBuilder);
+					ActiveToolBarRows.Emplace(ScriptableMode->GetID(), Palette, GetToolPaletteDisplayName(Palette), ModeToolbarBuilder.MakeWidget());
+				}
+			}
+			TSharedPtr<SDockTab> CreatedToolbarTab = ModeUILayerPtr->GetTabManager()->TryInvokeTab(FAssetEditorModeUILayer::VerticalToolbarID);
+			ModeToolbarTab = CreatedToolbarTab;
+
+		}
+	}
+}
+
+TSharedRef<SDockTab> FModeToolkit::CreatePrimaryModePanel(const FSpawnTabArgs& Args)
+{
+	TSharedRef<SWidget> TabContent = SNew(SBorder)
+		.BorderImage(FEditorStyle::GetBrush("ToolPanel.GroupBorder"))
+		.Padding(0.0f)
+		[
+			SNew(SVerticalBox)
+			+ SVerticalBox::Slot()
+			.AutoHeight()
+			.HAlign(HAlign_Left)
+			[
+				SAssignNew(ModeToolBarContainer, SBorder)
+				.BorderImage(FEditorStyle::GetBrush("ToolPanel.GroupBorder"))
+				.Padding(FMargin(4, 0, 0, 0))
+			]
+
+			+ SVerticalBox::Slot()
+			.FillHeight(1.0f)
+			[
+				SNew(SVerticalBox)
+
+				+ SVerticalBox::Slot()
+				.Padding(0.0, 8.0, 0.0, 0.0)
+				.AutoHeight()
+				[
+					SAssignNew(ModeToolHeader, SBorder)
+					.BorderImage(FEditorStyle::GetBrush("ToolPanel.GroupBorder"))
+				]
+
+				+ SVerticalBox::Slot()
+				.FillHeight(1)
+				[
+					SAssignNew(InlineContentHolder, SBorder)
+					.BorderImage(FEditorStyle::GetBrush("ToolPanel.GroupBorder"))
+					.Visibility(this, &FModeToolkit::GetInlineContentHolderVisibility)
+				]
+		]
+		];
+
+
+	TSharedPtr<SDockTab> CreatedTab = SNew(SDockTab)
+		[
+			TabContent
+		];
+
+	PrimaryTab = CreatedTab;	
+	UpdatePrimaryModePanel();
+	return CreatedTab.ToSharedRef();
+
+}
+
+
+
+void FModeToolkit::UpdatePrimaryModePanel()
+{
+	if (GetEditorMode() || GetScriptableEditorMode().IsValid())
+	{
+		const FText TabName = GetEditorModeDisplayName();
+		const FSlateBrush* TabIcon = GetEditorModeIcon().GetSmallIcon();
+		if (PrimaryTab.IsValid())
+		{
+			PrimaryTab.Pin()->SetTabIcon(TabIcon);
+			PrimaryTab.Pin()->SetLabel(TabName);
+		}
+
+		if (HasIntegratedToolPalettes())
+		{
+			TSharedRef< SUniformWrapPanel> PaletteTabBox = SNew(SUniformWrapPanel)
+				.SlotPadding(FMargin(1.f, 2.f))
+				.HAlign(HAlign_Center);
+
+			// Only show if there is more than one child in the switcher
+			PaletteTabBox->SetVisibility(TAttribute<EVisibility>::Create(TAttribute<EVisibility>::FGetter::CreateLambda([PaletteTabBox]() -> EVisibility
+				{
+					return PaletteTabBox->GetChildren()->Num() > 1 ? EVisibility::Visible : EVisibility::Collapsed;
+				})));
+
+			// Also build the toolkit here 
+			TArray<FName> PaletteNames;
+			GetToolPaletteNames(PaletteNames);
+
+			TSharedPtr<FUICommandList> CommandList;
+			CommandList = GetToolkitCommands();
+
+			TSharedRef< SWidgetSwitcher > PaletteSwitcher = SNew(SWidgetSwitcher)
+				.WidgetIndex_Lambda([this, PaletteNames]() -> int32 {
+				int32 FoundIndex;
+				if (PaletteNames.Find(GetCurrentPalette(), FoundIndex))
+				{
+					return FoundIndex;
+				}
+				return 0;
+					});
+
+
+			for (auto Palette : PaletteNames)
+			{
+				FName ToolbarCustomizationName = GetEditorMode() ?  GetEditorMode()->GetModeInfo().ToolbarCustomizationName : GetScriptableEditorMode()->GetModeInfo().ToolbarCustomizationName;
+				FUniformToolBarBuilder ModeToolbarBuilder(CommandList, FMultiBoxCustomization(ToolbarCustomizationName));
+				ModeToolbarBuilder.SetStyle(&FEditorStyle::Get(), "PaletteToolBar");
+
+				 BuildToolPalette(Palette, ModeToolbarBuilder);
+
+				TSharedRef<SWidget> PaletteWidget = ModeToolbarBuilder.MakeWidget();
+
+				PaletteTabBox->AddSlot()
+					[
+						SNew(SCheckBox)
+						.Padding(FMargin(8.f, 4.f, 8.f, 5.f))
+						.Style(FEditorStyle::Get(), "PaletteToolBar.Tab")
+						.OnCheckStateChanged_Lambda([this, Palette](const ECheckBoxState) {
+							SetCurrentPalette(Palette);
+						}
+					)
+					.IsChecked_Lambda([this, Palette]() -> ECheckBoxState {
+							if (GetCurrentPalette() == Palette)
+							{
+								return ECheckBoxState::Checked;
+							}
+							return ECheckBoxState::Unchecked;
+						})
+							[
+								SNew(STextBlock)
+								.TextStyle(FAppStyle::Get(), "NormalText")
+								.Text(GetToolPaletteDisplayName(Palette))
+								.Justification(ETextJustify::Center)
+							]
+					];
+
+
+				PaletteSwitcher->AddSlot()
+					[
+						PaletteWidget
+					];
+			}
+			if (ModeToolHeader)
+			{
+				ModeToolHeader->SetContent(
+					SNew(SVerticalBox)
+
+					+ SVerticalBox::Slot()
+					.Padding(8.f, 0.f, 0.f, 8.f)
+					.AutoHeight()
+					[
+						PaletteTabBox
+					]
+
+				+ SVerticalBox::Slot()
+					.AutoHeight()
+					[
+						PaletteSwitcher
+					]
+				);
+			}
+		}
+		else
+		{
+			if (ModeToolHeader)
+			{
+				ModeToolHeader->SetContent(SNullWidget::NullWidget);
+			}
+			if (ModeToolBarContainer)
+			{
+				ModeToolBarContainer->SetContent(SNullWidget::NullWidget);
+			}
+		}
+
+		if (GetInlineContent().IsValid() && InlineContentHolder.IsValid())
+		{
+			InlineContentHolder->SetContent(GetInlineContent().ToSharedRef());
+		}
+
+	}
+	
+}
+
+EVisibility FModeToolkit::GetInlineContentHolderVisibility() const
+{
+	if (InlineContentHolder)
+	{
+		return InlineContentHolder->GetContent() == SNullWidget::NullWidget ? EVisibility::Collapsed : EVisibility::Visible;
+	}
+	return EVisibility::Collapsed;
+}
+
+EVisibility FModeToolkit::GetNoToolSelectedTextVisibility() const
+{
+	if (InlineContentHolder)
+	{
+		return InlineContentHolder->GetContent() == SNullWidget::NullWidget ? EVisibility::Visible : EVisibility::Collapsed;
+	}
+	return EVisibility::Collapsed;
+}
+
+TSharedRef<SDockTab> FModeToolkit::MakeModeToolbarTab(const FSpawnTabArgs& Args)
+{
+	TSharedRef<SDockTab> ToolbarTabRef =
+		SNew(SDockTab)
+		.Label(NSLOCTEXT("EditorModes", "EditorModesToolbarTitle", "Mode Toolbar"))
+		.ContentPadding(0.0f)
+		[
+			SAssignNew(ModeToolbarBox, SVerticalBox)
+		];
+
+	ModeToolbarTab = ToolbarTabRef;
+	SpawnOrUpdateModeToolbar();
+	return ToolbarTabRef;
+
+}
+
+void FModeToolkit::RequestModeUITabs()
+{
+	if (ModeUILayer.IsValid())
+	{
+		TSharedPtr<FAssetEditorModeUILayer> ModeUILayerPtr = ModeUILayer.Pin();
+		TSharedRef<FWorkspaceItem> MenuGroup = ModeUILayerPtr->GetModeMenuCategory().ToSharedRef();
+		PrimaryTabInfo.OnSpawnTab = FOnSpawnTab::CreateSP(SharedThis(this), &FModeToolkit::CreatePrimaryModePanel);
+		PrimaryTabInfo.TabLabel = LOCTEXT("ModesToolboxTab", "Mode Toolbox");
+		PrimaryTabInfo.TabTooltip = LOCTEXT("ModesToolboxTabTooltipText", "Open the  Modes tab, which contains the active editor mode's settings.");
+		ModeUILayerPtr->SetModePanelInfo(FAssetEditorModeUILayer::TopLeftTabID, PrimaryTabInfo);
+		if (!HasIntegratedToolPalettes())
+		{
+
+			ToolbarInfo.OnSpawnTab = FOnSpawnTab::CreateSP(SharedThis(this), &FModeToolkit::MakeModeToolbarTab);
+			ToolbarInfo.TabLabel = LOCTEXT("ModesToolbarTab", "Mode Toolbar");
+			ToolbarInfo.TabTooltip = LOCTEXT("LevelEditorModesToolbarTabTooltipText", "Opens a toolbar for the active editor mode");
+			ModeUILayerPtr->SetModePanelInfo(FAssetEditorModeUILayer::VerticalToolbarID, ToolbarInfo);
+		}
+	}
+}
+
+bool FModeToolkit::ShouldShowModeToolbar() const
+{
+	return ActiveToolBarRows.Num() > 0;
+}
+
+void FModeToolkit::SpawnOrUpdateModeToolbar()
+{
+	if (ShouldShowModeToolbar())
+	{
+		if (ModeToolbarTab.IsValid())
+		{
+			RebuildModeToolBar();
+		}
+		else if (ModeUILayer.IsValid())
+		{
+			ModeUILayer.Pin()->GetTabManager()->TryInvokeTab(FAssetEditorModeUILayer::VerticalToolbarID);
+		}
+	}
+}
+
+void FModeToolkit::RebuildModeToolBar()
+{
+	// If the tab or box is not valid the toolbar has not been opened or has been closed by the user
+	TSharedPtr<SVerticalBox> ModeToolbarBoxPinned = ModeToolbarBox.Pin();
+	if (ModeToolbarTab.IsValid() && ModeToolbarBoxPinned.IsValid())
+	{
+		ModeToolbarBoxPinned->ClearChildren();
+		bool bExclusivePalettes = true;
+		TSharedRef<SVerticalBox> ToolBoxVBox = SNew(SVerticalBox);
+
+		TSharedRef< SUniformWrapPanel> PaletteTabBox = SNew(SUniformWrapPanel)
+			.SlotPadding(FMargin(1.f, 2.f))
+			.HAlign(HAlign_Left);
+		TSharedRef< SWidgetSwitcher > PaletteSwitcher = SNew(SWidgetSwitcher);
+
+		int32 PaletteCount = ActiveToolBarRows.Num();
+		if (PaletteCount > 0)
+		{
+			for (int32 RowIdx = 0; RowIdx < PaletteCount; ++RowIdx)
+			{
+				const FEdModeToolbarRow& Row = ActiveToolBarRows[RowIdx];
+				if (ensure(Row.ToolbarWidget.IsValid()))
+				{
+					TSharedRef<SWidget> PaletteWidget = Row.ToolbarWidget.ToSharedRef();
+
+					bExclusivePalettes = HasExclusiveToolPalettes();
+
+					if (!bExclusivePalettes)
+					{
+						ToolBoxVBox->AddSlot()
+							.AutoHeight()
+							.Padding(FMargin(2.0, 2.0))
+							[
+								SNew(SExpandableArea)
+								.AreaTitle(Row.DisplayName)
+								.AreaTitleFont(FAppStyle::Get().GetFontStyle("NormalFont"))
+								.BorderImage(FAppStyle::Get().GetBrush("PaletteToolbar.ExpandableAreaHeader"))
+								.BodyBorderImage(FAppStyle::Get().GetBrush("PaletteToolbar.ExpandableAreaBody"))
+								.HeaderPadding(FMargin(4.f))
+								.Padding(FMargin(4.0, 0.0))
+								.BodyContent()
+								[
+									PaletteWidget
+								]
+							];
+					}
+					else
+					{
+						// Don't show Palette Tabs if there is only one
+						if (PaletteCount > 1)
+						{
+							PaletteTabBox->AddSlot()
+								[
+									SNew(SCheckBox)
+									.Style(FEditorStyle::Get(), "ToolPalette.DockingTab")
+								.OnCheckStateChanged_Lambda([PaletteSwitcher, Row, this](const ECheckBoxState) {
+										PaletteSwitcher->SetActiveWidget(Row.ToolbarWidget.ToSharedRef());
+										SetCurrentPalette(Row.PaletteName);
+									}
+								)
+								.IsChecked_Lambda([PaletteSwitcher, PaletteWidget]() -> ECheckBoxState { return PaletteSwitcher->GetActiveWidget() == PaletteWidget ? ECheckBoxState::Checked : ECheckBoxState::Unchecked; })
+										[
+											SNew(STextBlock)
+											.Text(Row.DisplayName)
+										]
+								];
+						}
+
+						PaletteSwitcher->AddSlot()
+							[
+								PaletteWidget
+							];
+					}
+				}
+			}
+
+			ModeToolbarBoxPinned->AddSlot()
+				.AutoHeight()
+				[
+					SNew(SOverlay)
+					+ SOverlay::Slot()
+				[
+					SNew(SImage)
+					.Image(FEditorStyle::GetBrush("ToolPalette.DockingWell"))
+				]
+
+			+ SOverlay::Slot()
+				[
+					PaletteTabBox
+				]
+				];
+
+			ModeToolbarBoxPinned->AddSlot()
+				.AutoHeight()
+				.Padding(1.f)
+				[
+					SNew(SBox)
+					.HeightOverride(PaletteSwitcher->GetNumWidgets() > 0 ? 45.f : 0.f)
+				[
+					PaletteSwitcher
+				]
+				];
+
+			ModeToolbarBoxPinned->AddSlot()
+				[
+					SNew(SScrollBox)
+
+					+ SScrollBox::Slot()
+				[
+					ToolBoxVBox
+				]
+				];
+
+			ModeToolbarPaletteSwitcher = PaletteSwitcher;
+		}
+	}
+}
+
+#undef LOCTEXT_NAMESPACE
\ No newline at end of file
diff --git a/Engine/Source/Editor/UnrealEd/Private/Tools/UEdMode.cpp b/Engine/Source/Editor/UnrealEd/Private/Tools/UEdMode.cpp
--- a/Engine/Source/Editor/UnrealEd/Private/Tools/UEdMode.cpp
+++ b/Engine/Source/Editor/UnrealEd/Private/Tools/UEdMode.cpp
@@ -131,6 +131,7 @@ void UEdMode::Exit()
 		{
 			CommandList->UnmapAction(RegisteredTool.Key);
 		}
+		FToolkitManager::Get().CloseToolkit(Toolkit.ToSharedRef());
 		Toolkit.Reset();
 	}
 	RegisteredTools.SetNum(0);
diff --git a/Engine/Source/Editor/UnrealEd/Public/EditorModeManager.h b/Engine/Source/Editor/UnrealEd/Public/EditorModeManager.h
--- a/Engine/Source/Editor/UnrealEd/Public/EditorModeManager.h
+++ b/Engine/Source/Editor/UnrealEd/Public/EditorModeManager.h
@@ -95,15 +95,7 @@ public:
 	 */
 	void DestroyMode(FEditorModeID InID);
 
-	/**
-	 * Creates the mode toolbar tab if needed
-	 */
-	TSharedRef<SDockTab> MakeModeToolbarTab();
 
-	/**
-	 * Whether or not the mode toolbar should be shown.  If any active modes generated a toolbar this method will return true
-	 */
-	bool ShouldShowModeToolbar() const;
 
 	/**
 	 * Whether or not the mode toolbox (where mode details panels and some tools are) should be shown.
@@ -117,9 +109,6 @@ protected:
 	/** Removes the mode ID from the tools manager when a mode is unregistered */
 	void OnModeUnregistered(FEditorModeID ModeID);
 		
-private:
-	void RebuildModeToolBar();
-	void SpawnOrUpdateModeToolbar();
 public:
 
 	/**
@@ -416,9 +405,6 @@ public:
 	/** Is the viewport UI hidden? */
 	bool IsViewportUIHidden() const { return bHideViewportUI; }
 
-	/** The toolbar tab name that should be used as the tab identifier */
-	static const FName EditorModeToolbarTabName;
-
 	bool PivotShown;
 	bool Snapping;
 	bool SnappedActor;
@@ -526,7 +512,7 @@ public:
 	/*
 	* Sets the active Modes ToolBar Palette Tab to the named Palette
 	*/
-	void  InvokeToolPaletteTab(FEditorModeID InMode, FName InPaletteName);
+	//void  InvokeToolPaletteTab(FEditorModeID InMode, FName InPaletteName);
 
 	/** returns true if all active EdModes are OK with an AutoSave happening now  */
 	bool CanAutoSave() const;
@@ -597,22 +583,6 @@ protected:
 	float WidgetScale;
 
 private:
-	struct FEdModeToolbarRow
-	{
-		FEdModeToolbarRow(FName InModeID, FName InPaletteName, FText InDisplayName, TSharedRef<SWidget> InToolbarWidget)
-			: ModeID(InModeID)
-			, PaletteName(InPaletteName)
-			, DisplayName(InDisplayName)
-			, ToolbarWidget(InToolbarWidget)
-		{}
-		FName ModeID;
-		FName PaletteName;
-		FText DisplayName;
-		TSharedPtr<SWidget> ToolbarWidget;
-	};
-
-	/** All toolbar rows generated by active modes.  There will be one row per active mode that generates a toolbar */
-	TArray<FEdModeToolbarRow> ActiveToolBarRows;
 
 	/** The coordinate system the widget is operating within. */
 	ECoordSystem CoordSystem;
@@ -626,15 +596,6 @@ private:
 	/** Multicast delegate that is broadcast when the coordinate system is changed */
 	FCoordSystemChangedEvent CoordSystemChangedEvent;
 
-	/** The dock tab for any modes that generate a toolbar */
-	TWeakPtr<SDockTab> ModeToolbarTab;
-
-	/** The actual toolbar rows will be placed in this vertical box */
-	TWeakPtr<SVerticalBox> ModeToolbarBox;
-
-	/** The modes palette toolbar **/	
-	TWeakPtr<SWidgetSwitcher> ModeToolbarPaletteSwitcher;
-
 	/** Flag set between calls to StartTracking() and EndTracking() */
 	bool bIsTracking;
 
diff --git a/Engine/Source/Editor/UnrealEd/Public/Toolkits/AssetEditorToolkit.h b/Engine/Source/Editor/UnrealEd/Public/Toolkits/AssetEditorToolkit.h
--- a/Engine/Source/Editor/UnrealEd/Public/Toolkits/AssetEditorToolkit.h
+++ b/Engine/Source/Editor/UnrealEd/Public/Toolkits/AssetEditorToolkit.h
@@ -134,10 +134,6 @@ public:
 	virtual FText GetTabSuffix() const override;
 	virtual FText GetToolkitToolTipText() const override;
 	virtual FString GetWorldCentricTabPrefix() const override = 0;	// Must implement in derived class!
-	virtual class FEdMode* GetEditorMode() const override;
-	virtual TWeakObjectPtr<class UEdMode> GetScriptableEditorMode() const final;
-	virtual FText GetEditorModeDisplayName() const override;
-	virtual FSlateIcon GetEditorModeIcon() const override;
 	virtual FEditorModeTools& GetEditorModeManager() const final;
 
 	/** IAssetEditorInstance interface */
diff --git a/Engine/Source/Editor/UnrealEd/Public/Toolkits/BaseToolkit.h b/Engine/Source/Editor/UnrealEd/Public/Toolkits/BaseToolkit.h
--- a/Engine/Source/Editor/UnrealEd/Public/Toolkits/BaseToolkit.h
+++ b/Engine/Source/Editor/UnrealEd/Public/Toolkits/BaseToolkit.h
@@ -9,12 +9,15 @@
 #include "Framework/Commands/UICommandList.h"
 #include "Toolkits/IToolkitHost.h"
 
+
 class UInteractiveTool;
 class UInteractiveToolManager;
 class UEdMode;
 class IDetailsView;
 struct FDetailsViewArgs;
+class FAssetEditorModeUILayer;
 enum class EToolShutdownType;
+class SBorder;
 
 /**
  * Base class for all toolkits (abstract).
@@ -104,13 +107,12 @@ public:
 	~FModeToolkit();
 
 public:
+	virtual void SetModeUILayer(const TSharedPtr<FAssetEditorModeUILayer> InLayer) override;
 
 	// FBaseToolkit overrides
 
-	virtual void RegisterTabSpawners(const TSharedRef<FTabManager>& TabManager) override { }
-	virtual void UnregisterTabSpawners(const TSharedRef<FTabManager>& TabManager) override { }
-
-public:
+	virtual void RegisterTabSpawners(const TSharedRef<FTabManager>& TabManager) final;
+	virtual void UnregisterTabSpawners(const TSharedRef<FTabManager>& TabManager) final;
 
 	// IToolkit interface
 	virtual FName GetToolkitFName() const override;
@@ -121,14 +123,13 @@ public:
 	virtual bool IsAssetEditor() const override;
 	virtual const TArray<UObject*>* GetObjectsCurrentlyBeingEdited() const override;
 	virtual FLinearColor GetWorldCentricTabColorScale() const override;
-	virtual class FEdMode* GetEditorMode() const override;
-	virtual FText GetEditorModeDisplayName() const override;
-	virtual FSlateIcon GetEditorModeIcon() const override;
 	virtual FEditorModeTools& GetEditorModeManager() const final;
 
 	virtual TWeakObjectPtr<UEdMode> GetScriptableEditorMode() const final;
-	virtual TSharedPtr<SWidget> GetInlineContent() const override;
-
+	virtual TSharedPtr<SWidget> GetInlineContent() const;
+	virtual FEdMode* GetEditorMode() const;
+	virtual	FText GetEditorModeDisplayName() const;
+	virtual FSlateIcon GetEditorModeIcon() const;
 	/** Returns the number of Mode specific tabs in the mode toolbar **/
 	virtual void GetToolPaletteNames(TArray<FName>& PaletteNames) const {}
 
@@ -160,7 +161,7 @@ public:
 	void SetCurrentPalette(FName InName);
 	FName GetCurrentPalette() const;
 	void SetModeSettingsObject(UObject* InSettingsObject);
-
+	virtual void InvokeUI();
 protected:
 	bool CanStartTool(const FString& ToolTypeIdentifier);
 	bool CanAcceptActiveTool();
@@ -173,6 +174,27 @@ protected:
 	virtual void CustomizeModeDetailsViewArgs(FDetailsViewArgs& ArgsInOut) {}
 	virtual void CustomizeDetailsViewArgs(FDetailsViewArgs& ArgsInOut) {}
 
+	virtual void RequestModeUITabs();
+	/**
+ * Whether or not the mode toolbar should be shown.  If any active modes generated a toolbar this method will return true
+ */
+	bool ShouldShowModeToolbar() const;
+	TSharedRef<SDockTab> CreatePrimaryModePanel(const FSpawnTabArgs& Args);
+	void UpdatePrimaryModePanel();
+	EVisibility GetInlineContentHolderVisibility() const;
+	EVisibility GetNoToolSelectedTextVisibility() const;
+
+	/**
+	 * Creates the mode toolbar tab if needed
+	 */
+	TSharedRef<SDockTab> MakeModeToolbarTab(const FSpawnTabArgs& Args);
+
+
+
+	void SpawnOrUpdateModeToolbar();
+	void RebuildModeToolBar();
+
+
 protected:
 	TSharedPtr<SWidget> ToolkitWidget;
 	TSharedPtr<IDetailsView> ModeDetailsView;
@@ -182,4 +204,46 @@ protected:
 
 	FName CurrentPaletteName;
 	FOnPaletteChanged OnPaletteChangedDelegate;
+
+	/** Inline content area for editor modes */
+	TSharedPtr<SBorder> InlineContentHolder;
+
+	/** The container holding the mode toolbar */
+	TSharedPtr<SBorder> ModeToolBarContainer;
+
+	/** The active tool header area **/
+	TSharedPtr<SBorder> ModeToolHeader;
+	TWeakPtr<SDockTab> PrimaryTab;
+	/** The dock tab for any modes that generate a toolbar */
+	TWeakPtr<SDockTab> ModeToolbarTab;
+	FMinorTabConfig PrimaryTabInfo;
+	FMinorTabConfig ToolbarInfo;
+
+
+	/** The actual toolbar rows will be placed in this vertical box */
+	TWeakPtr<SVerticalBox> ModeToolbarBox;
+
+	/** The modes palette toolbar **/
+	TWeakPtr<SWidgetSwitcher> ModeToolbarPaletteSwitcher;
+
+	TWeakPtr<FAssetEditorModeUILayer> ModeUILayer;
+
+	TArray<TSharedPtr<SDockTab>> CreatedTabs;
+
+	struct FEdModeToolbarRow
+	{
+		FEdModeToolbarRow(FName InModeID, FName InPaletteName, FText InDisplayName, TSharedRef<SWidget> InToolbarWidget)
+			: ModeID(InModeID)
+			, PaletteName(InPaletteName)
+			, DisplayName(InDisplayName)
+			, ToolbarWidget(InToolbarWidget)
+		{}
+		FName ModeID;
+		FName PaletteName;
+		FText DisplayName;
+		TSharedPtr<SWidget> ToolbarWidget;
+	};
+
+	/** All toolbar rows generated by active modes.  There will be one row per active mode that generates a toolbar */
+	TArray<FEdModeToolbarRow> ActiveToolBarRows;
 };
diff --git a/Engine/Source/Runtime/Slate/Private/Framework/Docking/LayoutExtender.cpp b/Engine/Source/Runtime/Slate/Private/Framework/Docking/LayoutExtender.cpp
--- a/Engine/Source/Runtime/Slate/Private/Framework/Docking/LayoutExtender.cpp
+++ b/Engine/Source/Runtime/Slate/Private/Framework/Docking/LayoutExtender.cpp
@@ -12,6 +12,11 @@ void FLayoutExtender::ExtendArea(FName ExtensionId, const FAreaExtension& AreaEx
 	AreaExtensions.Add(ExtensionId, FExtendedArea(AreaExtension));
 }
 
+SLATE_API void FLayoutExtender::ExtendStack(FName ExtensionId, ELayoutExtensionPosition Position, FTabManager::FTab TabToAdd)
+{
+	StackExtensions.Add(ExtensionId, FExtendedTab(Position, TabToAdd));
+}
+
 void FLayoutExtender::ExtendAreaRecursive(const TSharedRef<FTabManager::FArea>& Area) const
 {
 	FName ExtensionId = Area->GetExtensionId();
diff --git a/Engine/Source/Runtime/Slate/Private/Framework/Docking/TabManager.cpp b/Engine/Source/Runtime/Slate/Private/Framework/Docking/TabManager.cpp
--- a/Engine/Source/Runtime/Slate/Private/Framework/Docking/TabManager.cpp
+++ b/Engine/Source/Runtime/Slate/Private/Framework/Docking/TabManager.cpp
@@ -553,6 +553,18 @@ void FTabManager::FLayout::ProcessExtensions(const FLayoutExtender& Extender)
 
 	for (FTabManager::FStack* Stack : AllTabs.AllStacks)
 	{
+		// First add to the front of the stack
+		Extender.FindStackExtensions(Stack->GetExtensionId(), ELayoutExtensionPosition::Before, ExtendedTabs);
+		int32 InsertedTabIndex = 0;
+		for (FTab& NewTab : ExtendedTabs)
+		{
+			if (!AllTabs.Contains(NewTab.TabId))
+			{
+				Stack->Tabs.Insert(NewTab, InsertedTabIndex++);
+			}
+		}
+
+		// This is the per-tab extension section
 		FSplitter* ParentSplitter = AllTabs.StackToParentSplitterMap.FindRef(Stack);
 		for (int32 TabIndex = 0; TabIndex < Stack->Tabs.Num();)
 		{
@@ -614,6 +626,17 @@ void FTabManager::FLayout::ProcessExtensions(const FLayoutExtender& Extender)
 				}
 			}
 		}
+
+		// Finally add to the end of the stack
+		Extender.FindStackExtensions(Stack->GetExtensionId(), ELayoutExtensionPosition::After, ExtendedTabs);
+		InsertedTabIndex = Stack->Tabs.Num();
+		for (FTab& NewTab : ExtendedTabs)
+		{
+			if (!AllTabs.Contains(NewTab.TabId))
+			{
+				Stack->Tabs.Insert(NewTab, InsertedTabIndex++);
+			}
+		}
 	}
 }
 
diff --git a/Engine/Source/Runtime/Slate/Public/Framework/Docking/LayoutExtender.h b/Engine/Source/Runtime/Slate/Public/Framework/Docking/LayoutExtender.h
--- a/Engine/Source/Runtime/Slate/Public/Framework/Docking/LayoutExtender.h
+++ b/Engine/Source/Runtime/Slate/Public/Framework/Docking/LayoutExtender.h
@@ -46,6 +46,16 @@ public:
 	 */
 	SLATE_API void ExtendArea(FName ExtensionId, const FAreaExtension& AreaExtension);
 
+	/**
+	 * Extend the tab stack identified by the specified extension ID
+	 *
+	 * @param ExtensionId			The ID of the area to extend (FTabManager::FStack::ExtensionId)
+	 * @param AreaExtensions		A callback to call with the default layout for the desired area
+	 * @param Position				Whether to place the new tab at the start or end of this stack. Above and below are not used for this case.
+	 * @param TabToAdd				The new tab definition
+	 */
+	SLATE_API void ExtendStack(FName ExtensionId, ELayoutExtensionPosition Position, FTabManager::FTab TabToAdd);
+
 
 	/**
 	 * Populate the specified container with the tabs that relate to the specified tab ID
@@ -71,6 +81,30 @@ public:
 		}
 	}
 
+	/**
+	 * Populate the specified container with the tabs that relate to the specified stack
+	 *
+	 * @param TabId					The existing tab that may be extended
+	 * @param Position				The position to acquire extensions for (before/after)
+	 * @param OutValues				The container to populate with extended tabs
+	 */
+	template<typename AllocatorType>
+	void FindStackExtensions(FName ExtensionID, ELayoutExtensionPosition Position, TArray<FTabManager::FTab, AllocatorType>& OutValues) const
+	{
+		OutValues.Reset();
+
+		TArray<FExtendedTab, TInlineAllocator<2>> Extensions;
+		StackExtensions.MultiFind(ExtensionID, Extensions, true);
+
+		for (FExtendedTab& Extension : Extensions)
+		{
+			if (Extension.Position == Position)
+			{
+				OutValues.Add(Extension.TabToAdd);
+			}
+		}
+	}
+
 	/**
 	 * Recursively extend the specified area
 	 */
@@ -105,6 +139,9 @@ private:
 	/** Map of extensions for tabs */
 	TMultiMap<FTabId, FExtendedTab> TabExtensions;
 
+	/** Map of extensions for stacks */
+	TMultiMap<FName, FExtendedTab> StackExtensions;
+
 	/** Map of extensions for areas */
 	TMultiMap<FName, FExtendedArea> AreaExtensions;
 };
\ No newline at end of file
diff --git a/Engine/Source/Runtime/Slate/Public/Framework/Docking/TabManager.h b/Engine/Source/Runtime/Slate/Public/Framework/Docking/TabManager.h
--- a/Engine/Source/Runtime/Slate/Public/Framework/Docking/TabManager.h
+++ b/Engine/Source/Runtime/Slate/Public/Framework/Docking/TabManager.h
@@ -161,6 +161,35 @@ DECLARE_DELEGATE_RetVal_OneParam(bool, FCanSpawnTab, const FSpawnTabArgs&);
  */
 DECLARE_DELEGATE_RetVal_OneParam( TSharedPtr<SDockTab>, FOnFindTabToReuse, const FTabId& )
 
+struct SLATE_API FMinorTabConfig
+{
+public:
+	FMinorTabConfig()
+	{
+	}
+
+	FMinorTabConfig(const FName& InTabID)
+		: TabId(InTabID)
+	{
+	}
+
+	FName TabId;
+
+	FText TabLabel;
+
+	FText TabTooltip;
+
+	FSlateIcon TabIcon;
+
+	FOnSpawnTab OnSpawnTab;
+
+	FCanSpawnTab CanSpawnTab;
+
+	FOnFindTabToReuse OnFindTabToReuse;
+
+	TSharedPtr<FWorkspaceItem> WorkspaceGroup;
+};
+
 /** An enum to describe how TabSpawnerEntries will be handled by menus. */
 namespace ETabSpawnerMenuType
 {
@@ -193,15 +222,27 @@ struct FTabSpawnerEntry : public FWorkspaceItem
 		return *this;
 	}
 
+	FTabSpawnerEntry& SetDisplayNameAttribute(const TAttribute<FText>& InLegibleName)
+	{
+		DisplayNameAttribute = InLegibleName;
+		return *this;
+	}
+
+	FTabSpawnerEntry& SetTooltipTextAttribute(const TAttribute<FText>& InTooltipText)
+	{
+		TooltipTextAttribute = InTooltipText;
+		return *this;
+	}
+
 	FTabSpawnerEntry& SetDisplayName( const FText& InLegibleName )
 	{
-		DisplayName = InLegibleName;
+		DisplayNameAttribute = InLegibleName;
 		return *this;
 	}
 
 	FTabSpawnerEntry& SetTooltipText( const FText& InTooltipText )
 	{
-		TooltipText = InTooltipText;
+		TooltipTextAttribute = InTooltipText;
 		return *this;
 	}
 
@@ -454,6 +495,17 @@ class SLATE_API FTabManager : public TSharedFromThis<FTabManager>
 				{
 				}
 
+				TSharedRef<FStack> SetExtensionId(FName InExtensionId)
+				{
+					ExtensionId = InExtensionId;
+					return SharedThis(this);
+				}
+
+				FName GetExtensionId() const
+				{
+					return ExtensionId;
+				}
+
 			protected:
 
 				FStack()
@@ -466,6 +518,7 @@ class SLATE_API FTabManager : public TSharedFromThis<FTabManager>
 				TArray<FTab> Tabs;
 				bool bHideTabWell;
 				FTabId ForegroundTabId;
+				FName ExtensionId;
 		};
 
 
diff --git a/Engine/Source/Runtime/Slate/Public/Framework/Docking/WorkspaceItem.h b/Engine/Source/Runtime/Slate/Public/Framework/Docking/WorkspaceItem.h
--- a/Engine/Source/Runtime/Slate/Public/Framework/Docking/WorkspaceItem.h
+++ b/Engine/Source/Runtime/Slate/Public/Framework/Docking/WorkspaceItem.h
@@ -5,6 +5,7 @@
 
 #include "CoreMinimal.h"
 #include "Textures/SlateIcon.h"
+#include "Misc/Attribute.h"
 
 struct FTabSpawnerEntry;
 
@@ -57,12 +58,12 @@ public:
 
 	const FText& GetDisplayName() const
 	{
-		return DisplayName;
+		return DisplayNameAttribute.Get();
 	}
 	
 	const FText& GetTooltipText() const
 	{
-		return TooltipText;
+		return TooltipTextAttribute.Get();
 	}
 
 	const FSlateIcon& GetIcon() const
@@ -139,28 +140,34 @@ public:
 		return bIsGroupPopulated;
 	}
 
+PRAGMA_DISABLE_DEPRECATION_WARNINGS
 	virtual ~FWorkspaceItem()
 	{
 	}
+PRAGMA_ENABLE_DEPRECATION_WARNINGS
 
 protected:
 	FWorkspaceItem( const FText& InDisplayName, const FSlateIcon& InIcon, const bool bInSortChildren )
 		: Icon(InIcon)
-		, DisplayName(InDisplayName)
+		, DisplayNameAttribute(InDisplayName)
 		, bSortChildren(bInSortChildren)
 	{
 	}
 
 	FWorkspaceItem( const FText& InDisplayName, const FText& InTooltipText, const FSlateIcon& InIcon, const bool bInSortChildren )
 		: Icon(InIcon)
-		, DisplayName(InDisplayName)
-		, TooltipText(InTooltipText)
+		, DisplayNameAttribute(InDisplayName)
+		, TooltipTextAttribute(InTooltipText)
 		, bSortChildren(bInSortChildren)
 	{
 	}
 
 	FSlateIcon Icon;
+	TAttribute<FText> DisplayNameAttribute;
+	TAttribute<FText> TooltipTextAttribute;
+	UE_DEPRECATED(5.0, "Use DisplayNameAttribute instead.")
 	FText DisplayName;
+	UE_DEPRECATED(5.0, "Use TooltipTextAttribute instead.")
 	FText TooltipText;
 	bool bSortChildren;
 
