Minor fix for conditional

diff --git a/Engine/Source/Editor/LevelEditor/Private/SInViewportDetails.cpp b/Engine/Source/Editor/LevelEditor/Private/SInViewportDetails.cpp
--- a/Engine/Source/Editor/LevelEditor/Private/SInViewportDetails.cpp
+++ b/Engine/Source/Editor/LevelEditor/Private/SInViewportDetails.cpp
@@ -413,7 +413,7 @@ void SInViewportDetails::PostRedo(bool bSuccess)
 
 void SInViewportDetails::OnEditorSelectionChanged(UObject* Object)
 {
-	if (Object->GetWorld()->WorldType != EWorldType::Editor)
+	if (Object && Object->GetWorld() && Object->GetWorld()->WorldType != EWorldType::Editor)
 	{
 		return;
 	}
