Changed the TypedElementInterface to use UInterfaces instead of an object base api.

diff --git a/Engine/Plugins/MeshPainting/Source/MeshPaintingToolset/Private/MeshPaintInteractions.cpp b/Engine/Plugins/MeshPainting/Source/MeshPaintingToolset/Private/MeshPaintInteractions.cpp
--- a/Engine/Plugins/MeshPainting/Source/MeshPaintingToolset/Private/MeshPaintInteractions.cpp
+++ b/Engine/Plugins/MeshPainting/Source/MeshPaintingToolset/Private/MeshPaintInteractions.cpp
@@ -90,7 +90,7 @@ bool UMeshPaintSelectionMechanic::FindClickedComponentsAndCacheAdapters(const FI
 	{
 		if (HHitProxy* HitProxy = FocusedViewport->GetHitProxy(ClickPos.ScreenPosition.X, ClickPos.ScreenPosition.Y))
 		{
-			if (TTypedElement<UTypedElementObjectInterface> ObjectInterface = UTypedElementRegistry::GetInstance()->GetElement<UTypedElementObjectInterface>(HitProxy->GetElementHandle()))
+			if (TTypedElement<ITypedElementObjectInterface> ObjectInterface = UTypedElementRegistry::GetInstance()->GetElement<ITypedElementObjectInterface>(HitProxy->GetElementHandle()))
 			{
 				if (AActor* Actor = ObjectInterface.GetObjectAs<UActorComponent>()->GetOwner())
 				{
diff --git a/Engine/Source/Editor/AssetPlacementEdMode/Private/Modes/PlacementModeSubsystem.cpp b/Engine/Source/Editor/AssetPlacementEdMode/Private/Modes/PlacementModeSubsystem.cpp
--- a/Engine/Source/Editor/AssetPlacementEdMode/Private/Modes/PlacementModeSubsystem.cpp
+++ b/Engine/Source/Editor/AssetPlacementEdMode/Private/Modes/PlacementModeSubsystem.cpp
@@ -38,7 +38,7 @@ UAssetPlacementSettings* UPlacementModeSubsystem::GetMutableModeSettingsObject()
 
 bool UPlacementModeSubsystem::DoesActivePaletteSupportElement(const FTypedElementHandle& InElementToCheck) const
 {
-	if (TTypedElement<UTypedElementAssetDataInterface> AssetDataInterface = UTypedElementRegistry::GetInstance()->GetElement<UTypedElementAssetDataInterface>(InElementToCheck))
+	if (TTypedElement<ITypedElementAssetDataInterface> AssetDataInterface = UTypedElementRegistry::GetInstance()->GetElement<ITypedElementAssetDataInterface>(InElementToCheck))
 	{
 		TArray<FAssetData> ReferencedAssetDatas = AssetDataInterface.GetAllReferencedAssetDatas();
 		for (const FPaletteItem& Item : ModeSettings->GetActivePaletteItems())
diff --git a/Engine/Source/Editor/AssetPlacementEdMode/Private/Tools/PlacementPlaceSingleTool.cpp b/Engine/Source/Editor/AssetPlacementEdMode/Private/Tools/PlacementPlaceSingleTool.cpp
--- a/Engine/Source/Editor/AssetPlacementEdMode/Private/Tools/PlacementPlaceSingleTool.cpp
+++ b/Engine/Source/Editor/AssetPlacementEdMode/Private/Tools/PlacementPlaceSingleTool.cpp
@@ -347,7 +347,7 @@ void UPlacementModePlaceSingleTool::UpdateElementTransforms(TArrayView<const FTy
 	// Update the transform positions for the preview elements.
 	for (const FTypedElementHandle& ElementHandle : InElements)
 	{
-		if (TTypedElement<UTypedElementWorldInterface> WorldInterfaceElement = UTypedElementRegistry::GetInstance()->GetElement<UTypedElementWorldInterface>(ElementHandle))
+		if (TTypedElement<ITypedElementWorldInterface> WorldInterfaceElement = UTypedElementRegistry::GetInstance()->GetElement<ITypedElementWorldInterface>(ElementHandle))
 		{
 			bLocalTransform ? WorldInterfaceElement.SetRelativeTransform(InTransform) : WorldInterfaceElement.SetWorldTransform(InTransform);
 			WorldInterfaceElement.NotifyMovementOngoing();
@@ -360,7 +360,7 @@ void UPlacementModePlaceSingleTool::NotifyMovementStarted(TArrayView<const FType
 	// Notify Movement started
 	for (const FTypedElementHandle& PreviewElement : PreviewElements)
 	{
-		if (TTypedElement<UTypedElementWorldInterface> WorldInterfaceElement = UTypedElementRegistry::GetInstance()->GetElement<UTypedElementWorldInterface>(PreviewElement))
+		if (TTypedElement<ITypedElementWorldInterface> WorldInterfaceElement = UTypedElementRegistry::GetInstance()->GetElement<ITypedElementWorldInterface>(PreviewElement))
 		{
 			WorldInterfaceElement.NotifyMovementStarted();
 		}
@@ -383,7 +383,7 @@ void UPlacementModePlaceSingleTool::NotifyMovementEnded(TArrayView<const FTypedE
 
 	for (const FTypedElementHandle& PreviewElement : PreviewElements)
 	{
-		if (TTypedElement<UTypedElementWorldInterface> WorldInterfaceElement = UTypedElementRegistry::GetInstance()->GetElement<UTypedElementWorldInterface>(PreviewElement))
+		if (TTypedElement<ITypedElementWorldInterface> WorldInterfaceElement = UTypedElementRegistry::GetInstance()->GetElement<ITypedElementWorldInterface>(PreviewElement))
 		{
 			WorldInterfaceElement.NotifyMovementEnded();
 			WorldInterfaceElement.DeleteElement(WorldInterfaceElement.GetOwnerWorld(), SelectionSet, FTypedElementDeletionOptions());
diff --git a/Engine/Source/Editor/EditorFramework/Private/Elements/Framework/TypedElementViewportInteraction.cpp b/Engine/Source/Editor/EditorFramework/Private/Elements/Framework/TypedElementViewportInteraction.cpp
--- a/Engine/Source/Editor/EditorFramework/Private/Elements/Framework/TypedElementViewportInteraction.cpp
+++ b/Engine/Source/Editor/EditorFramework/Private/Elements/Framework/TypedElementViewportInteraction.cpp
@@ -5,7 +5,7 @@
 #include "Elements/Framework/TypedElementRegistry.h"
 #include "Elements/Framework/TypedElementUtil.h"
 
-bool FTypedElementViewportInteractionCustomization::GetGizmoPivotLocation(const TTypedElement<UTypedElementWorldInterface>& InElementWorldHandle, const UE::Widget::EWidgetMode InWidgetMode, FVector& OutPivotLocation)
+bool FTypedElementViewportInteractionCustomization::GetGizmoPivotLocation(const TTypedElement<ITypedElementWorldInterface>& InElementWorldHandle, const UE::Widget::EWidgetMode InWidgetMode, FVector& OutPivotLocation)
 {
 	FTransform ElementWorldTransform;
 	if (InElementWorldHandle.GetWorldTransform(ElementWorldTransform))
@@ -21,12 +21,12 @@ void FTypedElementViewportInteractionCustomization::PreGizmoManipulationStarted(
 {
 }
 
-void FTypedElementViewportInteractionCustomization::GizmoManipulationStarted(const TTypedElement<UTypedElementWorldInterface>& InElementWorldHandle, const UE::Widget::EWidgetMode InWidgetMode)
+void FTypedElementViewportInteractionCustomization::GizmoManipulationStarted(const TTypedElement<ITypedElementWorldInterface>& InElementWorldHandle, const UE::Widget::EWidgetMode InWidgetMode)
 {
 	InElementWorldHandle.NotifyMovementStarted();
 }
 
-void FTypedElementViewportInteractionCustomization::GizmoManipulationDeltaUpdate(const TTypedElement<UTypedElementWorldInterface>& InElementWorldHandle, const UE::Widget::EWidgetMode InWidgetMode, const EAxisList::Type InDragAxis, const FInputDeviceState& InInputState, const FTransform& InDeltaTransform, const FVector& InPivotLocation)
+void FTypedElementViewportInteractionCustomization::GizmoManipulationDeltaUpdate(const TTypedElement<ITypedElementWorldInterface>& InElementWorldHandle, const UE::Widget::EWidgetMode InWidgetMode, const EAxisList::Type InDragAxis, const FInputDeviceState& InInputState, const FTransform& InDeltaTransform, const FVector& InPivotLocation)
 {
 	FTransform ElementWorldTransform;
 	if (InElementWorldHandle.GetWorldTransform(ElementWorldTransform))
@@ -72,7 +72,7 @@ void FTypedElementViewportInteractionCustomization::GizmoManipulationDeltaUpdate
 	}
 }
 
-void FTypedElementViewportInteractionCustomization::GizmoManipulationStopped(const TTypedElement<UTypedElementWorldInterface>& InElementWorldHandle, const UE::Widget::EWidgetMode InWidgetMode, const ETypedElementViewportInteractionGizmoManipulationType InManipulationType)
+void FTypedElementViewportInteractionCustomization::GizmoManipulationStopped(const TTypedElement<ITypedElementWorldInterface>& InElementWorldHandle, const UE::Widget::EWidgetMode InWidgetMode, const ETypedElementViewportInteractionGizmoManipulationType InManipulationType)
 {
 	InElementWorldHandle.NotifyMovementEnded();
 }
@@ -81,7 +81,7 @@ void FTypedElementViewportInteractionCustomization::PostGizmoManipulationStopped
 {
 }
 
-void FTypedElementViewportInteractionCustomization::MirrorElement(const TTypedElement<UTypedElementWorldInterface>& InElementWorldHandle, const FVector& InMirrorScale, const FVector& InPivotLocation)
+void FTypedElementViewportInteractionCustomization::MirrorElement(const TTypedElement<ITypedElementWorldInterface>& InElementWorldHandle, const FVector& InMirrorScale, const FVector& InPivotLocation)
 {
 	FTransform ElementWorldTransform;
 	if (InElementWorldHandle.GetWorldTransform(ElementWorldTransform))
@@ -130,7 +130,7 @@ void FTypedElementViewportInteractionCustomization::MirrorElement(const TTypedEl
 	}
 }
 
-bool FTypedElementViewportInteractionCustomization::GetFocusBounds(const TTypedElement<UTypedElementWorldInterface>& InElementWorldHandle, FBoxSphereBounds& OutBounds)
+bool FTypedElementViewportInteractionCustomization::GetFocusBounds(const TTypedElement<ITypedElementWorldInterface>& InElementWorldHandle, FBoxSphereBounds& OutBounds)
 {
 	return InElementWorldHandle.GetBounds(OutBounds);
 }
@@ -241,6 +241,6 @@ bool UTypedElementViewportInteraction::GetFocusBounds(FTypedElementListConstRef
 FTypedElementViewportInteractionElement UTypedElementViewportInteraction::ResolveViewportInteractionElement(const FTypedElementHandle& InElementHandle) const
 {
 	return InElementHandle
-		? FTypedElementViewportInteractionElement(UTypedElementRegistry::GetInstance()->GetElement<UTypedElementWorldInterface>(InElementHandle), GetInterfaceCustomizationByTypeId(InElementHandle.GetId().GetTypeId()))
+		? FTypedElementViewportInteractionElement(UTypedElementRegistry::GetInstance()->GetElement<ITypedElementWorldInterface>(InElementHandle), GetInterfaceCustomizationByTypeId(InElementHandle.GetId().GetTypeId()))
 		: FTypedElementViewportInteractionElement();
 }
diff --git a/Engine/Source/Editor/EditorFramework/Public/Elements/Framework/TypedElementViewportInteraction.h b/Engine/Source/Editor/EditorFramework/Public/Elements/Framework/TypedElementViewportInteraction.h
--- a/Engine/Source/Editor/EditorFramework/Public/Elements/Framework/TypedElementViewportInteraction.h
+++ b/Engine/Source/Editor/EditorFramework/Public/Elements/Framework/TypedElementViewportInteraction.h
@@ -28,14 +28,14 @@ public:
 	virtual ~FTypedElementViewportInteractionCustomization() = default;
 
 	//~ See UTypedElementViewportInteraction for API docs
-	virtual bool GetGizmoPivotLocation(const TTypedElement<UTypedElementWorldInterface>& InElementWorldHandle, const UE::Widget::EWidgetMode InWidgetMode, FVector& OutPivotLocation);
+	virtual bool GetGizmoPivotLocation(const TTypedElement<ITypedElementWorldInterface>& InElementWorldHandle, const UE::Widget::EWidgetMode InWidgetMode, FVector& OutPivotLocation);
 	virtual void PreGizmoManipulationStarted(TArrayView<const FTypedElementHandle> InElementHandles, const UE::Widget::EWidgetMode InWidgetMode);
-	virtual void GizmoManipulationStarted(const TTypedElement<UTypedElementWorldInterface>& InElementWorldHandle, const UE::Widget::EWidgetMode InWidgetMode);
-	virtual void GizmoManipulationDeltaUpdate(const TTypedElement<UTypedElementWorldInterface>& InElementWorldHandle, const UE::Widget::EWidgetMode InWidgetMode, const EAxisList::Type InDragAxis, const FInputDeviceState& InInputState, const FTransform& InDeltaTransform, const FVector& InPivotLocation);
-	virtual void GizmoManipulationStopped(const TTypedElement<UTypedElementWorldInterface>& InElementWorldHandle, const UE::Widget::EWidgetMode InWidgetMode, const ETypedElementViewportInteractionGizmoManipulationType InManipulationType);
+	virtual void GizmoManipulationStarted(const TTypedElement<ITypedElementWorldInterface>& InElementWorldHandle, const UE::Widget::EWidgetMode InWidgetMode);
+	virtual void GizmoManipulationDeltaUpdate(const TTypedElement<ITypedElementWorldInterface>& InElementWorldHandle, const UE::Widget::EWidgetMode InWidgetMode, const EAxisList::Type InDragAxis, const FInputDeviceState& InInputState, const FTransform& InDeltaTransform, const FVector& InPivotLocation);
+	virtual void GizmoManipulationStopped(const TTypedElement<ITypedElementWorldInterface>& InElementWorldHandle, const UE::Widget::EWidgetMode InWidgetMode, const ETypedElementViewportInteractionGizmoManipulationType InManipulationType);
 	virtual void PostGizmoManipulationStopped(TArrayView<const FTypedElementHandle> InElementHandles, const UE::Widget::EWidgetMode InWidgetMode);
-	virtual void MirrorElement(const TTypedElement<UTypedElementWorldInterface>& InElementWorldHandle, const FVector& InMirrorScale, const FVector& InPivotLocation);
-	virtual bool GetFocusBounds(const TTypedElement<UTypedElementWorldInterface>& InElementWorldHandle, FBoxSphereBounds& OutBounds);
+	virtual void MirrorElement(const TTypedElement<ITypedElementWorldInterface>& InElementWorldHandle, const FVector& InMirrorScale, const FVector& InPivotLocation);
+	virtual bool GetFocusBounds(const TTypedElement<ITypedElementWorldInterface>& InElementWorldHandle, FBoxSphereBounds& OutBounds);
 };
 
 /**
@@ -46,7 +46,7 @@ struct EDITORFRAMEWORK_API FTypedElementViewportInteractionElement
 public:
 	FTypedElementViewportInteractionElement() = default;
 
-	FTypedElementViewportInteractionElement(TTypedElement<UTypedElementWorldInterface> InElementWorldHandle, FTypedElementViewportInteractionCustomization* InViewportInteractionCustomization)
+	FTypedElementViewportInteractionElement(TTypedElement<ITypedElementWorldInterface> InElementWorldHandle, FTypedElementViewportInteractionCustomization* InViewportInteractionCustomization)
 		: ElementWorldHandle(MoveTemp(InElementWorldHandle))
 		, ViewportInteractionCustomization(InViewportInteractionCustomization)
 	{
@@ -78,7 +78,7 @@ public:
 	bool GetFocusBounds(FBoxSphereBounds& OutBounds) const { return ViewportInteractionCustomization->GetFocusBounds(ElementWorldHandle, OutBounds); }
 	
 private:
-	TTypedElement<UTypedElementWorldInterface> ElementWorldHandle;
+	TTypedElement<ITypedElementWorldInterface> ElementWorldHandle;
 	FTypedElementViewportInteractionCustomization* ViewportInteractionCustomization = nullptr;
 };
 
diff --git a/Engine/Source/Editor/EditorFramework/Public/Elements/Interfaces/TypedElementDetailsInterface.h b/Engine/Source/Editor/EditorFramework/Public/Elements/Interfaces/TypedElementDetailsInterface.h
--- a/Engine/Source/Editor/EditorFramework/Public/Elements/Interfaces/TypedElementDetailsInterface.h
+++ b/Engine/Source/Editor/EditorFramework/Public/Elements/Interfaces/TypedElementDetailsInterface.h
@@ -3,6 +3,8 @@
 #pragma once
 
 #include "Elements/Framework/TypedElementHandle.h"
+#include "UObject/Interface.h"
+
 #include "TypedElementDetailsInterface.generated.h"
 
 /**
@@ -26,8 +28,13 @@ public:
 	virtual void AddReferencedObjects(FReferenceCollector& Collector) {}
 };
 
-UCLASS(Abstract)
-class EDITORFRAMEWORK_API UTypedElementDetailsInterface : public UTypedElementInterface
+UINTERFACE(MinimalAPI, meta = (CannotImplementInterfaceInBlueprint))
+class UTypedElementDetailsInterface : public UInterface
+{
+	GENERATED_BODY()
+};
+
+class EDITORFRAMEWORK_API ITypedElementDetailsInterface
 {
 	GENERATED_BODY()
 
@@ -46,7 +53,7 @@ public:
 };
 
 template <>
-struct TTypedElement<UTypedElementDetailsInterface> : public TTypedElementBase<UTypedElementDetailsInterface>
+struct TTypedElement<ITypedElementDetailsInterface> : public TTypedElementBase<ITypedElementDetailsInterface>
 {
 	bool IsTopLevelElement() const { return InterfacePtr->IsTopLevelElement(*this); }
 	TUniquePtr<ITypedElementDetailsObject> GetDetailsObject() const { return InterfacePtr->GetDetailsObject(*this); }
diff --git a/Engine/Source/Editor/Experimental/EditorInteractiveToolsFramework/Private/EditorInteractiveGizmoSelectionBuilder.cpp b/Engine/Source/Editor/Experimental/EditorInteractiveToolsFramework/Private/EditorInteractiveGizmoSelectionBuilder.cpp
--- a/Engine/Source/Editor/Experimental/EditorInteractiveToolsFramework/Private/EditorInteractiveGizmoSelectionBuilder.cpp
+++ b/Engine/Source/Editor/Experimental/EditorInteractiveToolsFramework/Private/EditorInteractiveGizmoSelectionBuilder.cpp
@@ -18,11 +18,11 @@ UTransformProxy* FEditorGizmoSelectionBuilderHelper::CreateTransformProxyForSele
 			bool bHasSelectedElements = false;
 			UTransformProxy* TransformProxy = NewObject<UTransformProxy>();
 
-			SelectionSet->ForEachSelectedElement<UTypedElementWorldInterface>([TransformProxy, SelectionSet, &bHasSelectedElements](const TTypedElement<UTypedElementWorldInterface>& InWorldElement)
+			SelectionSet->ForEachSelectedElement<ITypedElementWorldInterface>([TransformProxy, SelectionSet, &bHasSelectedElements](const TTypedElement<ITypedElementWorldInterface>& InWorldElement)
 				{
 					if (InWorldElement.CanMoveElement(ETypedElementWorldType::Editor))
 					{
-						if (TTypedElement<UTypedElementObjectInterface> ObjectTypedElement = SelectionSet->GetElementList()->GetElement<UTypedElementObjectInterface>(InWorldElement))
+						if (TTypedElement<ITypedElementObjectInterface> ObjectTypedElement = SelectionSet->GetElementList()->GetElement<ITypedElementObjectInterface>(InWorldElement))
 						{
 							if (AActor* Actor = ObjectTypedElement.GetObjectAs<AActor>())
 							{
diff --git a/Engine/Source/Editor/LevelEditor/Private/Elements/Actor/ActorElementLevelEditorCommonActionsCustomization.cpp b/Engine/Source/Editor/LevelEditor/Private/Elements/Actor/ActorElementLevelEditorCommonActionsCustomization.cpp
--- a/Engine/Source/Editor/LevelEditor/Private/Elements/Actor/ActorElementLevelEditorCommonActionsCustomization.cpp
+++ b/Engine/Source/Editor/LevelEditor/Private/Elements/Actor/ActorElementLevelEditorCommonActionsCustomization.cpp
@@ -9,7 +9,7 @@
 #include "EditorModeManager.h"
 #include "Toolkits/IToolkitHost.h"
 
-bool FActorElementLevelEditorCommonActionsCustomization::DeleteElements(UTypedElementWorldInterface* InWorldInterface, TArrayView<const FTypedElementHandle> InElementHandles, UWorld* InWorld, UTypedElementSelectionSet* InSelectionSet, const FTypedElementDeletionOptions& InDeletionOptions)
+bool FActorElementLevelEditorCommonActionsCustomization::DeleteElements(ITypedElementWorldInterface* InWorldInterface, TArrayView<const FTypedElementHandle> InElementHandles, UWorld* InWorld, UTypedElementSelectionSet* InSelectionSet, const FTypedElementDeletionOptions& InDeletionOptions)
 {
 	if (const IToolkitHost* ToolkitHostPtr = GetToolkitHost())
 	{
@@ -23,7 +23,7 @@ bool FActorElementLevelEditorCommonActionsCustomization::DeleteElements(UTypedEl
 	return FTypedElementCommonActionsCustomization::DeleteElements(InWorldInterface, InElementHandles, InWorld, InSelectionSet, InDeletionOptions);
 }
 
-void FActorElementLevelEditorCommonActionsCustomization::DuplicateElements(UTypedElementWorldInterface* InWorldInterface, TArrayView<const FTypedElementHandle> InElementHandles, UWorld* InWorld, const FVector& InLocationOffset, TArray<FTypedElementHandle>& OutNewElements)
+void FActorElementLevelEditorCommonActionsCustomization::DuplicateElements(ITypedElementWorldInterface* InWorldInterface, TArrayView<const FTypedElementHandle> InElementHandles, UWorld* InWorld, const FVector& InLocationOffset, TArray<FTypedElementHandle>& OutNewElements)
 {
 	if (const IToolkitHost* ToolkitHostPtr = GetToolkitHost())
 	{
diff --git a/Engine/Source/Editor/LevelEditor/Private/Elements/Actor/ActorElementLevelEditorCommonActionsCustomization.h b/Engine/Source/Editor/LevelEditor/Private/Elements/Actor/ActorElementLevelEditorCommonActionsCustomization.h
--- a/Engine/Source/Editor/LevelEditor/Private/Elements/Actor/ActorElementLevelEditorCommonActionsCustomization.h
+++ b/Engine/Source/Editor/LevelEditor/Private/Elements/Actor/ActorElementLevelEditorCommonActionsCustomization.h
@@ -8,6 +8,6 @@
 class FActorElementLevelEditorCommonActionsCustomization : public FTypedElementCommonActionsCustomization, public FTypedElementAssetEditorToolkitHostMixin
 {
 public:
-	virtual bool DeleteElements(UTypedElementWorldInterface* InWorldInterface, TArrayView<const FTypedElementHandle> InElementHandles, UWorld* InWorld, UTypedElementSelectionSet* InSelectionSet, const FTypedElementDeletionOptions& InDeletionOptions) override;
-	virtual void DuplicateElements(UTypedElementWorldInterface* InWorldInterface, TArrayView<const FTypedElementHandle> InElementHandles, UWorld* InWorld, const FVector& InLocationOffset, TArray<FTypedElementHandle>& OutNewElements) override;
+	virtual bool DeleteElements(ITypedElementWorldInterface* InWorldInterface, TArrayView<const FTypedElementHandle> InElementHandles, UWorld* InWorld, UTypedElementSelectionSet* InSelectionSet, const FTypedElementDeletionOptions& InDeletionOptions) override;
+	virtual void DuplicateElements(ITypedElementWorldInterface* InWorldInterface, TArrayView<const FTypedElementHandle> InElementHandles, UWorld* InWorld, const FVector& InLocationOffset, TArray<FTypedElementHandle>& OutNewElements) override;
 };
diff --git a/Engine/Source/Editor/LevelEditor/Private/Elements/Actor/ActorElementLevelEditorSelectionCustomization.cpp b/Engine/Source/Editor/LevelEditor/Private/Elements/Actor/ActorElementLevelEditorSelectionCustomization.cpp
--- a/Engine/Source/Editor/LevelEditor/Private/Elements/Actor/ActorElementLevelEditorSelectionCustomization.cpp
+++ b/Engine/Source/Editor/LevelEditor/Private/Elements/Actor/ActorElementLevelEditorSelectionCustomization.cpp
@@ -20,33 +20,33 @@
 
 DEFINE_LOG_CATEGORY_STATIC(LogActorLevelEditorSelection, Log, All);
 
-bool FActorElementLevelEditorSelectionCustomization::CanSelectElement(const TTypedElement<UTypedElementSelectionInterface>& InElementSelectionHandle, const FTypedElementSelectionOptions& InSelectionOptions)
+bool FActorElementLevelEditorSelectionCustomization::CanSelectElement(const TTypedElement<ITypedElementSelectionInterface>& InElementSelectionHandle, const FTypedElementSelectionOptions& InSelectionOptions)
 {
 	return CanSelectActorElement(InElementSelectionHandle, InSelectionOptions);
 }
 
-bool FActorElementLevelEditorSelectionCustomization::CanDeselectElement(const TTypedElement<UTypedElementSelectionInterface>& InElementSelectionHandle, const FTypedElementSelectionOptions& InSelectionOptions)
+bool FActorElementLevelEditorSelectionCustomization::CanDeselectElement(const TTypedElement<ITypedElementSelectionInterface>& InElementSelectionHandle, const FTypedElementSelectionOptions& InSelectionOptions)
 {
 	return CanDeselectActorElement(InElementSelectionHandle, InSelectionOptions);
 }
 
-bool FActorElementLevelEditorSelectionCustomization::SelectElement(const TTypedElement<UTypedElementSelectionInterface>& InElementSelectionHandle, FTypedElementListRef InSelectionSet, const FTypedElementSelectionOptions& InSelectionOptions)
+bool FActorElementLevelEditorSelectionCustomization::SelectElement(const TTypedElement<ITypedElementSelectionInterface>& InElementSelectionHandle, FTypedElementListRef InSelectionSet, const FTypedElementSelectionOptions& InSelectionOptions)
 {
 	return SelectActorElement(InElementSelectionHandle, InSelectionSet, InSelectionOptions);
 }
 
-bool FActorElementLevelEditorSelectionCustomization::DeselectElement(const TTypedElement<UTypedElementSelectionInterface>& InElementSelectionHandle, FTypedElementListRef InSelectionSet, const FTypedElementSelectionOptions& InSelectionOptions)
+bool FActorElementLevelEditorSelectionCustomization::DeselectElement(const TTypedElement<ITypedElementSelectionInterface>& InElementSelectionHandle, FTypedElementListRef InSelectionSet, const FTypedElementSelectionOptions& InSelectionOptions)
 {
 	return DeselectActorElement(InElementSelectionHandle, InSelectionSet, InSelectionOptions);
 }
 
-bool FActorElementLevelEditorSelectionCustomization::AllowSelectionModifiers(const TTypedElement<UTypedElementSelectionInterface>& InElementSelectionHandle, FTypedElementListConstRef InSelectionSet)
+bool FActorElementLevelEditorSelectionCustomization::AllowSelectionModifiers(const TTypedElement<ITypedElementSelectionInterface>& InElementSelectionHandle, FTypedElementListConstRef InSelectionSet)
 {
 	// Ctrl or Shift clicking an actor is the same as regular clicking when components are selected
 	return !InSelectionSet->HasElementsOfType(NAME_Components);
 }
 
-FTypedElementHandle FActorElementLevelEditorSelectionCustomization::GetSelectionElement(const TTypedElement<UTypedElementSelectionInterface>& InElementSelectionHandle, FTypedElementListConstRef InCurrentSelection, const ETypedElementSelectionMethod InSelectionMethod)
+FTypedElementHandle FActorElementLevelEditorSelectionCustomization::GetSelectionElement(const TTypedElement<ITypedElementSelectionInterface>& InElementSelectionHandle, FTypedElementListConstRef InCurrentSelection, const ETypedElementSelectionMethod InSelectionMethod)
 {
 	if (AActor* ConsideredActor = ActorElementDataUtil::GetActorFromHandle(InElementSelectionHandle))
 	{
@@ -59,7 +59,7 @@ FTypedElementHandle FActorElementLevelEditorSelectionCustomization::GetSelection
 	return InElementSelectionHandle;
 }
 
-void FActorElementLevelEditorSelectionCustomization::GetNormalizedElements(const TTypedElement<UTypedElementSelectionInterface>& InElementSelectionHandle, FTypedElementListConstRef InSelectionSet, const FTypedElementSelectionNormalizationOptions& InNormalizationOptions, FTypedElementListRef OutNormalizedElements)
+void FActorElementLevelEditorSelectionCustomization::GetNormalizedElements(const TTypedElement<ITypedElementSelectionInterface>& InElementSelectionHandle, FTypedElementListConstRef InSelectionSet, const FTypedElementSelectionNormalizationOptions& InNormalizationOptions, FTypedElementListRef OutNormalizedElements)
 {
 	AActor* Actor = ActorElementDataUtil::GetActorFromHandleChecked(InElementSelectionHandle);
 
@@ -73,7 +73,7 @@ void FActorElementLevelEditorSelectionCustomization::GetNormalizedElements(const
 	FActorElementLevelEditorSelectionCustomization::AppendNormalizedActors(Actor, InSelectionSet, InNormalizationOptions, OutNormalizedElements);
 }
 
-bool FActorElementLevelEditorSelectionCustomization::CanSelectActorElement(const TTypedElement<UTypedElementSelectionInterface>& InActorSelectionHandle, const FTypedElementSelectionOptions& InSelectionOptions) const
+bool FActorElementLevelEditorSelectionCustomization::CanSelectActorElement(const TTypedElement<ITypedElementSelectionInterface>& InActorSelectionHandle, const FTypedElementSelectionOptions& InSelectionOptions) const
 {
 	AActor* Actor = ActorElementDataUtil::GetActorFromHandleChecked(InActorSelectionHandle);
 
@@ -124,7 +124,7 @@ bool FActorElementLevelEditorSelectionCustomization::CanSelectActorElement(const
 	return true;
 }
 
-bool FActorElementLevelEditorSelectionCustomization::CanDeselectActorElement(const TTypedElement<UTypedElementSelectionInterface>& InActorSelectionHandle, const FTypedElementSelectionOptions& InSelectionOptions) const
+bool FActorElementLevelEditorSelectionCustomization::CanDeselectActorElement(const TTypedElement<ITypedElementSelectionInterface>& InActorSelectionHandle, const FTypedElementSelectionOptions& InSelectionOptions) const
 {
 	AActor* Actor = ActorElementDataUtil::GetActorFromHandleChecked(InActorSelectionHandle);
 
@@ -143,7 +143,7 @@ bool FActorElementLevelEditorSelectionCustomization::CanDeselectActorElement(con
 	return true;
 }
 
-bool FActorElementLevelEditorSelectionCustomization::SelectActorElement(const TTypedElement<UTypedElementSelectionInterface>& InActorSelectionHandle, FTypedElementListRef InSelectionSet, const FTypedElementSelectionOptions& InSelectionOptions)
+bool FActorElementLevelEditorSelectionCustomization::SelectActorElement(const TTypedElement<ITypedElementSelectionInterface>& InActorSelectionHandle, FTypedElementListRef InSelectionSet, const FTypedElementSelectionOptions& InSelectionOptions)
 {
 	AActor* Actor = ActorElementDataUtil::GetActorFromHandleChecked(InActorSelectionHandle);
 
@@ -181,7 +181,7 @@ bool FActorElementLevelEditorSelectionCustomization::SelectActorElement(const TT
 
 	// Select the desired actor
 	{
-		TTypedElement<UTypedElementSelectionInterface> ActorSelectionHandle = InSelectionSet->GetElement<UTypedElementSelectionInterface>(UEngineElementsLibrary::AcquireEditorActorElementHandle(Actor));
+		TTypedElement<ITypedElementSelectionInterface> ActorSelectionHandle = InSelectionSet->GetElement<ITypedElementSelectionInterface>(UEngineElementsLibrary::AcquireEditorActorElementHandle(Actor));
 		if (!ActorSelectionHandle.SelectElement(InSelectionSet, InSelectionOptions))
 		{
 			return bSelectionChanged;
@@ -213,7 +213,7 @@ bool FActorElementLevelEditorSelectionCustomization::SelectActorElement(const TT
 	return true;
 }
 
-bool FActorElementLevelEditorSelectionCustomization::DeselectActorElement(const TTypedElement<UTypedElementSelectionInterface>& InActorSelectionHandle, FTypedElementListRef InSelectionSet, const FTypedElementSelectionOptions& InSelectionOptions)
+bool FActorElementLevelEditorSelectionCustomization::DeselectActorElement(const TTypedElement<ITypedElementSelectionInterface>& InActorSelectionHandle, FTypedElementListRef InSelectionSet, const FTypedElementSelectionOptions& InSelectionOptions)
 {
 	AActor* Actor = ActorElementDataUtil::GetActorFromHandleChecked(InActorSelectionHandle);
 
@@ -245,7 +245,7 @@ bool FActorElementLevelEditorSelectionCustomization::DeselectActorElement(const
 
 	// Deselect the desired actor
 	{
-		TTypedElement<UTypedElementSelectionInterface> ActorSelectionHandle = InSelectionSet->GetElement<UTypedElementSelectionInterface>(UEngineElementsLibrary::AcquireEditorActorElementHandle(Actor));
+		TTypedElement<ITypedElementSelectionInterface> ActorSelectionHandle = InSelectionSet->GetElement<ITypedElementSelectionInterface>(UEngineElementsLibrary::AcquireEditorActorElementHandle(Actor));
 		if (!ActorSelectionHandle.DeselectElement(InSelectionSet, InSelectionOptions))
 		{
 			return bSelectionChanged;
@@ -265,7 +265,7 @@ bool FActorElementLevelEditorSelectionCustomization::DeselectActorElement(const
 				continue;
 			}
 
-			TTypedElement<UTypedElementSelectionInterface> ComponentSelectionHandle = InSelectionSet->GetElement<UTypedElementSelectionInterface>(UEngineElementsLibrary::AcquireEditorComponentElementHandle(Component));
+			TTypedElement<ITypedElementSelectionInterface> ComponentSelectionHandle = InSelectionSet->GetElement<ITypedElementSelectionInterface>(UEngineElementsLibrary::AcquireEditorComponentElementHandle(Component));
 			ComponentSelectionHandle.DeselectElement(InSelectionSet, InSelectionOptions);
 
 			if (USceneComponent* SceneComponent = Cast<USceneComponent>(Component))
@@ -288,7 +288,7 @@ bool FActorElementLevelEditorSelectionCustomization::SelectActorGroup(AGroupActo
 {
 	bool bSelectionChanged = false;
 
-	TTypedElement<UTypedElementSelectionInterface> GroupSelectionHandle = InSelectionSet->GetElement<UTypedElementSelectionInterface>(UEngineElementsLibrary::AcquireEditorActorElementHandle(InGroupActor));
+	TTypedElement<ITypedElementSelectionInterface> GroupSelectionHandle = InSelectionSet->GetElement<ITypedElementSelectionInterface>(UEngineElementsLibrary::AcquireEditorActorElementHandle(InGroupActor));
 
 	const FTypedElementSelectionOptions GroupSelectionOptions = FTypedElementSelectionOptions(InSelectionOptions)
 		.SetAllowGroups(false);
@@ -305,7 +305,7 @@ bool FActorElementLevelEditorSelectionCustomization::SelectActorGroup(AGroupActo
 		InGroupActor->GetGroupActors(GroupActors);
 		for (AActor* Actor : GroupActors)
 		{
-			TTypedElement<UTypedElementSelectionInterface> ActorSelectionHandle = InSelectionSet->GetElement<UTypedElementSelectionInterface>(UEngineElementsLibrary::AcquireEditorActorElementHandle(Actor));
+			TTypedElement<ITypedElementSelectionInterface> ActorSelectionHandle = InSelectionSet->GetElement<ITypedElementSelectionInterface>(UEngineElementsLibrary::AcquireEditorActorElementHandle(Actor));
 			if (CanSelectActorElement(ActorSelectionHandle, GroupSelectionOptions))
 			{
 				bSelectionChanged |= SelectActorElement(ActorSelectionHandle, InSelectionSet, GroupSelectionOptions);
@@ -320,7 +320,7 @@ bool FActorElementLevelEditorSelectionCustomization::DeselectActorGroup(AGroupAc
 {
 	bool bSelectionChanged = false;
 
-	TTypedElement<UTypedElementSelectionInterface> GroupSelectionHandle = InSelectionSet->GetElement<UTypedElementSelectionInterface>(UEngineElementsLibrary::AcquireEditorActorElementHandle(InGroupActor));
+	TTypedElement<ITypedElementSelectionInterface> GroupSelectionHandle = InSelectionSet->GetElement<ITypedElementSelectionInterface>(UEngineElementsLibrary::AcquireEditorActorElementHandle(InGroupActor));
 
 	const FTypedElementSelectionOptions GroupSelectionOptions = FTypedElementSelectionOptions(InSelectionOptions)
 		.SetAllowGroups(false);
@@ -337,7 +337,7 @@ bool FActorElementLevelEditorSelectionCustomization::DeselectActorGroup(AGroupAc
 		InGroupActor->GetGroupActors(GroupActors);
 		for (AActor* Actor : GroupActors)
 		{
-			TTypedElement<UTypedElementSelectionInterface> ActorSelectionHandle = InSelectionSet->GetElement<UTypedElementSelectionInterface>(UEngineElementsLibrary::AcquireEditorActorElementHandle(Actor));
+			TTypedElement<ITypedElementSelectionInterface> ActorSelectionHandle = InSelectionSet->GetElement<ITypedElementSelectionInterface>(UEngineElementsLibrary::AcquireEditorActorElementHandle(Actor));
 			if (CanDeselectActorElement(ActorSelectionHandle, GroupSelectionOptions))
 			{
 				bSelectionChanged |= DeselectActorElement(ActorSelectionHandle, InSelectionSet, GroupSelectionOptions);
diff --git a/Engine/Source/Editor/LevelEditor/Private/Elements/Actor/ActorElementLevelEditorSelectionCustomization.h b/Engine/Source/Editor/LevelEditor/Private/Elements/Actor/ActorElementLevelEditorSelectionCustomization.h
--- a/Engine/Source/Editor/LevelEditor/Private/Elements/Actor/ActorElementLevelEditorSelectionCustomization.h
+++ b/Engine/Source/Editor/LevelEditor/Private/Elements/Actor/ActorElementLevelEditorSelectionCustomization.h
@@ -10,19 +10,19 @@ class AGroupActor;
 class FActorElementLevelEditorSelectionCustomization : public FTypedElementSelectionCustomization, public FTypedElementAssetEditorToolkitHostMixin
 {
 public:
-	virtual bool CanSelectElement(const TTypedElement<UTypedElementSelectionInterface>& InElementSelectionHandle, const FTypedElementSelectionOptions& InSelectionOptions) override;
-	virtual bool CanDeselectElement(const TTypedElement<UTypedElementSelectionInterface>& InElementSelectionHandle, const FTypedElementSelectionOptions& InSelectionOptions) override;
-	virtual bool SelectElement(const TTypedElement<UTypedElementSelectionInterface>& InElementSelectionHandle, FTypedElementListRef InSelectionSet, const FTypedElementSelectionOptions& InSelectionOptions) override;
-	virtual bool DeselectElement(const TTypedElement<UTypedElementSelectionInterface>& InElementSelectionHandle, FTypedElementListRef InSelectionSet, const FTypedElementSelectionOptions& InSelectionOptions) override;
-	virtual bool AllowSelectionModifiers(const TTypedElement<UTypedElementSelectionInterface>& InElementSelectionHandle, FTypedElementListConstRef InSelectionSet) override;
-	virtual FTypedElementHandle GetSelectionElement(const TTypedElement<UTypedElementSelectionInterface>& InElementSelectionHandle, FTypedElementListConstRef InCurrentSelection, const ETypedElementSelectionMethod InSelectionMethod) override;
-	virtual void GetNormalizedElements(const TTypedElement<UTypedElementSelectionInterface>& InElementSelectionHandle, FTypedElementListConstRef InSelectionSet, const FTypedElementSelectionNormalizationOptions& InNormalizationOptions, FTypedElementListRef OutNormalizedElements) override;
-
-	bool CanSelectActorElement(const TTypedElement<UTypedElementSelectionInterface>& InActorSelectionHandle, const FTypedElementSelectionOptions& InSelectionOptions) const;
-	bool CanDeselectActorElement(const TTypedElement<UTypedElementSelectionInterface>& InActorSelectionHandle, const FTypedElementSelectionOptions& InSelectionOptions) const;
-
-	bool SelectActorElement(const TTypedElement<UTypedElementSelectionInterface>& InActorSelectionHandle, FTypedElementListRef InSelectionSet, const FTypedElementSelectionOptions& InSelectionOptions);
-	bool DeselectActorElement(const TTypedElement<UTypedElementSelectionInterface>& InActorSelectionHandle, FTypedElementListRef InSelectionSet, const FTypedElementSelectionOptions& InSelectionOptions);
+	virtual bool CanSelectElement(const TTypedElement<ITypedElementSelectionInterface>& InElementSelectionHandle, const FTypedElementSelectionOptions& InSelectionOptions) override;
+	virtual bool CanDeselectElement(const TTypedElement<ITypedElementSelectionInterface>& InElementSelectionHandle, const FTypedElementSelectionOptions& InSelectionOptions) override;
+	virtual bool SelectElement(const TTypedElement<ITypedElementSelectionInterface>& InElementSelectionHandle, FTypedElementListRef InSelectionSet, const FTypedElementSelectionOptions& InSelectionOptions) override;
+	virtual bool DeselectElement(const TTypedElement<ITypedElementSelectionInterface>& InElementSelectionHandle, FTypedElementListRef InSelectionSet, const FTypedElementSelectionOptions& InSelectionOptions) override;
+	virtual bool AllowSelectionModifiers(const TTypedElement<ITypedElementSelectionInterface>& InElementSelectionHandle, FTypedElementListConstRef InSelectionSet) override;
+	virtual FTypedElementHandle GetSelectionElement(const TTypedElement<ITypedElementSelectionInterface>& InElementSelectionHandle, FTypedElementListConstRef InCurrentSelection, const ETypedElementSelectionMethod InSelectionMethod) override;
+	virtual void GetNormalizedElements(const TTypedElement<ITypedElementSelectionInterface>& InElementSelectionHandle, FTypedElementListConstRef InSelectionSet, const FTypedElementSelectionNormalizationOptions& InNormalizationOptions, FTypedElementListRef OutNormalizedElements) override;
+
+	bool CanSelectActorElement(const TTypedElement<ITypedElementSelectionInterface>& InActorSelectionHandle, const FTypedElementSelectionOptions& InSelectionOptions) const;
+	bool CanDeselectActorElement(const TTypedElement<ITypedElementSelectionInterface>& InActorSelectionHandle, const FTypedElementSelectionOptions& InSelectionOptions) const;
+
+	bool SelectActorElement(const TTypedElement<ITypedElementSelectionInterface>& InActorSelectionHandle, FTypedElementListRef InSelectionSet, const FTypedElementSelectionOptions& InSelectionOptions);
+	bool DeselectActorElement(const TTypedElement<ITypedElementSelectionInterface>& InActorSelectionHandle, FTypedElementListRef InSelectionSet, const FTypedElementSelectionOptions& InSelectionOptions);
 
 	bool SelectActorGroup(AGroupActor* InGroupActor, FTypedElementListRef InSelectionSet, const FTypedElementSelectionOptions& InSelectionOptions, const bool bForce);
 	bool DeselectActorGroup(AGroupActor* InGroupActor, FTypedElementListRef InSelectionSet, const FTypedElementSelectionOptions& InSelectionOptions, const bool bForce);
diff --git a/Engine/Source/Editor/LevelEditor/Private/Elements/Actor/ActorElementLevelEditorViewportInteractionCustomization.cpp b/Engine/Source/Editor/LevelEditor/Private/Elements/Actor/ActorElementLevelEditorViewportInteractionCustomization.cpp
--- a/Engine/Source/Editor/LevelEditor/Private/Elements/Actor/ActorElementLevelEditorViewportInteractionCustomization.cpp
+++ b/Engine/Source/Editor/LevelEditor/Private/Elements/Actor/ActorElementLevelEditorViewportInteractionCustomization.cpp
@@ -14,7 +14,7 @@
 #include "ActorGroupingUtils.h"
 #include "LevelEditorViewport.h"
 
-void FActorElementLevelEditorViewportInteractionCustomization::GizmoManipulationStarted(const TTypedElement<UTypedElementWorldInterface>& InElementWorldHandle, const UE::Widget::EWidgetMode InWidgetMode)
+void FActorElementLevelEditorViewportInteractionCustomization::GizmoManipulationStarted(const TTypedElement<ITypedElementWorldInterface>& InElementWorldHandle, const UE::Widget::EWidgetMode InWidgetMode)
 {
 	AActor* Actor = ActorElementDataUtil::GetActorFromHandleChecked(InElementWorldHandle);
 
@@ -30,7 +30,7 @@ void FActorElementLevelEditorViewportInteractionCustomization::GizmoManipulation
 	}
 }
 
-void FActorElementLevelEditorViewportInteractionCustomization::GizmoManipulationDeltaUpdate(const TTypedElement<UTypedElementWorldInterface>& InElementWorldHandle, const UE::Widget::EWidgetMode InWidgetMode, const EAxisList::Type InDragAxis, const FInputDeviceState& InInputState, const FTransform& InDeltaTransform, const FVector& InPivotLocation)
+void FActorElementLevelEditorViewportInteractionCustomization::GizmoManipulationDeltaUpdate(const TTypedElement<ITypedElementWorldInterface>& InElementWorldHandle, const UE::Widget::EWidgetMode InWidgetMode, const EAxisList::Type InDragAxis, const FInputDeviceState& InInputState, const FTransform& InDeltaTransform, const FVector& InPivotLocation)
 {
 	AActor* Actor = ActorElementDataUtil::GetActorFromHandleChecked(InElementWorldHandle);
 
@@ -61,7 +61,7 @@ void FActorElementLevelEditorViewportInteractionCustomization::GizmoManipulation
 	GetMutableLevelEditorViewportClient()->UpdateLockedActorViewports(Actor, true);
 }
 
-void FActorElementLevelEditorViewportInteractionCustomization::GizmoManipulationStopped(const TTypedElement<UTypedElementWorldInterface>& InElementWorldHandle, const UE::Widget::EWidgetMode InWidgetMode, const ETypedElementViewportInteractionGizmoManipulationType InManipulationType)
+void FActorElementLevelEditorViewportInteractionCustomization::GizmoManipulationStopped(const TTypedElement<ITypedElementWorldInterface>& InElementWorldHandle, const UE::Widget::EWidgetMode InWidgetMode, const ETypedElementViewportInteractionGizmoManipulationType InManipulationType)
 {
 	AActor* Actor = ActorElementDataUtil::GetActorFromHandleChecked(InElementWorldHandle);
 
diff --git a/Engine/Source/Editor/LevelEditor/Private/Elements/Component/ComponentElementLevelEditorCommonActionsCustomization.cpp b/Engine/Source/Editor/LevelEditor/Private/Elements/Component/ComponentElementLevelEditorCommonActionsCustomization.cpp
--- a/Engine/Source/Editor/LevelEditor/Private/Elements/Component/ComponentElementLevelEditorCommonActionsCustomization.cpp
+++ b/Engine/Source/Editor/LevelEditor/Private/Elements/Component/ComponentElementLevelEditorCommonActionsCustomization.cpp
@@ -7,7 +7,7 @@
 #include "EditorModeManager.h"
 #include "Toolkits/IToolkitHost.h"
 
-bool FComponentElementLevelEditorCommonActionsCustomization::DeleteElements(UTypedElementWorldInterface* InWorldInterface, TArrayView<const FTypedElementHandle> InElementHandles, UWorld* InWorld, UTypedElementSelectionSet* InSelectionSet, const FTypedElementDeletionOptions& InDeletionOptions)
+bool FComponentElementLevelEditorCommonActionsCustomization::DeleteElements(ITypedElementWorldInterface* InWorldInterface, TArrayView<const FTypedElementHandle> InElementHandles, UWorld* InWorld, UTypedElementSelectionSet* InSelectionSet, const FTypedElementDeletionOptions& InDeletionOptions)
 {
 	if (const IToolkitHost* ToolkitHostPtr = GetToolkitHost())
 	{
@@ -21,7 +21,7 @@ bool FComponentElementLevelEditorCommonActionsCustomization::DeleteElements(UTyp
 	return FTypedElementCommonActionsCustomization::DeleteElements(InWorldInterface, InElementHandles, InWorld, InSelectionSet, InDeletionOptions);
 }
 
-void FComponentElementLevelEditorCommonActionsCustomization::DuplicateElements(UTypedElementWorldInterface* InWorldInterface, TArrayView<const FTypedElementHandle> InElementHandles, UWorld* InWorld, const FVector& InLocationOffset, TArray<FTypedElementHandle>& OutNewElements)
+void FComponentElementLevelEditorCommonActionsCustomization::DuplicateElements(ITypedElementWorldInterface* InWorldInterface, TArrayView<const FTypedElementHandle> InElementHandles, UWorld* InWorld, const FVector& InLocationOffset, TArray<FTypedElementHandle>& OutNewElements)
 {
 	if (const IToolkitHost* ToolkitHostPtr = GetToolkitHost())
 	{
diff --git a/Engine/Source/Editor/LevelEditor/Private/Elements/Component/ComponentElementLevelEditorCommonActionsCustomization.h b/Engine/Source/Editor/LevelEditor/Private/Elements/Component/ComponentElementLevelEditorCommonActionsCustomization.h
--- a/Engine/Source/Editor/LevelEditor/Private/Elements/Component/ComponentElementLevelEditorCommonActionsCustomization.h
+++ b/Engine/Source/Editor/LevelEditor/Private/Elements/Component/ComponentElementLevelEditorCommonActionsCustomization.h
@@ -8,6 +8,6 @@
 class FComponentElementLevelEditorCommonActionsCustomization : public FTypedElementCommonActionsCustomization, public FTypedElementAssetEditorToolkitHostMixin
 {
 public:
-	virtual bool DeleteElements(UTypedElementWorldInterface* InWorldInterface, TArrayView<const FTypedElementHandle> InElementHandles, UWorld* InWorld, UTypedElementSelectionSet* InSelectionSet, const FTypedElementDeletionOptions& InDeletionOptions) override;
-	virtual void DuplicateElements(UTypedElementWorldInterface* InWorldInterface, TArrayView<const FTypedElementHandle> InElementHandles, UWorld* InWorld, const FVector& InLocationOffset, TArray<FTypedElementHandle>& OutNewElements) override;
+	virtual bool DeleteElements(ITypedElementWorldInterface* InWorldInterface, TArrayView<const FTypedElementHandle> InElementHandles, UWorld* InWorld, UTypedElementSelectionSet* InSelectionSet, const FTypedElementDeletionOptions& InDeletionOptions) override;
+	virtual void DuplicateElements(ITypedElementWorldInterface* InWorldInterface, TArrayView<const FTypedElementHandle> InElementHandles, UWorld* InWorld, const FVector& InLocationOffset, TArray<FTypedElementHandle>& OutNewElements) override;
 };
diff --git a/Engine/Source/Editor/LevelEditor/Private/Elements/Component/ComponentElementLevelEditorSelectionCustomization.cpp b/Engine/Source/Editor/LevelEditor/Private/Elements/Component/ComponentElementLevelEditorSelectionCustomization.cpp
--- a/Engine/Source/Editor/LevelEditor/Private/Elements/Component/ComponentElementLevelEditorSelectionCustomization.cpp
+++ b/Engine/Source/Editor/LevelEditor/Private/Elements/Component/ComponentElementLevelEditorSelectionCustomization.cpp
@@ -16,27 +16,27 @@
 
 DEFINE_LOG_CATEGORY_STATIC(LogComponentLevelEditorSelection, Log, All);
 
-bool FComponentElementLevelEditorSelectionCustomization::CanSelectElement(const TTypedElement<UTypedElementSelectionInterface>& InElementSelectionHandle, const FTypedElementSelectionOptions& InSelectionOptions)
+bool FComponentElementLevelEditorSelectionCustomization::CanSelectElement(const TTypedElement<ITypedElementSelectionInterface>& InElementSelectionHandle, const FTypedElementSelectionOptions& InSelectionOptions)
 {
 	return CanSelectComponentElement(InElementSelectionHandle, InSelectionOptions);
 }
 
-bool FComponentElementLevelEditorSelectionCustomization::CanDeselectElement(const TTypedElement<UTypedElementSelectionInterface>& InElementSelectionHandle, const FTypedElementSelectionOptions& InSelectionOptions)
+bool FComponentElementLevelEditorSelectionCustomization::CanDeselectElement(const TTypedElement<ITypedElementSelectionInterface>& InElementSelectionHandle, const FTypedElementSelectionOptions& InSelectionOptions)
 {
 	return CanDeselectComponentElement(InElementSelectionHandle, InSelectionOptions);
 }
 
-bool FComponentElementLevelEditorSelectionCustomization::SelectElement(const TTypedElement<UTypedElementSelectionInterface>& InElementSelectionHandle, FTypedElementListRef InSelectionSet, const FTypedElementSelectionOptions& InSelectionOptions)
+bool FComponentElementLevelEditorSelectionCustomization::SelectElement(const TTypedElement<ITypedElementSelectionInterface>& InElementSelectionHandle, FTypedElementListRef InSelectionSet, const FTypedElementSelectionOptions& InSelectionOptions)
 {
 	return SelectComponentElement(InElementSelectionHandle, InSelectionSet, InSelectionOptions);
 }
 
-bool FComponentElementLevelEditorSelectionCustomization::DeselectElement(const TTypedElement<UTypedElementSelectionInterface>& InElementSelectionHandle, FTypedElementListRef InSelectionSet, const FTypedElementSelectionOptions& InSelectionOptions)
+bool FComponentElementLevelEditorSelectionCustomization::DeselectElement(const TTypedElement<ITypedElementSelectionInterface>& InElementSelectionHandle, FTypedElementListRef InSelectionSet, const FTypedElementSelectionOptions& InSelectionOptions)
 {
 	return DeselectComponentElement(InElementSelectionHandle, InSelectionSet, InSelectionOptions);
 }
 
-FTypedElementHandle FComponentElementLevelEditorSelectionCustomization::GetSelectionElement(const TTypedElement<UTypedElementSelectionInterface>& InElementSelectionHandle, FTypedElementListConstRef InCurrentSelection, const ETypedElementSelectionMethod InSelectionMethod)
+FTypedElementHandle FComponentElementLevelEditorSelectionCustomization::GetSelectionElement(const TTypedElement<ITypedElementSelectionInterface>& InElementSelectionHandle, FTypedElementListConstRef InCurrentSelection, const ETypedElementSelectionMethod InSelectionMethod)
 {
 	if (const UActorComponent* Component = ComponentElementDataUtil::GetComponentFromHandle(InElementSelectionHandle))
 	{
@@ -83,7 +83,7 @@ FTypedElementHandle FComponentElementLevelEditorSelectionCustomization::GetSelec
 	return InElementSelectionHandle;
 }
 
-void FComponentElementLevelEditorSelectionCustomization::GetNormalizedElements(const TTypedElement<UTypedElementSelectionInterface>& InElementSelectionHandle, FTypedElementListConstRef InSelectionSet, const FTypedElementSelectionNormalizationOptions& InNormalizationOptions, FTypedElementListRef OutNormalizedElements)
+void FComponentElementLevelEditorSelectionCustomization::GetNormalizedElements(const TTypedElement<ITypedElementSelectionInterface>& InElementSelectionHandle, FTypedElementListConstRef InSelectionSet, const FTypedElementSelectionNormalizationOptions& InNormalizationOptions, FTypedElementListRef OutNormalizedElements)
 {
 	UActorComponent* Component = ComponentElementDataUtil::GetComponentFromHandleChecked(InElementSelectionHandle);
 
@@ -126,7 +126,7 @@ void FComponentElementLevelEditorSelectionCustomization::GetNormalizedElements(c
 	OutNormalizedElements->Add(InElementSelectionHandle);
 }
 
-bool FComponentElementLevelEditorSelectionCustomization::CanSelectComponentElement(const TTypedElement<UTypedElementSelectionInterface>& InComponentSelectionHandle, const FTypedElementSelectionOptions& InSelectionOptions) const
+bool FComponentElementLevelEditorSelectionCustomization::CanSelectComponentElement(const TTypedElement<ITypedElementSelectionInterface>& InComponentSelectionHandle, const FTypedElementSelectionOptions& InSelectionOptions) const
 {
 	const UActorComponent* Component = ComponentElementDataUtil::GetComponentFromHandleChecked(InComponentSelectionHandle);
 
@@ -134,7 +134,7 @@ bool FComponentElementLevelEditorSelectionCustomization::CanSelectComponentEleme
 	return !GEdSelectionLock;
 }
 
-bool FComponentElementLevelEditorSelectionCustomization::CanDeselectComponentElement(const TTypedElement<UTypedElementSelectionInterface>& InComponentSelectionHandle, const FTypedElementSelectionOptions& InSelectionOptions) const
+bool FComponentElementLevelEditorSelectionCustomization::CanDeselectComponentElement(const TTypedElement<ITypedElementSelectionInterface>& InComponentSelectionHandle, const FTypedElementSelectionOptions& InSelectionOptions) const
 {
 	const UActorComponent* Component = ComponentElementDataUtil::GetComponentFromHandleChecked(InComponentSelectionHandle);
 
@@ -142,7 +142,7 @@ bool FComponentElementLevelEditorSelectionCustomization::CanDeselectComponentEle
 	return !GEdSelectionLock;
 }
 
-bool FComponentElementLevelEditorSelectionCustomization::SelectComponentElement(const TTypedElement<UTypedElementSelectionInterface>& InComponentSelectionHandle, FTypedElementListRef InSelectionSet, const FTypedElementSelectionOptions& InSelectionOptions)
+bool FComponentElementLevelEditorSelectionCustomization::SelectComponentElement(const TTypedElement<ITypedElementSelectionInterface>& InComponentSelectionHandle, FTypedElementListRef InSelectionSet, const FTypedElementSelectionOptions& InSelectionOptions)
 {
 	UActorComponent* Component = ComponentElementDataUtil::GetComponentFromHandleChecked(InComponentSelectionHandle);
 
@@ -162,7 +162,7 @@ bool FComponentElementLevelEditorSelectionCustomization::SelectComponentElement(
 	if (AActor* ComponentOwner = Component->GetOwner())
 	{
 		// Selecting a component requires that its owner actor be selected too
-		TTypedElement<UTypedElementSelectionInterface> ActorSelectionHandle = InSelectionSet->GetElement<UTypedElementSelectionInterface>(UEngineElementsLibrary::AcquireEditorActorElementHandle(ComponentOwner));
+		TTypedElement<ITypedElementSelectionInterface> ActorSelectionHandle = InSelectionSet->GetElement<ITypedElementSelectionInterface>(UEngineElementsLibrary::AcquireEditorActorElementHandle(ComponentOwner));
 		ActorSelectionHandle.SelectElement(InSelectionSet, InSelectionOptions);
 
 		// Update the selection visualization
@@ -178,7 +178,7 @@ bool FComponentElementLevelEditorSelectionCustomization::SelectComponentElement(
 	return true;
 }
 
-bool FComponentElementLevelEditorSelectionCustomization::DeselectComponentElement(const TTypedElement<UTypedElementSelectionInterface>& InComponentSelectionHandle, FTypedElementListRef InSelectionSet, const FTypedElementSelectionOptions& InSelectionOptions)
+bool FComponentElementLevelEditorSelectionCustomization::DeselectComponentElement(const TTypedElement<ITypedElementSelectionInterface>& InComponentSelectionHandle, FTypedElementListRef InSelectionSet, const FTypedElementSelectionOptions& InSelectionOptions)
 {
 	UActorComponent* Component = ComponentElementDataUtil::GetComponentFromHandleChecked(InComponentSelectionHandle);
 
@@ -192,7 +192,7 @@ bool FComponentElementLevelEditorSelectionCustomization::DeselectComponentElemen
 	// Update the selection visualization
 	if (AActor* ComponentOwner = Component->GetOwner())
 	{
-		TTypedElement<UTypedElementSelectionInterface> ActorSelectionHandle = InSelectionSet->GetElement<UTypedElementSelectionInterface>(UEngineElementsLibrary::AcquireEditorActorElementHandle(ComponentOwner));
+		TTypedElement<ITypedElementSelectionInterface> ActorSelectionHandle = InSelectionSet->GetElement<ITypedElementSelectionInterface>(UEngineElementsLibrary::AcquireEditorActorElementHandle(ComponentOwner));
 		if (!ActorSelectionHandle.IsElementSelected(InSelectionSet, FTypedElementIsSelectedOptions().SetAllowIndirect(true)))
 		{
 			// Make sure the override delegate is unbound properly
diff --git a/Engine/Source/Editor/LevelEditor/Private/Elements/Component/ComponentElementLevelEditorSelectionCustomization.h b/Engine/Source/Editor/LevelEditor/Private/Elements/Component/ComponentElementLevelEditorSelectionCustomization.h
--- a/Engine/Source/Editor/LevelEditor/Private/Elements/Component/ComponentElementLevelEditorSelectionCustomization.h
+++ b/Engine/Source/Editor/LevelEditor/Private/Elements/Component/ComponentElementLevelEditorSelectionCustomization.h
@@ -8,16 +8,16 @@
 class FComponentElementLevelEditorSelectionCustomization : public FTypedElementSelectionCustomization, public FTypedElementAssetEditorToolkitHostMixin
 {
 public:
-	virtual bool CanSelectElement(const TTypedElement<UTypedElementSelectionInterface>& InElementSelectionHandle, const FTypedElementSelectionOptions& InSelectionOptions) override;
-	virtual bool CanDeselectElement(const TTypedElement<UTypedElementSelectionInterface>& InElementSelectionHandle, const FTypedElementSelectionOptions& InSelectionOptions) override;
-	virtual bool SelectElement(const TTypedElement<UTypedElementSelectionInterface>& InElementSelectionHandle, FTypedElementListRef InSelectionSet, const FTypedElementSelectionOptions& InSelectionOptions) override;
-	virtual bool DeselectElement(const TTypedElement<UTypedElementSelectionInterface>& InElementSelectionHandle, FTypedElementListRef InSelectionSet, const FTypedElementSelectionOptions& InSelectionOptions) override;
-	virtual FTypedElementHandle GetSelectionElement(const TTypedElement<UTypedElementSelectionInterface>& InElementSelectionHandle, FTypedElementListConstRef InCurrentSelection, const ETypedElementSelectionMethod InSelectionMethod) override;
-	virtual void GetNormalizedElements(const TTypedElement<UTypedElementSelectionInterface>& InElementSelectionHandle, FTypedElementListConstRef InSelectionSet, const FTypedElementSelectionNormalizationOptions& InNormalizationOptions, FTypedElementListRef OutNormalizedElements) override;
+	virtual bool CanSelectElement(const TTypedElement<ITypedElementSelectionInterface>& InElementSelectionHandle, const FTypedElementSelectionOptions& InSelectionOptions) override;
+	virtual bool CanDeselectElement(const TTypedElement<ITypedElementSelectionInterface>& InElementSelectionHandle, const FTypedElementSelectionOptions& InSelectionOptions) override;
+	virtual bool SelectElement(const TTypedElement<ITypedElementSelectionInterface>& InElementSelectionHandle, FTypedElementListRef InSelectionSet, const FTypedElementSelectionOptions& InSelectionOptions) override;
+	virtual bool DeselectElement(const TTypedElement<ITypedElementSelectionInterface>& InElementSelectionHandle, FTypedElementListRef InSelectionSet, const FTypedElementSelectionOptions& InSelectionOptions) override;
+	virtual FTypedElementHandle GetSelectionElement(const TTypedElement<ITypedElementSelectionInterface>& InElementSelectionHandle, FTypedElementListConstRef InCurrentSelection, const ETypedElementSelectionMethod InSelectionMethod) override;
+	virtual void GetNormalizedElements(const TTypedElement<ITypedElementSelectionInterface>& InElementSelectionHandle, FTypedElementListConstRef InSelectionSet, const FTypedElementSelectionNormalizationOptions& InNormalizationOptions, FTypedElementListRef OutNormalizedElements) override;
 
-	bool CanSelectComponentElement(const TTypedElement<UTypedElementSelectionInterface>& InComponentSelectionHandle, const FTypedElementSelectionOptions& InSelectionOptions) const;
-	bool CanDeselectComponentElement(const TTypedElement<UTypedElementSelectionInterface>& InComponentSelectionHandle, const FTypedElementSelectionOptions& InSelectionOptions) const;
+	bool CanSelectComponentElement(const TTypedElement<ITypedElementSelectionInterface>& InComponentSelectionHandle, const FTypedElementSelectionOptions& InSelectionOptions) const;
+	bool CanDeselectComponentElement(const TTypedElement<ITypedElementSelectionInterface>& InComponentSelectionHandle, const FTypedElementSelectionOptions& InSelectionOptions) const;
 
-	bool SelectComponentElement(const TTypedElement<UTypedElementSelectionInterface>& InComponentSelectionHandle, FTypedElementListRef InSelectionSet, const FTypedElementSelectionOptions& InSelectionOptions);
-	bool DeselectComponentElement(const TTypedElement<UTypedElementSelectionInterface>& InComponentSelectionHandle, FTypedElementListRef InSelectionSet, const FTypedElementSelectionOptions& InSelectionOptions);
+	bool SelectComponentElement(const TTypedElement<ITypedElementSelectionInterface>& InComponentSelectionHandle, FTypedElementListRef InSelectionSet, const FTypedElementSelectionOptions& InSelectionOptions);
+	bool DeselectComponentElement(const TTypedElement<ITypedElementSelectionInterface>& InComponentSelectionHandle, FTypedElementListRef InSelectionSet, const FTypedElementSelectionOptions& InSelectionOptions);
 };
diff --git a/Engine/Source/Editor/LevelEditor/Private/Elements/Component/ComponentElementLevelEditorViewportInteractionCustomization.cpp b/Engine/Source/Editor/LevelEditor/Private/Elements/Component/ComponentElementLevelEditorViewportInteractionCustomization.cpp
--- a/Engine/Source/Editor/LevelEditor/Private/Elements/Component/ComponentElementLevelEditorViewportInteractionCustomization.cpp
+++ b/Engine/Source/Editor/LevelEditor/Private/Elements/Component/ComponentElementLevelEditorViewportInteractionCustomization.cpp
@@ -12,7 +12,7 @@
 #include "LevelEditorViewport.h"
 #include "Kismet2/ComponentEditorUtils.h"
 
-void FComponentElementLevelEditorViewportInteractionCustomization::GizmoManipulationStarted(const TTypedElement<UTypedElementWorldInterface>& InElementWorldHandle, const UE::Widget::EWidgetMode InWidgetMode)
+void FComponentElementLevelEditorViewportInteractionCustomization::GizmoManipulationStarted(const TTypedElement<ITypedElementWorldInterface>& InElementWorldHandle, const UE::Widget::EWidgetMode InWidgetMode)
 {
 	UActorComponent* Component = ComponentElementDataUtil::GetComponentFromHandleChecked(InElementWorldHandle);
 
@@ -31,7 +31,7 @@ void FComponentElementLevelEditorViewportInteractionCustomization::GizmoManipula
 	}
 }
 
-void FComponentElementLevelEditorViewportInteractionCustomization::GizmoManipulationDeltaUpdate(const TTypedElement<UTypedElementWorldInterface>& InElementWorldHandle, const UE::Widget::EWidgetMode InWidgetMode, const EAxisList::Type InDragAxis, const FInputDeviceState& InInputState, const FTransform& InDeltaTransform, const FVector& InPivotLocation)
+void FComponentElementLevelEditorViewportInteractionCustomization::GizmoManipulationDeltaUpdate(const TTypedElement<ITypedElementWorldInterface>& InElementWorldHandle, const UE::Widget::EWidgetMode InWidgetMode, const EAxisList::Type InDragAxis, const FInputDeviceState& InInputState, const FTransform& InDeltaTransform, const FVector& InPivotLocation)
 {
 	UActorComponent* Component = ComponentElementDataUtil::GetComponentFromHandleChecked(InElementWorldHandle);
 
@@ -66,7 +66,7 @@ void FComponentElementLevelEditorViewportInteractionCustomization::GizmoManipula
 	}
 }
 
-void FComponentElementLevelEditorViewportInteractionCustomization::GizmoManipulationStopped(const TTypedElement<UTypedElementWorldInterface>& InElementWorldHandle, const UE::Widget::EWidgetMode InWidgetMode, const ETypedElementViewportInteractionGizmoManipulationType InManipulationType)
+void FComponentElementLevelEditorViewportInteractionCustomization::GizmoManipulationStopped(const TTypedElement<ITypedElementWorldInterface>& InElementWorldHandle, const UE::Widget::EWidgetMode InWidgetMode, const ETypedElementViewportInteractionGizmoManipulationType InManipulationType)
 {
 	UActorComponent* Component = ComponentElementDataUtil::GetComponentFromHandleChecked(InElementWorldHandle);
 
diff --git a/Engine/Source/Editor/LevelEditor/Private/Elements/SMInstance/SMInstanceElementLevelEditorSelectionCustomization.cpp b/Engine/Source/Editor/LevelEditor/Private/Elements/SMInstance/SMInstanceElementLevelEditorSelectionCustomization.cpp
--- a/Engine/Source/Editor/LevelEditor/Private/Elements/SMInstance/SMInstanceElementLevelEditorSelectionCustomization.cpp
+++ b/Engine/Source/Editor/LevelEditor/Private/Elements/SMInstance/SMInstanceElementLevelEditorSelectionCustomization.cpp
@@ -9,7 +9,7 @@
 
 DEFINE_LOG_CATEGORY_STATIC(LogSMInstanceLevelEditorSelection, Log, All);
 
-bool FSMInstanceElementLevelEditorSelectionCustomization::CanSelectElement(const TTypedElement<UTypedElementSelectionInterface>& InElementSelectionHandle, const FTypedElementSelectionOptions& InSelectionOptions)
+bool FSMInstanceElementLevelEditorSelectionCustomization::CanSelectElement(const TTypedElement<ITypedElementSelectionInterface>& InElementSelectionHandle, const FTypedElementSelectionOptions& InSelectionOptions)
 {
 	const FSMInstanceManager SMInstance = SMInstanceElementDataUtil::GetSMInstanceFromHandleChecked(InElementSelectionHandle);
 	if (!SMInstance)
@@ -21,7 +21,7 @@ bool FSMInstanceElementLevelEditorSelectionCustomization::CanSelectElement(const
 	return !GEdSelectionLock;
 }
 
-bool FSMInstanceElementLevelEditorSelectionCustomization::CanDeselectElement(const TTypedElement<UTypedElementSelectionInterface>& InElementSelectionHandle, const FTypedElementSelectionOptions& InSelectionOptions)
+bool FSMInstanceElementLevelEditorSelectionCustomization::CanDeselectElement(const TTypedElement<ITypedElementSelectionInterface>& InElementSelectionHandle, const FTypedElementSelectionOptions& InSelectionOptions)
 {
 	const FSMInstanceManager SMInstance = SMInstanceElementDataUtil::GetSMInstanceFromHandleChecked(InElementSelectionHandle);
 	if (!SMInstance)
@@ -33,7 +33,7 @@ bool FSMInstanceElementLevelEditorSelectionCustomization::CanDeselectElement(con
 	return !GEdSelectionLock;
 }
 
-bool FSMInstanceElementLevelEditorSelectionCustomization::SelectElement(const TTypedElement<UTypedElementSelectionInterface>& InElementSelectionHandle, FTypedElementListRef InSelectionSet, const FTypedElementSelectionOptions& InSelectionOptions)
+bool FSMInstanceElementLevelEditorSelectionCustomization::SelectElement(const TTypedElement<ITypedElementSelectionInterface>& InElementSelectionHandle, FTypedElementListRef InSelectionSet, const FTypedElementSelectionOptions& InSelectionOptions)
 {
 	const FSMInstanceManager SMInstance = SMInstanceElementDataUtil::GetSMInstanceFromHandleChecked(InElementSelectionHandle);
 	if (!SMInstance)
@@ -51,7 +51,7 @@ bool FSMInstanceElementLevelEditorSelectionCustomization::SelectElement(const TT
 	return true;
 }
 
-bool FSMInstanceElementLevelEditorSelectionCustomization::DeselectElement(const TTypedElement<UTypedElementSelectionInterface>& InElementSelectionHandle, FTypedElementListRef InSelectionSet, const FTypedElementSelectionOptions& InSelectionOptions)
+bool FSMInstanceElementLevelEditorSelectionCustomization::DeselectElement(const TTypedElement<ITypedElementSelectionInterface>& InElementSelectionHandle, FTypedElementListRef InSelectionSet, const FTypedElementSelectionOptions& InSelectionOptions)
 {
 	const FSMInstanceManager SMInstance = SMInstanceElementDataUtil::GetSMInstanceFromHandleChecked(InElementSelectionHandle);
 	if (!SMInstance)
diff --git a/Engine/Source/Editor/LevelEditor/Private/Elements/SMInstance/SMInstanceElementLevelEditorSelectionCustomization.h b/Engine/Source/Editor/LevelEditor/Private/Elements/SMInstance/SMInstanceElementLevelEditorSelectionCustomization.h
--- a/Engine/Source/Editor/LevelEditor/Private/Elements/SMInstance/SMInstanceElementLevelEditorSelectionCustomization.h
+++ b/Engine/Source/Editor/LevelEditor/Private/Elements/SMInstance/SMInstanceElementLevelEditorSelectionCustomization.h
@@ -8,8 +8,8 @@
 class FSMInstanceElementLevelEditorSelectionCustomization : public FTypedElementSelectionCustomization, public FTypedElementAssetEditorToolkitHostMixin
 {
 public:
-	virtual bool CanSelectElement(const TTypedElement<UTypedElementSelectionInterface>& InElementSelectionHandle, const FTypedElementSelectionOptions& InSelectionOptions) override;
-	virtual bool CanDeselectElement(const TTypedElement<UTypedElementSelectionInterface>& InElementSelectionHandle, const FTypedElementSelectionOptions& InSelectionOptions) override;
-	virtual bool SelectElement(const TTypedElement<UTypedElementSelectionInterface>& InElementSelectionHandle, FTypedElementListRef InSelectionSet, const FTypedElementSelectionOptions& InSelectionOptions) override;
-	virtual bool DeselectElement(const TTypedElement<UTypedElementSelectionInterface>& InElementSelectionHandle, FTypedElementListRef InSelectionSet, const FTypedElementSelectionOptions& InSelectionOptions) override;
+	virtual bool CanSelectElement(const TTypedElement<ITypedElementSelectionInterface>& InElementSelectionHandle, const FTypedElementSelectionOptions& InSelectionOptions) override;
+	virtual bool CanDeselectElement(const TTypedElement<ITypedElementSelectionInterface>& InElementSelectionHandle, const FTypedElementSelectionOptions& InSelectionOptions) override;
+	virtual bool SelectElement(const TTypedElement<ITypedElementSelectionInterface>& InElementSelectionHandle, FTypedElementListRef InSelectionSet, const FTypedElementSelectionOptions& InSelectionOptions) override;
+	virtual bool DeselectElement(const TTypedElement<ITypedElementSelectionInterface>& InElementSelectionHandle, FTypedElementListRef InSelectionSet, const FTypedElementSelectionOptions& InSelectionOptions) override;
 };
diff --git a/Engine/Source/Editor/LevelEditor/Private/LevelEditorActions.cpp b/Engine/Source/Editor/LevelEditor/Private/LevelEditorActions.cpp
--- a/Engine/Source/Editor/LevelEditor/Private/LevelEditorActions.cpp
+++ b/Engine/Source/Editor/LevelEditor/Private/LevelEditorActions.cpp
@@ -1534,7 +1534,7 @@ bool FLevelEditorActionCallbacks::Duplicate_CanExecute()
 		bool bCanDuplicate = false;
 
 		const UTypedElementSelectionSet* SelectionSet = LevelEditor->GetElementSelectionSet();
-		SelectionSet->ForEachSelectedElement<UTypedElementWorldInterface>([&bCanDuplicate](const TTypedElement<UTypedElementWorldInterface>& InWorldElement)
+		SelectionSet->ForEachSelectedElement<ITypedElementWorldInterface>([&bCanDuplicate](const TTypedElement<ITypedElementWorldInterface>& InWorldElement)
 		{
 			bCanDuplicate |= InWorldElement.CanDuplicateElement();
 			return !bCanDuplicate;
@@ -1572,7 +1572,7 @@ bool FLevelEditorActionCallbacks::Delete_CanExecute()
 		bool bCanDelete = false;
 
 		const UTypedElementSelectionSet* SelectionSet = LevelEditor->GetElementSelectionSet();
-		SelectionSet->ForEachSelectedElement<UTypedElementWorldInterface>([&bCanDelete](const TTypedElement<UTypedElementWorldInterface>& InWorldElement)
+		SelectionSet->ForEachSelectedElement<ITypedElementWorldInterface>([&bCanDelete](const TTypedElement<ITypedElementWorldInterface>& InWorldElement)
 		{
 			bCanDelete |= InWorldElement.CanDeleteElement();
 			return !bCanDelete;
@@ -2861,7 +2861,7 @@ void FLevelEditorActionCallbacks::MoveElementsToElement_Clicked(bool InAlign)
 	// TODO: Ideally this would come from some level editor context
 	if (const UTypedElementSelectionSet* SelectionSet = GEditor->GetSelectedActors()->GetElementSelectionSet())
 	{
-		if (const TTypedElement<UTypedElementWorldInterface> DestElement = SelectionSet->GetBottomSelectedElement<UTypedElementWorldInterface>())
+		if (const TTypedElement<ITypedElementWorldInterface> DestElement = SelectionSet->GetBottomSelectedElement<ITypedElementWorldInterface>())
 		{
 			const FScopedTransaction Transaction(NSLOCTEXT("UnrealEd", "MoveElementsToElement", "Snap Origin to Element"));
 			MoveTo_Clicked(SelectionSet, InAlign, /*bPerElement*/false, DestElement);
@@ -2869,7 +2869,7 @@ void FLevelEditorActionCallbacks::MoveElementsToElement_Clicked(bool InAlign)
 	}
 }
 
-void FLevelEditorActionCallbacks::MoveTo_Clicked( const UTypedElementSelectionSet* InSelectionSet, const bool InAlign, bool InPerElement, const TTypedElement<UTypedElementWorldInterface>& InDestination )
+void FLevelEditorActionCallbacks::MoveTo_Clicked( const UTypedElementSelectionSet* InSelectionSet, const bool InAlign, bool InPerElement, const TTypedElement<ITypedElementWorldInterface>& InDestination )
 {
 	// Fires ULevel::LevelDirtiedEvent when falling out of scope.
 	FScopedLevelDirtied		LevelDirtyCallback;
@@ -2900,7 +2900,7 @@ void FLevelEditorActionCallbacks::MoveTo_Clicked( const UTypedElementSelectionSe
 		}
 	}
 
-	InSelectionSet->ForEachSelectedElement<UTypedElementWorldInterface>([InAlign, InPerElement, &InDestination, &Delta, &NewLocation, &NewRotation, &LevelDirtyCallback](const TTypedElement<UTypedElementWorldInterface>& InElement)
+	InSelectionSet->ForEachSelectedElement<ITypedElementWorldInterface>([InAlign, InPerElement, &InDestination, &Delta, &NewLocation, &NewRotation, &LevelDirtyCallback](const TTypedElement<ITypedElementWorldInterface>& InElement)
 	{
 		// Skip moving the destination element
 		if (InElement == InDestination)
@@ -3078,7 +3078,7 @@ void FLevelEditorActionCallbacks::SnapElementsToElement_Clicked( bool InAlign, b
 	// TODO: Ideally this would come from some level editor context
 	if (const UTypedElementSelectionSet* SelectionSet = GEditor->GetSelectedActors()->GetElementSelectionSet())
 	{
-		if (const TTypedElement<UTypedElementWorldInterface> DestElement = SelectionSet->GetBottomSelectedElement<UTypedElementWorldInterface>())
+		if (const TTypedElement<ITypedElementWorldInterface> DestElement = SelectionSet->GetBottomSelectedElement<ITypedElementWorldInterface>())
 		{
 			const FScopedTransaction Transaction(NSLOCTEXT("UnrealEd", "SnapElementsToElement", "Snap Elements to Element"));
 			SnapTo_Clicked(SelectionSet, InAlign, InUseLineTrace, InUseBounds, InUsePivot, DestElement);
@@ -3086,7 +3086,7 @@ void FLevelEditorActionCallbacks::SnapElementsToElement_Clicked( bool InAlign, b
 	}
 }
 
-void FLevelEditorActionCallbacks::SnapTo_Clicked( const UTypedElementSelectionSet* InSelectionSet, const bool InAlign, const bool InUseLineTrace, const bool InUseBounds, const bool InUsePivot, const TTypedElement<UTypedElementWorldInterface>& InDestination )
+void FLevelEditorActionCallbacks::SnapTo_Clicked( const UTypedElementSelectionSet* InSelectionSet, const bool InAlign, const bool InUseLineTrace, const bool InUseBounds, const bool InUsePivot, const TTypedElement<ITypedElementWorldInterface>& InDestination )
 {
 	// Fires ULevel::LevelDirtiedEvent when falling out of scope.
 	FScopedLevelDirtied		LevelDirtyCallback;
@@ -3095,7 +3095,7 @@ void FLevelEditorActionCallbacks::SnapTo_Clicked( const UTypedElementSelectionSe
 	// TODO: Should this also take an element?
 	{
 		AActor* DestinationActor = nullptr;
-		if (TTypedElement<UTypedElementObjectInterface> DestinationObjectHandle = InSelectionSet->GetElementList()->GetElement<UTypedElementObjectInterface>(InDestination))
+		if (TTypedElement<ITypedElementObjectInterface> DestinationObjectHandle = InSelectionSet->GetElementList()->GetElement<ITypedElementObjectInterface>(InDestination))
 		{
 			DestinationActor = Cast<AActor>(DestinationObjectHandle.GetObject());
 		}
@@ -3111,7 +3111,7 @@ void FLevelEditorActionCallbacks::SnapTo_Clicked( const UTypedElementSelectionSe
 
 	// Snap each selected element
 	bool bSnappedElements = false;
-	InSelectionSet->ForEachSelectedElement<UTypedElementWorldInterface>([InAlign, InUseLineTrace, InUseBounds, InUsePivot, &InDestination, &ElementsToIgnore, &LevelDirtyCallback, &bSnappedElements](const TTypedElement<UTypedElementWorldInterface>& InElement)
+	InSelectionSet->ForEachSelectedElement<ITypedElementWorldInterface>([InAlign, InUseLineTrace, InUseBounds, InUsePivot, &InDestination, &ElementsToIgnore, &LevelDirtyCallback, &bSnappedElements](const TTypedElement<ITypedElementWorldInterface>& InElement)
 	{
 		if (GEditor->SnapElementTo(InElement, InAlign, InUseLineTrace, InUseBounds, InUsePivot, InDestination, ElementsToIgnore))
 		{
@@ -3124,7 +3124,7 @@ void FLevelEditorActionCallbacks::SnapTo_Clicked( const UTypedElementSelectionSe
 	// Update the pivot location
 	if (bSnappedElements)
 	{
-		if (TTypedElement<UTypedElementWorldInterface> LastElement = InSelectionSet->GetBottomSelectedElement<UTypedElementWorldInterface>())
+		if (TTypedElement<ITypedElementWorldInterface> LastElement = InSelectionSet->GetBottomSelectedElement<ITypedElementWorldInterface>())
 		{
 			FTransform LastElementTransform;
 			if (LastElement.GetWorldTransform(LastElementTransform))
@@ -3133,7 +3133,7 @@ void FLevelEditorActionCallbacks::SnapTo_Clicked( const UTypedElementSelectionSe
 
 				if (UActorGroupingUtils::IsGroupingActive())
 				{
-					if (TTypedElement<UTypedElementObjectInterface> LastObjectElement = InSelectionSet->GetElementList()->GetElement<UTypedElementObjectInterface>(LastElement))
+					if (TTypedElement<ITypedElementObjectInterface> LastObjectElement = InSelectionSet->GetElementList()->GetElement<ITypedElementObjectInterface>(LastElement))
 					{
 						if (AActor* LastActor = Cast<AActor>(LastObjectElement.GetObject()))
 						{
diff --git a/Engine/Source/Editor/LevelEditor/Private/LevelEditorContextMenu.cpp b/Engine/Source/Editor/LevelEditor/Private/LevelEditorContextMenu.cpp
--- a/Engine/Source/Editor/LevelEditor/Private/LevelEditorContextMenu.cpp
+++ b/Engine/Source/Editor/LevelEditor/Private/LevelEditorContextMenu.cpp
@@ -609,7 +609,7 @@ FName FLevelEditorContextMenu::InitMenuContext(FToolMenuContext& Context, TWeakP
 	ContextObject->CurrentSelection = LevelEditorPtr->GetElementSelectionSet();
 	ContextObject->HitProxyElement = HitProxyElement;
 	{
-		TTypedElement<UTypedElementObjectInterface> HitProxyObjectElement = ContextObject->CurrentSelection->GetElementList()->GetElement<UTypedElementObjectInterface>(ContextObject->HitProxyElement);
+		TTypedElement<ITypedElementObjectInterface> HitProxyObjectElement = ContextObject->CurrentSelection->GetElementList()->GetElement<ITypedElementObjectInterface>(ContextObject->HitProxyElement);
 		ContextObject->HitProxyActor = HitProxyObjectElement ? Cast<AActor>(HitProxyObjectElement.GetObject()) : nullptr;
 	}
 	for (FSelectedEditableComponentIterator It(GEditor->GetSelectedEditableComponentIterator()); It; ++It)
diff --git a/Engine/Source/Editor/LevelEditor/Private/SActorDetails.cpp b/Engine/Source/Editor/LevelEditor/Private/SActorDetails.cpp
--- a/Engine/Source/Editor/LevelEditor/Private/SActorDetails.cpp
+++ b/Engine/Source/Editor/LevelEditor/Private/SActorDetails.cpp
@@ -253,9 +253,9 @@ void SActorDetails::RefreshSelection(const bool bForceRefresh)
 		return;
 	}
 
-	TArray<TTypedElement<UTypedElementDetailsInterface>> DetailsElements;
+	TArray<TTypedElement<ITypedElementDetailsInterface>> DetailsElements;
 	DetailsElements.Reserve(SelectionSet->GetNumSelectedElements());
-	SelectionSet->ForEachSelectedElement<UTypedElementDetailsInterface>([&DetailsElements](const TTypedElement<UTypedElementDetailsInterface>& InDetailsElement)
+	SelectionSet->ForEachSelectedElement<ITypedElementDetailsInterface>([&DetailsElements](const TTypedElement<ITypedElementDetailsInterface>& InDetailsElement)
 	{
 		DetailsElements.Add(InDetailsElement);
 		return true;
@@ -271,13 +271,13 @@ void SActorDetails::OverrideSelection(const TArray<AActor*>& InActors, const boo
 {
 	UTypedElementRegistry* Registry = UTypedElementRegistry::GetInstance();
 
-	TArray<TTypedElement<UTypedElementDetailsInterface>> DetailsElements;
+	TArray<TTypedElement<ITypedElementDetailsInterface>> DetailsElements;
 	DetailsElements.Reserve(SelectionSet->GetNumSelectedElements());
 	for (AActor* Actor : InActors)
 	{
 		if (FTypedElementHandle ActorElementHandle = UEngineElementsLibrary::AcquireEditorActorElementHandle(Actor))
 		{
-			if (TTypedElement<UTypedElementDetailsInterface> ActorDetailsHandle = Registry->GetElement<UTypedElementDetailsInterface>(ActorElementHandle))
+			if (TTypedElement<ITypedElementDetailsInterface> ActorDetailsHandle = Registry->GetElement<ITypedElementDetailsInterface>(ActorElementHandle))
 			{
 				DetailsElements.Add(MoveTemp(ActorDetailsHandle));
 			}
@@ -290,7 +290,7 @@ void SActorDetails::OverrideSelection(const TArray<AActor*>& InActors, const boo
 	RefreshTopLevelElements(DetailsElements, bForceRefresh, /*bOverrideLock*/false);
 }
 
-void SActorDetails::RefreshTopLevelElements(TArrayView<const TTypedElement<UTypedElementDetailsInterface>> InDetailsElements, const bool bForceRefresh, const bool bOverrideLock)
+void SActorDetails::RefreshTopLevelElements(TArrayView<const TTypedElement<ITypedElementDetailsInterface>> InDetailsElements, const bool bForceRefresh, const bool bOverrideLock)
 {
 	// Nothing to do if this view is locked!
 	if (DetailsView->IsLocked() && !bOverrideLock)
@@ -300,7 +300,7 @@ void SActorDetails::RefreshTopLevelElements(TArrayView<const TTypedElement<UType
 
 	// Build the array of top-level elements to edit
 	TopLevelElements.Reset(InDetailsElements.Num());
-	for (const TTypedElement<UTypedElementDetailsInterface>& DetailsElement : InDetailsElements)
+	for (const TTypedElement<ITypedElementDetailsInterface>& DetailsElement : InDetailsElements)
 	{
 		if (DetailsElement.IsTopLevelElement())
 		{
@@ -380,7 +380,7 @@ void SActorDetails::RefreshSubobjectTreeElements(TArrayView<const FSubobjectEdit
 		{
 			if (FTypedElementHandle ComponentElementHandle = UEngineElementsLibrary::AcquireEditorComponentElementHandle(Component))
 			{
-				if (TTypedElement<UTypedElementDetailsInterface> ComponentDetailsHandle = Registry->GetElement<UTypedElementDetailsInterface>(ComponentElementHandle))
+				if (TTypedElement<ITypedElementDetailsInterface> ComponentDetailsHandle = Registry->GetElement<ITypedElementDetailsInterface>(ComponentElementHandle))
 				{
 					if (TUniquePtr<ITypedElementDetailsObject> ElementDetailsObject = ComponentDetailsHandle.GetDetailsObject())
 					{
diff --git a/Engine/Source/Editor/LevelEditor/Private/SActorDetails.h b/Engine/Source/Editor/LevelEditor/Private/SActorDetails.h
--- a/Engine/Source/Editor/LevelEditor/Private/SActorDetails.h
+++ b/Engine/Source/Editor/LevelEditor/Private/SActorDetails.h
@@ -74,7 +74,7 @@ public:
 	void SetSubobjectEditorUICustomization(TSharedPtr<ISCSEditorUICustomization> ActorDetailsSCSEditorUICustomization);
 
 private:
-	void RefreshTopLevelElements(TArrayView<const TTypedElement<UTypedElementDetailsInterface>> InDetailsElements, const bool bForceRefresh, const bool bOverrideLock);
+	void RefreshTopLevelElements(TArrayView<const TTypedElement<ITypedElementDetailsInterface>> InDetailsElements, const bool bForceRefresh, const bool bOverrideLock);
 	void RefreshSubobjectTreeElements(TArrayView<const TSharedPtr<FSubobjectEditorTreeNode>> InSelectedNodes, const bool bForceRefresh, const bool bOverrideLock);
 	void SetElementDetailsObjects(TArrayView<const TUniquePtr<ITypedElementDetailsObject>> InElementDetailsObjects, const bool bForceRefresh, const bool bOverrideLock);
 
diff --git a/Engine/Source/Editor/LevelEditor/Public/Elements/Actor/ActorElementLevelEditorViewportInteractionCustomization.h b/Engine/Source/Editor/LevelEditor/Public/Elements/Actor/ActorElementLevelEditorViewportInteractionCustomization.h
--- a/Engine/Source/Editor/LevelEditor/Public/Elements/Actor/ActorElementLevelEditorViewportInteractionCustomization.h
+++ b/Engine/Source/Editor/LevelEditor/Public/Elements/Actor/ActorElementLevelEditorViewportInteractionCustomization.h
@@ -8,9 +8,9 @@
 class LEVELEDITOR_API FActorElementLevelEditorViewportInteractionCustomization : public FActorElementEditorViewportInteractionCustomization, public FTypedElementAssetEditorLevelEditorViewportClientMixin
 {
 public:
-	virtual void GizmoManipulationStarted(const TTypedElement<UTypedElementWorldInterface>& InElementWorldHandle, const UE::Widget::EWidgetMode InWidgetMode) override;
-	virtual void GizmoManipulationDeltaUpdate(const TTypedElement<UTypedElementWorldInterface>& InElementWorldHandle, const UE::Widget::EWidgetMode InWidgetMode, const EAxisList::Type InDragAxis, const FInputDeviceState& InInputState, const FTransform& InDeltaTransform, const FVector& InPivotLocation) override;
-	virtual void GizmoManipulationStopped(const TTypedElement<UTypedElementWorldInterface>& InElementWorldHandle, const UE::Widget::EWidgetMode InWidgetMode, const ETypedElementViewportInteractionGizmoManipulationType InManipulationType) override;
+	virtual void GizmoManipulationStarted(const TTypedElement<ITypedElementWorldInterface>& InElementWorldHandle, const UE::Widget::EWidgetMode InWidgetMode) override;
+	virtual void GizmoManipulationDeltaUpdate(const TTypedElement<ITypedElementWorldInterface>& InElementWorldHandle, const UE::Widget::EWidgetMode InWidgetMode, const EAxisList::Type InDragAxis, const FInputDeviceState& InInputState, const FTransform& InDeltaTransform, const FVector& InPivotLocation) override;
+	virtual void GizmoManipulationStopped(const TTypedElement<ITypedElementWorldInterface>& InElementWorldHandle, const UE::Widget::EWidgetMode InWidgetMode, const ETypedElementViewportInteractionGizmoManipulationType InManipulationType) override;
 	virtual void PostGizmoManipulationStopped(TArrayView<const FTypedElementHandle> InElementHandles, const UE::Widget::EWidgetMode InWidgetMode) override;
 
 private:
diff --git a/Engine/Source/Editor/LevelEditor/Public/Elements/Component/ComponentElementLevelEditorViewportInteractionCustomization.h b/Engine/Source/Editor/LevelEditor/Public/Elements/Component/ComponentElementLevelEditorViewportInteractionCustomization.h
--- a/Engine/Source/Editor/LevelEditor/Public/Elements/Component/ComponentElementLevelEditorViewportInteractionCustomization.h
+++ b/Engine/Source/Editor/LevelEditor/Public/Elements/Component/ComponentElementLevelEditorViewportInteractionCustomization.h
@@ -10,9 +10,9 @@ class FLevelEditorViewportClient;
 class LEVELEDITOR_API FComponentElementLevelEditorViewportInteractionCustomization : public FComponentElementEditorViewportInteractionCustomization, public FTypedElementAssetEditorLevelEditorViewportClientMixin
 {
 public:
-	virtual void GizmoManipulationStarted(const TTypedElement<UTypedElementWorldInterface>& InElementWorldHandle, const UE::Widget::EWidgetMode InWidgetMode) override;
-	virtual void GizmoManipulationDeltaUpdate(const TTypedElement<UTypedElementWorldInterface>& InElementWorldHandle, const UE::Widget::EWidgetMode InWidgetMode, const EAxisList::Type InDragAxis, const FInputDeviceState& InInputState, const FTransform& InDeltaTransform, const FVector& InPivotLocation) override;
-	virtual void GizmoManipulationStopped(const TTypedElement<UTypedElementWorldInterface>& InElementWorldHandle, const UE::Widget::EWidgetMode InWidgetMode, const ETypedElementViewportInteractionGizmoManipulationType InManipulationType) override;
+	virtual void GizmoManipulationStarted(const TTypedElement<ITypedElementWorldInterface>& InElementWorldHandle, const UE::Widget::EWidgetMode InWidgetMode) override;
+	virtual void GizmoManipulationDeltaUpdate(const TTypedElement<ITypedElementWorldInterface>& InElementWorldHandle, const UE::Widget::EWidgetMode InWidgetMode, const EAxisList::Type InDragAxis, const FInputDeviceState& InInputState, const FTransform& InDeltaTransform, const FVector& InPivotLocation) override;
+	virtual void GizmoManipulationStopped(const TTypedElement<ITypedElementWorldInterface>& InElementWorldHandle, const UE::Widget::EWidgetMode InWidgetMode, const ETypedElementViewportInteractionGizmoManipulationType InManipulationType) override;
 
 	static void ValidateScale(const FVector& InOriginalPreDragScale, const EAxisList::Type InDragAxis, const FVector& InCurrentScale, const FVector& InBoxExtent, FVector& InOutScaleDelta, bool bInCheckSmallExtent);
 	static FProperty* GetEditTransformProperty(const UE::Widget::EWidgetMode InWidgetMode);
diff --git a/Engine/Source/Editor/LevelEditor/Public/LevelEditorActions.h b/Engine/Source/Editor/LevelEditor/Public/LevelEditorActions.h
--- a/Engine/Source/Editor/LevelEditor/Public/LevelEditorActions.h
+++ b/Engine/Source/Editor/LevelEditor/Public/LevelEditorActions.h
@@ -1356,7 +1356,7 @@ private:
 	 * Moves the selected elements.
 	 * @param InDestination		The destination element we want to move this element to, or invalid to move to the grid
 	 */
-	static void MoveTo_Clicked( const UTypedElementSelectionSet* InSelectionSet, const bool InAlign, bool InPerElement, const TTypedElement<UTypedElementWorldInterface>& InDestination = TTypedElement<UTypedElementWorldInterface>() );
+	static void MoveTo_Clicked( const UTypedElementSelectionSet* InSelectionSet, const bool InAlign, bool InPerElement, const TTypedElement<ITypedElementWorldInterface>& InDestination = TTypedElement<ITypedElementWorldInterface>() );
 
 	/** 
 	 * Snaps the selected elements. Optionally will align with the trace normal.
@@ -1366,7 +1366,7 @@ private:
 	 * @param InUsePivot		Whether or not to use the pivot position.
 	 * @param InDestination		The destination element we want to move this actor to, or invalid to go towards the floor
 	 */
-	static void SnapTo_Clicked(const UTypedElementSelectionSet* InSelectionSet, const bool InAlign, const bool InUseLineTrace, const bool InUseBounds, const bool InUsePivot, const TTypedElement<UTypedElementWorldInterface>& InDestination = TTypedElement<UTypedElementWorldInterface>() );
+	static void SnapTo_Clicked(const UTypedElementSelectionSet* InSelectionSet, const bool InAlign, const bool InUseLineTrace, const bool InUseBounds, const bool InUsePivot, const TTypedElement<ITypedElementWorldInterface>& InDestination = TTypedElement<ITypedElementWorldInterface>() );
 
 	/** 
 	 * Create and apply animation to the SkeletalMeshComponent if Simulating
diff --git a/Engine/Source/Editor/UnrealEd/Private/AssetSelection.cpp b/Engine/Source/Editor/UnrealEd/Private/AssetSelection.cpp
--- a/Engine/Source/Editor/UnrealEd/Private/AssetSelection.cpp
+++ b/Engine/Source/Editor/UnrealEd/Private/AssetSelection.cpp
@@ -619,7 +619,7 @@ static AActor* PrivateAddActor( UObject* Asset, UActorFactory* Factory, bool Sel
 			TArray<FTypedElementHandle> PlacedElements = PlacementSubsystem->PlaceAsset(PlacementInfo, PlacementOptions);
 			if (PlacedElements.Num())
 			{
-				if (TTypedElement<UTypedElementObjectInterface> ObjectInterface = UTypedElementRegistry::GetInstance()->GetElement<UTypedElementObjectInterface>(PlacedElements[0]))
+				if (TTypedElement<ITypedElementObjectInterface> ObjectInterface = UTypedElementRegistry::GetInstance()->GetElement<ITypedElementObjectInterface>(PlacedElements[0]))
 				{
 					Actor = ObjectInterface.GetObjectAs<AActor>();
 					if (Actor)
diff --git a/Engine/Source/Editor/UnrealEd/Private/EditorLevelUtils.cpp b/Engine/Source/Editor/UnrealEd/Private/EditorLevelUtils.cpp
--- a/Engine/Source/Editor/UnrealEd/Private/EditorLevelUtils.cpp
+++ b/Engine/Source/Editor/UnrealEd/Private/EditorLevelUtils.cpp
@@ -1212,8 +1212,8 @@ void SetLevelVisibilityNoGlobalUpdateInternal(ULevel* Level, const bool bShouldB
 		{
 			TArray<FTypedElementHandle> LevelElementHandles;
 			
-			ActorSelectionSet->ForEachSelectedElement<UTypedElementWorldInterface>(
-				[Level, &LevelElementHandles](const TTypedElement<UTypedElementWorldInterface>& Element)
+			ActorSelectionSet->ForEachSelectedElement<ITypedElementWorldInterface>(
+				[Level, &LevelElementHandles](const TTypedElement<ITypedElementWorldInterface>& Element)
 			{
 				if (Element.GetOwnerLevel() == Level)
 				{
diff --git a/Engine/Source/Editor/UnrealEd/Private/EditorModeManager.cpp b/Engine/Source/Editor/UnrealEd/Private/EditorModeManager.cpp
--- a/Engine/Source/Editor/UnrealEd/Private/EditorModeManager.cpp
+++ b/Engine/Source/Editor/UnrealEd/Private/EditorModeManager.cpp
@@ -934,7 +934,7 @@ FMatrix FEditorModeTools::GetLocalCoordinateSystem()
 	// If there isn't an active mode overriding the local coordinate system, create it by looking at the current selection.
 	if (!CustomCoordinateSystemProvided && GetEditorSelectionSet())
 	{
-		if (TTypedElement<UTypedElementWorldInterface> BottomSelected = GetEditorSelectionSet()->GetBottomSelectedElement<UTypedElementWorldInterface>())
+		if (TTypedElement<ITypedElementWorldInterface> BottomSelected = GetEditorSelectionSet()->GetBottomSelectedElement<ITypedElementWorldInterface>())
 		{
 			FTransform RelativeTransform;
 			BottomSelected.GetRelativeTransform(RelativeTransform);
diff --git a/Engine/Source/Editor/UnrealEd/Private/EditorSelectUtils.cpp b/Engine/Source/Editor/UnrealEd/Private/EditorSelectUtils.cpp
--- a/Engine/Source/Editor/UnrealEd/Private/EditorSelectUtils.cpp
+++ b/Engine/Source/Editor/UnrealEd/Private/EditorSelectUtils.cpp
@@ -204,10 +204,10 @@ void UUnrealEdEngine::SetPivot( FVector NewPivot, bool bSnapPivotToGrid, bool bI
 	FVector Widget2DWorldXAxis = FVector::ZeroVector;
 
 	// Pick a new common pivot, or not.
-	TTypedElement<UTypedElementWorldInterface> SingleWorldElement;
+	TTypedElement<ITypedElementWorldInterface> SingleWorldElement;
 
 	FTypedElementListConstRef ElementsToManipulate = GCurrentLevelEditingViewportClient->GetElementsToManipulate();
-	ElementsToManipulate->ForEachElement<UTypedElementWorldInterface>([&NumElements, &TranslateRotateWidgetWorldXAxis, &Widget2DWorldXAxis, &SingleWorldElement](const TTypedElement<UTypedElementWorldInterface>& InWorldElement)
+	ElementsToManipulate->ForEachElement<ITypedElementWorldInterface>([&NumElements, &TranslateRotateWidgetWorldXAxis, &Widget2DWorldXAxis, &SingleWorldElement](const TTypedElement<ITypedElementWorldInterface>& InWorldElement)
 	{
 		if (NumElements == 0)
 		{
@@ -237,7 +237,7 @@ void UUnrealEdEngine::SetPivot( FVector NewPivot, bool bSnapPivotToGrid, bool bI
 	
 	if (bAssignPivot && UActorGroupingUtils::IsGroupingActive())
 	{
-		if (TTypedElement<UTypedElementObjectInterface> ObjectElement = ElementsToManipulate->GetElement<UTypedElementObjectInterface>(SingleWorldElement))
+		if (TTypedElement<ITypedElementObjectInterface> ObjectElement = ElementsToManipulate->GetElement<ITypedElementObjectInterface>(SingleWorldElement))
 		{
 			if (AActor* SingleActor = Cast<AActor>(ObjectElement.GetObject()))
 			{
@@ -326,10 +326,10 @@ void UUnrealEdEngine::UpdatePivotLocationForSelection( bool bOnChange )
 	}
 
 	// Pick a new common pivot, or not.
-	TTypedElement<UTypedElementWorldInterface> SingleWorldElement;
+	TTypedElement<ITypedElementWorldInterface> SingleWorldElement;
 	
 	FTypedElementListConstRef ElementsToManipulate = GCurrentLevelEditingViewportClient->GetElementsToManipulate();
-	ElementsToManipulate->ForEachElement<UTypedElementWorldInterface>([&SingleWorldElement](const TTypedElement<UTypedElementWorldInterface>& InWorldElement)
+	ElementsToManipulate->ForEachElement<ITypedElementWorldInterface>([&SingleWorldElement](const TTypedElement<ITypedElementWorldInterface>& InWorldElement)
 	{
 #if DO_CHECK
 		if (ULevel* OwnerLevel = InWorldElement.GetOwnerLevel())
@@ -365,7 +365,7 @@ void UUnrealEdEngine::UpdatePivotLocationForSelection( bool bOnChange )
 			// If grouping is active, see if this element is an actor that's part of a locked group and use that pivot instead
 			if (UActorGroupingUtils::IsGroupingActive())
 			{
-				if (TTypedElement<UTypedElementObjectInterface> ObjectElement = ElementsToManipulate->GetElement<UTypedElementObjectInterface>(SingleWorldElement))
+				if (TTypedElement<ITypedElementObjectInterface> ObjectElement = ElementsToManipulate->GetElement<ITypedElementObjectInterface>(SingleWorldElement))
 				{
 					if (AActor* SingleActor = Cast<AActor>(ObjectElement.GetObject()))
 					{
diff --git a/Engine/Source/Editor/UnrealEd/Private/EditorServer.cpp b/Engine/Source/Editor/UnrealEd/Private/EditorServer.cpp
--- a/Engine/Source/Editor/UnrealEd/Private/EditorServer.cpp
+++ b/Engine/Source/Editor/UnrealEd/Private/EditorServer.cpp
@@ -4980,7 +4980,7 @@ void UEditorEngine::MoveViewportCamerasToComponent(const USceneComponent* Compon
 
 void UEditorEngine::MoveViewportCamerasToElement(const UTypedElementSelectionSet* SelectionSet, bool bActiveViewportOnly) const
 {
-	if (!SelectionSet || !SelectionSet->HasSelectedElements(UTypedElementWorldInterface::StaticClass()))
+	if (!SelectionSet || !SelectionSet->HasSelectedElements(ITypedElementWorldInterface::UClassType::StaticClass()))
 	{
 		return;
 	}
@@ -5069,7 +5069,7 @@ bool UEditorEngine::SnapElementTo(const FTypedElementHandle& InElementHandle, co
 
 	const UTypedElementRegistry* Registry = UTypedElementRegistry::GetInstance();
 
-	TTypedElement<UTypedElementWorldInterface> ElementWorldHandle = Registry->GetElement<UTypedElementWorldInterface>(InElementHandle);
+	TTypedElement<ITypedElementWorldInterface> ElementWorldHandle = Registry->GetElement<ITypedElementWorldInterface>(InElementHandle);
 	if (!ElementWorldHandle)
 	{
 		return false;
@@ -5088,7 +5088,7 @@ bool UEditorEngine::SnapElementTo(const FTypedElementHandle& InElementHandle, co
 	}
 
 	ABrush* Brush = nullptr;
-	if (TTypedElement<UTypedElementObjectInterface> ElementObjectHandle = Registry->GetElement<UTypedElementObjectInterface>(InElementHandle))
+	if (TTypedElement<ITypedElementObjectInterface> ElementObjectHandle = Registry->GetElement<ITypedElementObjectInterface>(InElementHandle))
 	{
 		Brush = Cast<ABrush>(ElementObjectHandle.GetObject());
 	}
@@ -5142,7 +5142,7 @@ bool UEditorEngine::SnapElementTo(const FTypedElementHandle& InElementHandle, co
 	}
 
 	FVector Direction = FVector(0.f, 0.f, -1.f);
-	if (TTypedElement<UTypedElementWorldInterface> DestinationWorldHandle = Registry->GetElement<UTypedElementWorldInterface>(InDestination))	// If a destination actor was specified, work out the direction
+	if (TTypedElement<ITypedElementWorldInterface> DestinationWorldHandle = Registry->GetElement<ITypedElementWorldInterface>(InDestination))	// If a destination actor was specified, work out the direction
 	{
 		FTransform DestinationTransform;
 		if (!DestinationWorldHandle.GetWorldTransform(DestinationTransform))
@@ -5244,7 +5244,7 @@ void UEditorEngine::MoveActorInFrontOfCamera( AActor& InActor, const FVector& In
 
 void UEditorEngine::SnapViewTo(const FTypedElementHandle& InElementHandle)
 {
-	if (TTypedElement<UTypedElementWorldInterface> ElementWorldHandle = UTypedElementRegistry::GetInstance()->GetElement<UTypedElementWorldInterface>(InElementHandle))
+	if (TTypedElement<ITypedElementWorldInterface> ElementWorldHandle = UTypedElementRegistry::GetInstance()->GetElement<ITypedElementWorldInterface>(InElementHandle))
 	{
 		FTransform ElementTransform;
 		if (ElementWorldHandle.GetWorldTransform(ElementTransform))
diff --git a/Engine/Source/Editor/UnrealEd/Private/Elements/Actor/ActorElementDetailsInterface.h b/Engine/Source/Editor/UnrealEd/Private/Elements/Actor/ActorElementDetailsInterface.h
--- a/Engine/Source/Editor/UnrealEd/Private/Elements/Actor/ActorElementDetailsInterface.h
+++ b/Engine/Source/Editor/UnrealEd/Private/Elements/Actor/ActorElementDetailsInterface.h
@@ -6,7 +6,7 @@
 #include "ActorElementDetailsInterface.generated.h"
 
 UCLASS()
-class UActorElementDetailsInterface : public UTypedElementDetailsInterface
+class UActorElementDetailsInterface : public UObject, public ITypedElementDetailsInterface
 {
 	GENERATED_BODY()
 
diff --git a/Engine/Source/Editor/UnrealEd/Private/Elements/Actor/ActorElementEditorViewportInteractionCustomization.cpp b/Engine/Source/Editor/UnrealEd/Private/Elements/Actor/ActorElementEditorViewportInteractionCustomization.cpp
--- a/Engine/Source/Editor/UnrealEd/Private/Elements/Actor/ActorElementEditorViewportInteractionCustomization.cpp
+++ b/Engine/Source/Editor/UnrealEd/Private/Elements/Actor/ActorElementEditorViewportInteractionCustomization.cpp
@@ -14,7 +14,7 @@
 #include "Elements/Interfaces/TypedElementHierarchyInterface.h"
 #include "Elements/Framework/TypedElementRegistry.h"
 
-bool FActorElementEditorViewportInteractionCustomization::GetGizmoPivotLocation(const TTypedElement<UTypedElementWorldInterface>& InElementWorldHandle, const UE::Widget::EWidgetMode InWidgetMode, FVector& OutPivotLocation)
+bool FActorElementEditorViewportInteractionCustomization::GetGizmoPivotLocation(const TTypedElement<ITypedElementWorldInterface>& InElementWorldHandle, const UE::Widget::EWidgetMode InWidgetMode, FVector& OutPivotLocation)
 {
 	if (const IToolkitHost* ToolkitHostPtr = GetToolkitHost())
 	{
@@ -25,7 +25,7 @@ bool FActorElementEditorViewportInteractionCustomization::GetGizmoPivotLocation(
 	return FTypedElementViewportInteractionCustomization::GetGizmoPivotLocation(InElementWorldHandle, InWidgetMode, OutPivotLocation);
 }
 
-void FActorElementEditorViewportInteractionCustomization::GizmoManipulationDeltaUpdate(const TTypedElement<UTypedElementWorldInterface>& InElementWorldHandle, const UE::Widget::EWidgetMode InWidgetMode, const EAxisList::Type InDragAxis, const FInputDeviceState& InInputState, const FTransform& InDeltaTransform, const FVector& InPivotLocation)
+void FActorElementEditorViewportInteractionCustomization::GizmoManipulationDeltaUpdate(const TTypedElement<ITypedElementWorldInterface>& InElementWorldHandle, const UE::Widget::EWidgetMode InWidgetMode, const EAxisList::Type InDragAxis, const FInputDeviceState& InInputState, const FTransform& InDeltaTransform, const FVector& InPivotLocation)
 {
 	AActor* Actor = ActorElementDataUtil::GetActorFromHandleChecked(InElementWorldHandle);
 
@@ -36,7 +36,7 @@ void FActorElementEditorViewportInteractionCustomization::GizmoManipulationDelta
 	FActorElementEditorViewportInteractionCustomization::ApplyDeltaToActor(Actor, /*bDelta*/true, &DeltaTranslation, &DeltaRotation, &DeltaScale3D, InPivotLocation, InInputState);
 }
 
-void FActorElementEditorViewportInteractionCustomization::MirrorElement(const TTypedElement<UTypedElementWorldInterface>& InElementWorldHandle, const FVector& InMirrorScale, const FVector& InPivotLocation)
+void FActorElementEditorViewportInteractionCustomization::MirrorElement(const TTypedElement<ITypedElementWorldInterface>& InElementWorldHandle, const FVector& InMirrorScale, const FVector& InPivotLocation)
 {
 	AActor* Actor = ActorElementDataUtil::GetActorFromHandleChecked(InElementWorldHandle);
 
@@ -57,7 +57,7 @@ void FActorElementEditorViewportInteractionCustomization::MirrorElement(const TT
 	Actor->MarkPackageDirty();
 }
 
-bool FActorElementEditorViewportInteractionCustomization::GetFocusBounds( const TTypedElement<UTypedElementWorldInterface>& InElementWorldHandle, FBoxSphereBounds& OutBounds)
+bool FActorElementEditorViewportInteractionCustomization::GetFocusBounds( const TTypedElement<ITypedElementWorldInterface>& InElementWorldHandle, FBoxSphereBounds& OutBounds)
 {
 	AActor* Actor = ActorElementDataUtil::GetActorFromHandle(InElementWorldHandle);
 	if (!Actor)
diff --git a/Engine/Source/Editor/UnrealEd/Private/Elements/Component/ComponentElementDetailsInterface.h b/Engine/Source/Editor/UnrealEd/Private/Elements/Component/ComponentElementDetailsInterface.h
--- a/Engine/Source/Editor/UnrealEd/Private/Elements/Component/ComponentElementDetailsInterface.h
+++ b/Engine/Source/Editor/UnrealEd/Private/Elements/Component/ComponentElementDetailsInterface.h
@@ -6,7 +6,7 @@
 #include "ComponentElementDetailsInterface.generated.h"
 
 UCLASS()
-class UComponentElementDetailsInterface : public UTypedElementDetailsInterface
+class UComponentElementDetailsInterface : public UObject, public ITypedElementDetailsInterface
 {
 	GENERATED_BODY()
 
diff --git a/Engine/Source/Editor/UnrealEd/Private/Elements/Component/ComponentElementEditorViewportInteractionCustomization.cpp b/Engine/Source/Editor/UnrealEd/Private/Elements/Component/ComponentElementEditorViewportInteractionCustomization.cpp
--- a/Engine/Source/Editor/UnrealEd/Private/Elements/Component/ComponentElementEditorViewportInteractionCustomization.cpp
+++ b/Engine/Source/Editor/UnrealEd/Private/Elements/Component/ComponentElementEditorViewportInteractionCustomization.cpp
@@ -9,7 +9,7 @@
 #include "Toolkits/IToolkitHost.h"
 #include "EditorSupportDelegates.h"
 
-bool FComponentElementEditorViewportInteractionCustomization::GetGizmoPivotLocation(const TTypedElement<UTypedElementWorldInterface>& InElementWorldHandle, const UE::Widget::EWidgetMode InWidgetMode, FVector& OutPivotLocation)
+bool FComponentElementEditorViewportInteractionCustomization::GetGizmoPivotLocation(const TTypedElement<ITypedElementWorldInterface>& InElementWorldHandle, const UE::Widget::EWidgetMode InWidgetMode, FVector& OutPivotLocation)
 {
 	const UActorComponent* Component = ComponentElementDataUtil::GetComponentFromHandleChecked(InElementWorldHandle);
 
@@ -34,7 +34,7 @@ bool FComponentElementEditorViewportInteractionCustomization::GetGizmoPivotLocat
 	return false;
 }
 
-void FComponentElementEditorViewportInteractionCustomization::GizmoManipulationDeltaUpdate(const TTypedElement<UTypedElementWorldInterface>& InElementWorldHandle, const UE::Widget::EWidgetMode InWidgetMode, const EAxisList::Type InDragAxis, const FInputDeviceState& InInputState, const FTransform& InDeltaTransform, const FVector& InPivotLocation)
+void FComponentElementEditorViewportInteractionCustomization::GizmoManipulationDeltaUpdate(const TTypedElement<ITypedElementWorldInterface>& InElementWorldHandle, const UE::Widget::EWidgetMode InWidgetMode, const EAxisList::Type InDragAxis, const FInputDeviceState& InInputState, const FTransform& InDeltaTransform, const FVector& InPivotLocation)
 {
 	UActorComponent* Component = ComponentElementDataUtil::GetComponentFromHandleChecked(InElementWorldHandle);
 
@@ -54,7 +54,7 @@ namespace EditorEngineDefs
 static const float MinComponentBoundsForZoom = 50.0f;
 }
 
-bool FComponentElementEditorViewportInteractionCustomization::GetFocusBounds(const TTypedElement<UTypedElementWorldInterface>& InElementWorldHandle, FBoxSphereBounds& OutBounds)
+bool FComponentElementEditorViewportInteractionCustomization::GetFocusBounds(const TTypedElement<ITypedElementWorldInterface>& InElementWorldHandle, FBoxSphereBounds& OutBounds)
 {
 	if (UActorComponent* Component = ComponentElementDataUtil::GetComponentFromHandle(InElementWorldHandle))
 	{
diff --git a/Engine/Source/Editor/UnrealEd/Private/Elements/Framework/EditorElements.cpp b/Engine/Source/Editor/UnrealEd/Private/Elements/Framework/EditorElements.cpp
--- a/Engine/Source/Editor/UnrealEd/Private/Elements/Framework/EditorElements.cpp
+++ b/Engine/Source/Editor/UnrealEd/Private/Elements/Framework/EditorElements.cpp
@@ -34,36 +34,36 @@ void RegisterEditorObjectElements()
 {
 	UTypedElementRegistry* Registry = UTypedElementRegistry::GetInstance();
 
-	Registry->RegisterElementInterface<UTypedElementDetailsInterface>(NAME_Object, NewObject<UObjectElementDetailsInterface>());
-	Registry->RegisterElementInterface<UTypedElementSelectionInterface>(NAME_Object, NewObject<UObjectElementEditorSelectionInterface>(), /*bAllowOverride*/true);
+	Registry->RegisterElementInterface<ITypedElementDetailsInterface>(NAME_Object, NewObject<UObjectElementDetailsInterface>());
+	Registry->RegisterElementInterface<ITypedElementSelectionInterface>(NAME_Object, NewObject<UObjectElementEditorSelectionInterface>(), /*bAllowOverride*/true);
 }
 
 void RegisterEditorActorElements()
 {
 	UTypedElementRegistry* Registry = UTypedElementRegistry::GetInstance();
 
-	Registry->RegisterElementInterface<UTypedElementDetailsInterface>(NAME_Actor, NewObject<UActorElementDetailsInterface>());
-	Registry->RegisterElementInterface<UTypedElementWorldInterface>(NAME_Actor, NewObject<UActorElementEditorWorldInterface>(), /*bAllowOverride*/true);
-	Registry->RegisterElementInterface<UTypedElementSelectionInterface>(NAME_Actor, NewObject<UActorElementEditorSelectionInterface>(), /*bAllowOverride*/true);
-	Registry->RegisterElementInterface<UTypedElementAssetDataInterface>(NAME_Actor, NewObject<UActorElementEditorAssetDataInterface>(), /*bAllowOverride*/true);
+	Registry->RegisterElementInterface<ITypedElementDetailsInterface>(NAME_Actor, NewObject<UActorElementDetailsInterface>());
+	Registry->RegisterElementInterface<ITypedElementWorldInterface>(NAME_Actor, NewObject<UActorElementEditorWorldInterface>(), /*bAllowOverride*/true);
+	Registry->RegisterElementInterface<ITypedElementSelectionInterface>(NAME_Actor, NewObject<UActorElementEditorSelectionInterface>(), /*bAllowOverride*/true);
+	Registry->RegisterElementInterface<ITypedElementAssetDataInterface>(NAME_Actor, NewObject<UActorElementEditorAssetDataInterface>(), /*bAllowOverride*/true);
 }
 
 void RegisterEditorComponentElements()
 {
 	UTypedElementRegistry* Registry = UTypedElementRegistry::GetInstance();
 
-	Registry->RegisterElementInterface<UTypedElementDetailsInterface>(NAME_Components, NewObject<UComponentElementDetailsInterface>());
-	Registry->RegisterElementInterface<UTypedElementWorldInterface>(NAME_Components, NewObject<UComponentElementEditorWorldInterface>(), /*bAllowOverride*/true);
-	Registry->RegisterElementInterface<UTypedElementSelectionInterface>(NAME_Components, NewObject<UComponentElementEditorSelectionInterface>(), /*bAllowOverride*/true);
+	Registry->RegisterElementInterface<ITypedElementDetailsInterface>(NAME_Components, NewObject<UComponentElementDetailsInterface>());
+	Registry->RegisterElementInterface<ITypedElementWorldInterface>(NAME_Components, NewObject<UComponentElementEditorWorldInterface>(), /*bAllowOverride*/true);
+	Registry->RegisterElementInterface<ITypedElementSelectionInterface>(NAME_Components, NewObject<UComponentElementEditorSelectionInterface>(), /*bAllowOverride*/true);
 }
 
 void RegisterEditorSMInstanceElements()
 {
 	UTypedElementRegistry* Registry = UTypedElementRegistry::GetInstance();
 
-	Registry->RegisterElementInterface<UTypedElementDetailsInterface>(NAME_SMInstance, NewObject<USMInstanceElementDetailsInterface>());
-	Registry->RegisterElementInterface<UTypedElementWorldInterface>(NAME_SMInstance, NewObject<USMInstanceElementEditorWorldInterface>(), /*bAllowOverride*/true);
-	Registry->RegisterElementInterface<UTypedElementSelectionInterface>(NAME_SMInstance, NewObject<USMInstanceElementEditorSelectionInterface>(), /*bAllowOverride*/true);
+	Registry->RegisterElementInterface<ITypedElementDetailsInterface>(NAME_SMInstance, NewObject<USMInstanceElementDetailsInterface>());
+	Registry->RegisterElementInterface<ITypedElementWorldInterface>(NAME_SMInstance, NewObject<USMInstanceElementEditorWorldInterface>(), /*bAllowOverride*/true);
+	Registry->RegisterElementInterface<ITypedElementSelectionInterface>(NAME_SMInstance, NewObject<USMInstanceElementEditorSelectionInterface>(), /*bAllowOverride*/true);
 
 	FEditorWidgetsModule& EditorWidgetsModule = FModuleManager::LoadModuleChecked<FEditorWidgetsModule>("EditorWidgets");
 	EditorWidgetsModule.GetObjectNameEditSinkRegistry()->RegisterObjectNameEditSink(MakeShared<FSMInstanceElementDetailsProxyObjectNameEditSink>());
diff --git a/Engine/Source/Editor/UnrealEd/Private/Elements/Object/ObjectElementDetailsInterface.h b/Engine/Source/Editor/UnrealEd/Private/Elements/Object/ObjectElementDetailsInterface.h
--- a/Engine/Source/Editor/UnrealEd/Private/Elements/Object/ObjectElementDetailsInterface.h
+++ b/Engine/Source/Editor/UnrealEd/Private/Elements/Object/ObjectElementDetailsInterface.h
@@ -6,7 +6,7 @@
 #include "ObjectElementDetailsInterface.generated.h"
 
 UCLASS()
-class UObjectElementDetailsInterface : public UTypedElementDetailsInterface
+class UObjectElementDetailsInterface : public UObject, public ITypedElementDetailsInterface
 {
 	GENERATED_BODY()
 
diff --git a/Engine/Source/Editor/UnrealEd/Private/Elements/SMInstance/SMInstanceElementDetailsInterface.h b/Engine/Source/Editor/UnrealEd/Private/Elements/SMInstance/SMInstanceElementDetailsInterface.h
--- a/Engine/Source/Editor/UnrealEd/Private/Elements/SMInstance/SMInstanceElementDetailsInterface.h
+++ b/Engine/Source/Editor/UnrealEd/Private/Elements/SMInstance/SMInstanceElementDetailsInterface.h
@@ -6,7 +6,7 @@
 #include "SMInstanceElementDetailsInterface.generated.h"
 
 UCLASS()
-class USMInstanceElementDetailsInterface : public UTypedElementDetailsInterface
+class USMInstanceElementDetailsInterface : public UObject, public ITypedElementDetailsInterface
 {
 	GENERATED_BODY()
 
diff --git a/Engine/Source/Editor/UnrealEd/Private/Factories/ActorFactory.cpp b/Engine/Source/Editor/UnrealEd/Private/Factories/ActorFactory.cpp
--- a/Engine/Source/Editor/UnrealEd/Private/Factories/ActorFactory.cpp
+++ b/Engine/Source/Editor/UnrealEd/Private/Factories/ActorFactory.cpp
@@ -320,7 +320,7 @@ void UActorFactory::PostPlaceAsset(TArrayView<const FTypedElementHandle> InEleme
 {
 	for (const FTypedElementHandle& PlacedElement : InElementHandles)
 	{
-		if (TTypedElement<UTypedElementObjectInterface> ObjectInterface = UTypedElementRegistry::GetInstance()->GetElement<UTypedElementObjectInterface>(PlacedElement))
+		if (TTypedElement<ITypedElementObjectInterface> ObjectInterface = UTypedElementRegistry::GetInstance()->GetElement<ITypedElementObjectInterface>(PlacedElement))
 		{
 			if (AActor* CreatedActor = ObjectInterface.GetObjectAs<AActor>())
 			{
@@ -347,7 +347,7 @@ FAssetData UActorFactory::GetAssetDataFromElementHandle(const FTypedElementHandl
 	}
 
 	// Check if the handle is the type of actor created by this factory, and use the factory to find the wrapped asset data if possible.
-	if (TTypedElement<UTypedElementObjectInterface> ElementObjectInterface = UTypedElementRegistry::GetInstance()->GetElement<UTypedElementObjectInterface>(InHandle))
+	if (TTypedElement<ITypedElementObjectInterface> ElementObjectInterface = UTypedElementRegistry::GetInstance()->GetElement<ITypedElementObjectInterface>(InHandle))
 	{
 		if (AActor* RawElementActorPtr = ElementObjectInterface.GetObjectAs<AActor>(GetDefaultActorClass(FAssetData())))
 		{
@@ -359,7 +359,7 @@ FAssetData UActorFactory::GetAssetDataFromElementHandle(const FTypedElementHandl
 	}
 
 	// Check if any of the referenced content is created by this factory.
-	if (TTypedElement<UTypedElementAssetDataInterface> AssetDataInterface = UTypedElementRegistry::GetInstance()->GetElement<UTypedElementAssetDataInterface>(InHandle))
+	if (TTypedElement<ITypedElementAssetDataInterface> AssetDataInterface = UTypedElementRegistry::GetInstance()->GetElement<ITypedElementAssetDataInterface>(InHandle))
 	{
 		TArray<FAssetData> ContentAssetDatas = AssetDataInterface.GetAllReferencedAssetDatas();
 		for (FAssetData& AssetData : ContentAssetDatas)
diff --git a/Engine/Source/Editor/UnrealEd/Private/Factories/EditorStaticMeshFactory.cpp b/Engine/Source/Editor/UnrealEd/Private/Factories/EditorStaticMeshFactory.cpp
--- a/Engine/Source/Editor/UnrealEd/Private/Factories/EditorStaticMeshFactory.cpp
+++ b/Engine/Source/Editor/UnrealEd/Private/Factories/EditorStaticMeshFactory.cpp
@@ -97,7 +97,7 @@ TArray<FTypedElementHandle> UEditorStaticMeshFactory::PlaceAsset(const FAssetPla
 FAssetData UEditorStaticMeshFactory::GetAssetDataFromElementHandle(const FTypedElementHandle& InHandle)
 {
 	FAssetData FoundAssetData;
-	if (TTypedElement<UTypedElementAssetDataInterface> AssetDataInterface = UTypedElementRegistry::GetInstance()->GetElement<UTypedElementAssetDataInterface>(InHandle))
+	if (TTypedElement<ITypedElementAssetDataInterface> AssetDataInterface = UTypedElementRegistry::GetInstance()->GetElement<ITypedElementAssetDataInterface>(InHandle))
 	{
 		FoundAssetData = AssetDataInterface.GetAssetData();
 	}
@@ -105,7 +105,7 @@ FAssetData UEditorStaticMeshFactory::GetAssetDataFromElementHandle(const FTypedE
 	if (!FoundAssetData.IsValid())
 	{
 		UInstancedStaticMeshComponent* ISMComponent = nullptr;
-		if (TTypedElement<UTypedElementObjectInterface> ObjectInterface = UTypedElementRegistry::GetInstance()->GetElement<UTypedElementObjectInterface>(InHandle))
+		if (TTypedElement<ITypedElementObjectInterface> ObjectInterface = UTypedElementRegistry::GetInstance()->GetElement<ITypedElementObjectInterface>(InHandle))
 		{
 			// Try to pull from component handle
 			if (UInstancedStaticMeshComponent* RawComponentPtr = ObjectInterface.GetObjectAs<UInstancedStaticMeshComponent>())
diff --git a/Engine/Source/Editor/UnrealEd/Private/LevelEditorViewport.cpp b/Engine/Source/Editor/UnrealEd/Private/LevelEditorViewport.cpp
--- a/Engine/Source/Editor/UnrealEd/Private/LevelEditorViewport.cpp
+++ b/Engine/Source/Editor/UnrealEd/Private/LevelEditorViewport.cpp
@@ -3530,7 +3530,7 @@ void FLevelEditorViewportClient::CacheElementsToManipulate(const bool bForceRefr
 		SelectionSet->GetNormalizedSelection(NormalizationOptions, CachedElementsToManipulate);
 
 		// Remove any elements that cannot be moved
-		CachedElementsToManipulate->RemoveAll<UTypedElementWorldInterface>([this](const TTypedElement<UTypedElementWorldInterface>& InWorldElement)
+		CachedElementsToManipulate->RemoveAll<ITypedElementWorldInterface>([this](const TTypedElement<ITypedElementWorldInterface>& InWorldElement)
 		{
 			if (!InWorldElement.CanMoveElement(bIsSimulateInEditorViewport ? ETypedElementWorldType::Game : ETypedElementWorldType::Editor))
 			{
diff --git a/Engine/Source/Editor/UnrealEd/Private/Selection.cpp b/Engine/Source/Editor/UnrealEd/Private/Selection.cpp
--- a/Engine/Source/Editor/UnrealEd/Private/Selection.cpp
+++ b/Engine/Source/Editor/UnrealEd/Private/Selection.cpp
@@ -167,7 +167,7 @@ bool USelection::IsValidObjectToSelect(const UObject* InObject) const
 
 	if (bIsValid && ElementSelectionSet)
 	{
-		TTypedElement<UTypedElementObjectInterface> ObjectElement = ElementSelectionSet->GetElementList()->GetElement<UTypedElementObjectInterface>(SelectionElementBridge->GetElementHandleForObject(InObject));
+		TTypedElement<ITypedElementObjectInterface> ObjectElement = ElementSelectionSet->GetElementList()->GetElement<ITypedElementObjectInterface>(SelectionElementBridge->GetElementHandleForObject(InObject));
 		if (ObjectElement)
 		{
 			bIsValid &= ObjectElement.GetObject() == InObject;
@@ -186,7 +186,7 @@ UObject* USelection::GetObjectForElementHandle(const FTypedElementHandle& InElem
 {
 	check(InElementHandle && ElementSelectionSet);
 
-	if (TTypedElement<UTypedElementObjectInterface> ObjectElement = ElementSelectionSet->GetElementList()->GetElement<UTypedElementObjectInterface>(InElementHandle))
+	if (TTypedElement<ITypedElementObjectInterface> ObjectElement = ElementSelectionSet->GetElementList()->GetElement<ITypedElementObjectInterface>(InElementHandle))
 	{
 		UObject* Object = ObjectElement.GetObject();
 		if (Object && SelectionElementBridge->IsValidObjectType(Object))
diff --git a/Engine/Source/Editor/UnrealEd/Public/Elements/Actor/ActorElementEditorViewportInteractionCustomization.h b/Engine/Source/Editor/UnrealEd/Public/Elements/Actor/ActorElementEditorViewportInteractionCustomization.h
--- a/Engine/Source/Editor/UnrealEd/Public/Elements/Actor/ActorElementEditorViewportInteractionCustomization.h
+++ b/Engine/Source/Editor/UnrealEd/Public/Elements/Actor/ActorElementEditorViewportInteractionCustomization.h
@@ -10,10 +10,10 @@ class AActor;
 class UNREALED_API FActorElementEditorViewportInteractionCustomization : public FTypedElementViewportInteractionCustomization, public FTypedElementAssetEditorToolkitHostMixin
 {
 public:
-	virtual bool GetGizmoPivotLocation(const TTypedElement<UTypedElementWorldInterface>& InElementWorldHandle, const UE::Widget::EWidgetMode InWidgetMode, FVector& OutPivotLocation) override;
-	virtual void GizmoManipulationDeltaUpdate(const TTypedElement<UTypedElementWorldInterface>& InElementWorldHandle, const UE::Widget::EWidgetMode InWidgetMode, const EAxisList::Type InDragAxis, const FInputDeviceState& InInputState, const FTransform& InDeltaTransform, const FVector& InPivotLocation) override;
-	virtual void MirrorElement(const TTypedElement<UTypedElementWorldInterface>& InElementWorldHandle, const FVector& InMirrorScale, const FVector& InPivotLocation) override;
-	virtual bool GetFocusBounds(const TTypedElement<UTypedElementWorldInterface>& InElementWorldHandle, FBoxSphereBounds& OutBounds) override;
+	virtual bool GetGizmoPivotLocation(const TTypedElement<ITypedElementWorldInterface>& InElementWorldHandle, const UE::Widget::EWidgetMode InWidgetMode, FVector& OutPivotLocation) override;
+	virtual void GizmoManipulationDeltaUpdate(const TTypedElement<ITypedElementWorldInterface>& InElementWorldHandle, const UE::Widget::EWidgetMode InWidgetMode, const EAxisList::Type InDragAxis, const FInputDeviceState& InInputState, const FTransform& InDeltaTransform, const FVector& InPivotLocation) override;
+	virtual void MirrorElement(const TTypedElement<ITypedElementWorldInterface>& InElementWorldHandle, const FVector& InMirrorScale, const FVector& InPivotLocation) override;
+	virtual bool GetFocusBounds(const TTypedElement<ITypedElementWorldInterface>& InElementWorldHandle, FBoxSphereBounds& OutBounds) override;
 	
 	static void ApplyDeltaToActor(AActor* InActor, const bool InIsDelta, const FVector* InDeltaTranslationPtr, const FRotator* InDeltaRotationPtr, const FVector* InDeltaScalePtr, const FVector& InPivotLocation, const FInputDeviceState& InInputState);
 };
diff --git a/Engine/Source/Editor/UnrealEd/Public/Elements/Component/ComponentElementEditorViewportInteractionCustomization.h b/Engine/Source/Editor/UnrealEd/Public/Elements/Component/ComponentElementEditorViewportInteractionCustomization.h
--- a/Engine/Source/Editor/UnrealEd/Public/Elements/Component/ComponentElementEditorViewportInteractionCustomization.h
+++ b/Engine/Source/Editor/UnrealEd/Public/Elements/Component/ComponentElementEditorViewportInteractionCustomization.h
@@ -11,9 +11,9 @@ class USceneComponent;
 class UNREALED_API FComponentElementEditorViewportInteractionCustomization : public FTypedElementViewportInteractionCustomization, public FTypedElementAssetEditorToolkitHostMixin
 {
 public:
-	virtual bool GetGizmoPivotLocation(const TTypedElement<UTypedElementWorldInterface>& InElementWorldHandle, const UE::Widget::EWidgetMode InWidgetMode, FVector& OutPivotLocation) override;
-	virtual void GizmoManipulationDeltaUpdate(const TTypedElement<UTypedElementWorldInterface>& InElementWorldHandle, const UE::Widget::EWidgetMode InWidgetMode, const EAxisList::Type InDragAxis, const FInputDeviceState& InInputState, const FTransform& InDeltaTransform, const FVector& InPivotLocation) override;
-	virtual bool GetFocusBounds(const TTypedElement<UTypedElementWorldInterface>& InElementWorldHandle, FBoxSphereBounds& OutBounds) override;
+	virtual bool GetGizmoPivotLocation(const TTypedElement<ITypedElementWorldInterface>& InElementWorldHandle, const UE::Widget::EWidgetMode InWidgetMode, FVector& OutPivotLocation) override;
+	virtual void GizmoManipulationDeltaUpdate(const TTypedElement<ITypedElementWorldInterface>& InElementWorldHandle, const UE::Widget::EWidgetMode InWidgetMode, const EAxisList::Type InDragAxis, const FInputDeviceState& InInputState, const FTransform& InDeltaTransform, const FVector& InPivotLocation) override;
+	virtual bool GetFocusBounds(const TTypedElement<ITypedElementWorldInterface>& InElementWorldHandle, FBoxSphereBounds& OutBounds) override;
 	
 	static void ApplyDeltaToComponent(USceneComponent* InComponent, const bool InIsDelta, const FVector* InDeltaTranslationPtr, const FRotator* InDeltaRotationPtr, const FVector* InDeltaScalePtr, const FVector& InPivotLocation, const FInputDeviceState& InInputState);
 };
diff --git a/Engine/Source/Runtime/Engine/Private/Elements/Framework/EngineElements.cpp b/Engine/Source/Runtime/Engine/Private/Elements/Framework/EngineElements.cpp
--- a/Engine/Source/Runtime/Engine/Private/Elements/Framework/EngineElements.cpp
+++ b/Engine/Source/Runtime/Engine/Private/Elements/Framework/EngineElements.cpp
@@ -39,10 +39,10 @@ void RegisterEngineObjectElements()
 	UTypedElementRegistry* Registry = UTypedElementRegistry::GetInstance();
 
 	Registry->RegisterElementType<FObjectElementData>(NAME_Object);
-	Registry->RegisterElementInterface<UTypedElementAssetDataInterface>(NAME_Object, NewObject<UObjectElementAssetDataInterface>());
-	Registry->RegisterElementInterface<UTypedElementObjectInterface>(NAME_Object, NewObject<UObjectElementObjectInterface>());
-	Registry->RegisterElementInterface<UTypedElementCounterInterface>(NAME_Object, NewObject<UObjectElementCounterInterface>());
-	Registry->RegisterElementInterface<UTypedElementSelectionInterface>(NAME_Object, NewObject<UObjectElementSelectionInterface>());
+	Registry->RegisterElementInterface<ITypedElementAssetDataInterface>(NAME_Object, NewObject<UObjectElementAssetDataInterface>());
+	Registry->RegisterElementInterface<ITypedElementObjectInterface>(NAME_Object, NewObject<UObjectElementObjectInterface>());
+	Registry->RegisterElementInterface<ITypedElementCounterInterface>(NAME_Object, NewObject<UObjectElementCounterInterface>());
+	Registry->RegisterElementInterface<ITypedElementSelectionInterface>(NAME_Object, NewObject<UObjectElementSelectionInterface>());
 }
 
 void RegisterEngineActorElements()
@@ -50,12 +50,12 @@ void RegisterEngineActorElements()
 	UTypedElementRegistry* Registry = UTypedElementRegistry::GetInstance();
 
 	Registry->RegisterElementType<FActorElementData>(NAME_Actor);
-	Registry->RegisterElementInterface<UTypedElementAssetDataInterface>(NAME_Actor, NewObject<UActorElementAssetDataInterface>());
-	Registry->RegisterElementInterface<UTypedElementHierarchyInterface>(NAME_Actor, NewObject<UActorElementHierarchyInterface>());
-	Registry->RegisterElementInterface<UTypedElementObjectInterface>(NAME_Actor, NewObject<UActorElementObjectInterface>());
-	Registry->RegisterElementInterface<UTypedElementCounterInterface>(NAME_Actor, NewObject<UActorElementCounterInterface>());
-	Registry->RegisterElementInterface<UTypedElementWorldInterface>(NAME_Actor, NewObject<UActorElementWorldInterface>());
-	Registry->RegisterElementInterface<UTypedElementSelectionInterface>(NAME_Actor, NewObject<UActorElementSelectionInterface>());
+	Registry->RegisterElementInterface<ITypedElementAssetDataInterface>(NAME_Actor, NewObject<UActorElementAssetDataInterface>());
+	Registry->RegisterElementInterface<ITypedElementHierarchyInterface>(NAME_Actor, NewObject<UActorElementHierarchyInterface>());
+	Registry->RegisterElementInterface<ITypedElementObjectInterface>(NAME_Actor, NewObject<UActorElementObjectInterface>());
+	Registry->RegisterElementInterface<ITypedElementCounterInterface>(NAME_Actor, NewObject<UActorElementCounterInterface>());
+	Registry->RegisterElementInterface<ITypedElementWorldInterface>(NAME_Actor, NewObject<UActorElementWorldInterface>());
+	Registry->RegisterElementInterface<ITypedElementSelectionInterface>(NAME_Actor, NewObject<UActorElementSelectionInterface>());
 }
 
 void RegisterEngineComponentElements()
@@ -63,11 +63,11 @@ void RegisterEngineComponentElements()
 	UTypedElementRegistry* Registry = UTypedElementRegistry::GetInstance();
 
 	Registry->RegisterElementType<FComponentElementData>(NAME_Components);
-	Registry->RegisterElementInterface<UTypedElementHierarchyInterface>(NAME_Components, NewObject<UComponentElementHierarchyInterface>());
-	Registry->RegisterElementInterface<UTypedElementObjectInterface>(NAME_Components, NewObject<UComponentElementObjectInterface>());
-	Registry->RegisterElementInterface<UTypedElementCounterInterface>(NAME_Components, NewObject<UComponentElementCounterInterface>());
-	Registry->RegisterElementInterface<UTypedElementWorldInterface>(NAME_Components, NewObject<UComponentElementWorldInterface>());
-	Registry->RegisterElementInterface<UTypedElementSelectionInterface>(NAME_Components, NewObject<UComponentElementSelectionInterface>());
+	Registry->RegisterElementInterface<ITypedElementHierarchyInterface>(NAME_Components, NewObject<UComponentElementHierarchyInterface>());
+	Registry->RegisterElementInterface<ITypedElementObjectInterface>(NAME_Components, NewObject<UComponentElementObjectInterface>());
+	Registry->RegisterElementInterface<ITypedElementCounterInterface>(NAME_Components, NewObject<UComponentElementCounterInterface>());
+	Registry->RegisterElementInterface<ITypedElementWorldInterface>(NAME_Components, NewObject<UComponentElementWorldInterface>());
+	Registry->RegisterElementInterface<ITypedElementSelectionInterface>(NAME_Components, NewObject<UComponentElementSelectionInterface>());
 }
 
 void RegisterEngineSMInstanceElements()
@@ -75,10 +75,10 @@ void RegisterEngineSMInstanceElements()
 	UTypedElementRegistry* Registry = UTypedElementRegistry::GetInstance();
 
 	Registry->RegisterElementType<FSMInstanceElementData>(NAME_SMInstance);
-	Registry->RegisterElementInterface<UTypedElementHierarchyInterface>(NAME_SMInstance, NewObject<USMInstanceElementHierarchyInterface>());
-	Registry->RegisterElementInterface<UTypedElementWorldInterface>(NAME_SMInstance, NewObject<USMInstanceElementWorldInterface>());
-	Registry->RegisterElementInterface<UTypedElementSelectionInterface>(NAME_SMInstance, NewObject<USMInstanceElementSelectionInterface>());
-	Registry->RegisterElementInterface<UTypedElementAssetDataInterface>(NAME_SMInstance, NewObject<USMInstanceElementAssetDataInterface>());
+	Registry->RegisterElementInterface<ITypedElementHierarchyInterface>(NAME_SMInstance, NewObject<USMInstanceElementHierarchyInterface>());
+	Registry->RegisterElementInterface<ITypedElementWorldInterface>(NAME_SMInstance, NewObject<USMInstanceElementWorldInterface>());
+	Registry->RegisterElementInterface<ITypedElementSelectionInterface>(NAME_SMInstance, NewObject<USMInstanceElementSelectionInterface>());
+	Registry->RegisterElementInterface<ITypedElementAssetDataInterface>(NAME_SMInstance, NewObject<USMInstanceElementAssetDataInterface>());
 }
 
 void RegisterEngineElements()
diff --git a/Engine/Source/Runtime/Engine/Private/Elements/Framework/TypedElementCommonActions.cpp b/Engine/Source/Runtime/Engine/Private/Elements/Framework/TypedElementCommonActions.cpp
--- a/Engine/Source/Runtime/Engine/Private/Elements/Framework/TypedElementCommonActions.cpp
+++ b/Engine/Source/Runtime/Engine/Private/Elements/Framework/TypedElementCommonActions.cpp
@@ -7,12 +7,12 @@
 
 #include "UObject/GCObjectScopeGuard.h"
 
-bool FTypedElementCommonActionsCustomization::DeleteElements(UTypedElementWorldInterface* InWorldInterface, TArrayView<const FTypedElementHandle> InElementHandles, UWorld* InWorld, UTypedElementSelectionSet* InSelectionSet, const FTypedElementDeletionOptions& InDeletionOptions)
+bool FTypedElementCommonActionsCustomization::DeleteElements(ITypedElementWorldInterface* InWorldInterface, TArrayView<const FTypedElementHandle> InElementHandles, UWorld* InWorld, UTypedElementSelectionSet* InSelectionSet, const FTypedElementDeletionOptions& InDeletionOptions)
 {
 	return InWorldInterface->DeleteElements(InElementHandles, InWorld, InSelectionSet, InDeletionOptions);
 }
 
-void FTypedElementCommonActionsCustomization::DuplicateElements(UTypedElementWorldInterface* InWorldInterface, TArrayView<const FTypedElementHandle> InElementHandles, UWorld* InWorld, const FVector& InLocationOffset, TArray<FTypedElementHandle>& OutNewElements)
+void FTypedElementCommonActionsCustomization::DuplicateElements(ITypedElementWorldInterface* InWorldInterface, TArrayView<const FTypedElementHandle> InElementHandles, UWorld* InWorld, const FVector& InLocationOffset, TArray<FTypedElementHandle>& OutNewElements)
 {
 	InWorldInterface->DuplicateElements(InElementHandles, InWorld, InLocationOffset, OutNewElements);
 }
@@ -39,7 +39,7 @@ bool UTypedElementCommonActions::DeleteNormalizedElements(const FTypedElementLis
 		for (const auto& ElementsByTypePair : ElementsToDeleteByType)
 		{
 			FTypedElementCommonActionsCustomization* CommonActionsCustomization = GetInterfaceCustomizationByTypeId(ElementsByTypePair.Key);
-			UTypedElementWorldInterface* WorldInterface = Registry->GetElementInterface<UTypedElementWorldInterface>(ElementsByTypePair.Key);
+			ITypedElementWorldInterface* WorldInterface = Registry->GetElementInterface<ITypedElementWorldInterface>(ElementsByTypePair.Key);
 			if (CommonActionsCustomization && WorldInterface)
 			{
 				bSuccess |= CommonActionsCustomization->DeleteElements(WorldInterface, ElementsByTypePair.Value, World, SelectionSet, DeletionOptions);
@@ -70,7 +70,7 @@ TArray<FTypedElementHandle> UTypedElementCommonActions::DuplicateNormalizedEleme
 		for (const auto& ElementsByTypePair : ElementsToDuplicateByType)
 		{
 			FTypedElementCommonActionsCustomization* CommonActionsCustomization = GetInterfaceCustomizationByTypeId(ElementsByTypePair.Key);
-			UTypedElementWorldInterface* WorldInterface = Registry->GetElementInterface<UTypedElementWorldInterface>(ElementsByTypePair.Key);
+			ITypedElementWorldInterface* WorldInterface = Registry->GetElementInterface<ITypedElementWorldInterface>(ElementsByTypePair.Key);
 			if (CommonActionsCustomization && WorldInterface)
 			{
 				CommonActionsCustomization->DuplicateElements(WorldInterface, ElementsByTypePair.Value, World, LocationOffset, NewElements);
@@ -84,6 +84,6 @@ TArray<FTypedElementHandle> UTypedElementCommonActions::DuplicateNormalizedEleme
 FTypedElementCommonActionsElement UTypedElementCommonActions::ResolveCommonActionsElement(const FTypedElementHandle& InElementHandle) const
 {
 	return InElementHandle
-		? FTypedElementCommonActionsElement(UTypedElementRegistry::GetInstance()->GetElement<UTypedElementWorldInterface>(InElementHandle), GetInterfaceCustomizationByTypeId(InElementHandle.GetId().GetTypeId()))
+		? FTypedElementCommonActionsElement(UTypedElementRegistry::GetInstance()->GetElement<ITypedElementWorldInterface>(InElementHandle), GetInterfaceCustomizationByTypeId(InElementHandle.GetId().GetTypeId()))
 		: FTypedElementCommonActionsElement();
 }
diff --git a/Engine/Source/Runtime/Engine/Private/Elements/SMInstance/SMInstanceElementSelectionInterface.cpp b/Engine/Source/Runtime/Engine/Private/Elements/SMInstance/SMInstanceElementSelectionInterface.cpp
--- a/Engine/Source/Runtime/Engine/Private/Elements/SMInstance/SMInstanceElementSelectionInterface.cpp
+++ b/Engine/Source/Runtime/Engine/Private/Elements/SMInstance/SMInstanceElementSelectionInterface.cpp
@@ -8,7 +8,7 @@
 bool USMInstanceElementSelectionInterface::SelectElement(const FTypedElementHandle& InElementHandle, FTypedElementListProxy InSelectionSet, const FTypedElementSelectionOptions& InSelectionOptions)
 {
 	FSMInstanceManager SMInstance = SMInstanceElementDataUtil::GetSMInstanceFromHandle(InElementHandle);
-	if (SMInstance && Super::SelectElement(InElementHandle, InSelectionSet, InSelectionOptions))
+	if (SMInstance && ITypedElementSelectionInterface::SelectElement(InElementHandle, InSelectionSet, InSelectionOptions))
 	{
 		SMInstance.NotifySMInstanceSelectionChanged(/*bIsSelected*/true);
 		return true;
@@ -19,7 +19,7 @@ bool USMInstanceElementSelectionInterface::SelectElement(const FTypedElementHand
 bool USMInstanceElementSelectionInterface::DeselectElement(const FTypedElementHandle& InElementHandle, FTypedElementListProxy InSelectionSet, const FTypedElementSelectionOptions& InSelectionOptions)
 {
 	FSMInstanceManager SMInstance = SMInstanceElementDataUtil::GetSMInstanceFromHandle(InElementHandle);
-	if (SMInstance && Super::DeselectElement(InElementHandle, InSelectionSet, InSelectionOptions))
+	if (SMInstance && ITypedElementSelectionInterface::DeselectElement(InElementHandle, InSelectionSet, InSelectionOptions))
 	{
 		SMInstance.NotifySMInstanceSelectionChanged(/*bIsSelected*/false);
 		return true;
diff --git a/Engine/Source/Runtime/Engine/Public/Elements/Actor/ActorElementAssetDataInterface.h b/Engine/Source/Runtime/Engine/Public/Elements/Actor/ActorElementAssetDataInterface.h
--- a/Engine/Source/Runtime/Engine/Public/Elements/Actor/ActorElementAssetDataInterface.h
+++ b/Engine/Source/Runtime/Engine/Public/Elements/Actor/ActorElementAssetDataInterface.h
@@ -7,7 +7,7 @@
 #include "ActorElementAssetDataInterface.generated.h"
 
 UCLASS()
-class ENGINE_API UActorElementAssetDataInterface : public UTypedElementAssetDataInterface
+class ENGINE_API UActorElementAssetDataInterface : public UObject, public ITypedElementAssetDataInterface
 {
 	GENERATED_BODY()
 
diff --git a/Engine/Source/Runtime/Engine/Public/Elements/Actor/ActorElementCounterInterface.h b/Engine/Source/Runtime/Engine/Public/Elements/Actor/ActorElementCounterInterface.h
--- a/Engine/Source/Runtime/Engine/Public/Elements/Actor/ActorElementCounterInterface.h
+++ b/Engine/Source/Runtime/Engine/Public/Elements/Actor/ActorElementCounterInterface.h
@@ -6,7 +6,7 @@
 #include "ActorElementCounterInterface.generated.h"
 
 UCLASS()
-class ENGINE_API UActorElementCounterInterface : public UTypedElementCounterInterface
+class ENGINE_API UActorElementCounterInterface : public UObject, public ITypedElementCounterInterface
 {
 	GENERATED_BODY()
 
diff --git a/Engine/Source/Runtime/Engine/Public/Elements/Actor/ActorElementHierarchyInterface.h b/Engine/Source/Runtime/Engine/Public/Elements/Actor/ActorElementHierarchyInterface.h
--- a/Engine/Source/Runtime/Engine/Public/Elements/Actor/ActorElementHierarchyInterface.h
+++ b/Engine/Source/Runtime/Engine/Public/Elements/Actor/ActorElementHierarchyInterface.h
@@ -6,7 +6,7 @@
 #include "ActorElementHierarchyInterface.generated.h"
 
 UCLASS()
-class ENGINE_API UActorElementHierarchyInterface : public UTypedElementHierarchyInterface
+class ENGINE_API UActorElementHierarchyInterface : public UObject, public ITypedElementHierarchyInterface
 {
 	GENERATED_BODY()
 
diff --git a/Engine/Source/Runtime/Engine/Public/Elements/Actor/ActorElementObjectInterface.h b/Engine/Source/Runtime/Engine/Public/Elements/Actor/ActorElementObjectInterface.h
--- a/Engine/Source/Runtime/Engine/Public/Elements/Actor/ActorElementObjectInterface.h
+++ b/Engine/Source/Runtime/Engine/Public/Elements/Actor/ActorElementObjectInterface.h
@@ -6,7 +6,7 @@
 #include "ActorElementObjectInterface.generated.h"
 
 UCLASS()
-class ENGINE_API UActorElementObjectInterface : public UTypedElementObjectInterface
+class ENGINE_API UActorElementObjectInterface : public UObject, public ITypedElementObjectInterface
 {
 	GENERATED_BODY()
 
diff --git a/Engine/Source/Runtime/Engine/Public/Elements/Actor/ActorElementSelectionInterface.h b/Engine/Source/Runtime/Engine/Public/Elements/Actor/ActorElementSelectionInterface.h
--- a/Engine/Source/Runtime/Engine/Public/Elements/Actor/ActorElementSelectionInterface.h
+++ b/Engine/Source/Runtime/Engine/Public/Elements/Actor/ActorElementSelectionInterface.h
@@ -6,7 +6,7 @@
 #include "ActorElementSelectionInterface.generated.h"
 
 UCLASS()
-class ENGINE_API UActorElementSelectionInterface : public UTypedElementSelectionInterface
+class ENGINE_API UActorElementSelectionInterface : public UObject, public ITypedElementSelectionInterface
 {
 	GENERATED_BODY()
 };
diff --git a/Engine/Source/Runtime/Engine/Public/Elements/Actor/ActorElementWorldInterface.h b/Engine/Source/Runtime/Engine/Public/Elements/Actor/ActorElementWorldInterface.h
--- a/Engine/Source/Runtime/Engine/Public/Elements/Actor/ActorElementWorldInterface.h
+++ b/Engine/Source/Runtime/Engine/Public/Elements/Actor/ActorElementWorldInterface.h
@@ -8,7 +8,7 @@
 struct FCollisionQueryParams;
 
 UCLASS()
-class ENGINE_API UActorElementWorldInterface : public UTypedElementWorldInterface
+class ENGINE_API UActorElementWorldInterface : public UObject, public ITypedElementWorldInterface
 {
 	GENERATED_BODY()
 
diff --git a/Engine/Source/Runtime/Engine/Public/Elements/Component/ComponentElementCounterInterface.h b/Engine/Source/Runtime/Engine/Public/Elements/Component/ComponentElementCounterInterface.h
--- a/Engine/Source/Runtime/Engine/Public/Elements/Component/ComponentElementCounterInterface.h
+++ b/Engine/Source/Runtime/Engine/Public/Elements/Component/ComponentElementCounterInterface.h
@@ -6,7 +6,7 @@
 #include "ComponentElementCounterInterface.generated.h"
 
 UCLASS()
-class ENGINE_API UComponentElementCounterInterface : public UTypedElementCounterInterface
+class ENGINE_API UComponentElementCounterInterface : public UObject, public ITypedElementCounterInterface
 {
 	GENERATED_BODY()
 
diff --git a/Engine/Source/Runtime/Engine/Public/Elements/Component/ComponentElementHierarchyInterface.h b/Engine/Source/Runtime/Engine/Public/Elements/Component/ComponentElementHierarchyInterface.h
--- a/Engine/Source/Runtime/Engine/Public/Elements/Component/ComponentElementHierarchyInterface.h
+++ b/Engine/Source/Runtime/Engine/Public/Elements/Component/ComponentElementHierarchyInterface.h
@@ -6,7 +6,7 @@
 #include "ComponentElementHierarchyInterface.generated.h"
 
 UCLASS()
-class ENGINE_API UComponentElementHierarchyInterface : public UTypedElementHierarchyInterface
+class ENGINE_API UComponentElementHierarchyInterface : public UObject, public ITypedElementHierarchyInterface
 {
 	GENERATED_BODY()
 
diff --git a/Engine/Source/Runtime/Engine/Public/Elements/Component/ComponentElementObjectInterface.h b/Engine/Source/Runtime/Engine/Public/Elements/Component/ComponentElementObjectInterface.h
--- a/Engine/Source/Runtime/Engine/Public/Elements/Component/ComponentElementObjectInterface.h
+++ b/Engine/Source/Runtime/Engine/Public/Elements/Component/ComponentElementObjectInterface.h
@@ -6,7 +6,7 @@
 #include "ComponentElementObjectInterface.generated.h"
 
 UCLASS()
-class ENGINE_API UComponentElementObjectInterface : public UTypedElementObjectInterface
+class ENGINE_API UComponentElementObjectInterface : public UObject, public ITypedElementObjectInterface
 {
 	GENERATED_BODY()
 
diff --git a/Engine/Source/Runtime/Engine/Public/Elements/Component/ComponentElementSelectionInterface.h b/Engine/Source/Runtime/Engine/Public/Elements/Component/ComponentElementSelectionInterface.h
--- a/Engine/Source/Runtime/Engine/Public/Elements/Component/ComponentElementSelectionInterface.h
+++ b/Engine/Source/Runtime/Engine/Public/Elements/Component/ComponentElementSelectionInterface.h
@@ -6,7 +6,7 @@
 #include "ComponentElementSelectionInterface.generated.h"
 
 UCLASS()
-class ENGINE_API UComponentElementSelectionInterface : public UTypedElementSelectionInterface
+class ENGINE_API UComponentElementSelectionInterface : public UObject, public ITypedElementSelectionInterface
 {
 	GENERATED_BODY()
 };
diff --git a/Engine/Source/Runtime/Engine/Public/Elements/Component/ComponentElementWorldInterface.h b/Engine/Source/Runtime/Engine/Public/Elements/Component/ComponentElementWorldInterface.h
--- a/Engine/Source/Runtime/Engine/Public/Elements/Component/ComponentElementWorldInterface.h
+++ b/Engine/Source/Runtime/Engine/Public/Elements/Component/ComponentElementWorldInterface.h
@@ -6,7 +6,7 @@
 #include "ComponentElementWorldInterface.generated.h"
 
 UCLASS()
-class ENGINE_API UComponentElementWorldInterface : public UTypedElementWorldInterface
+class ENGINE_API UComponentElementWorldInterface : public UObject, public ITypedElementWorldInterface
 {
 	GENERATED_BODY()
 
diff --git a/Engine/Source/Runtime/Engine/Public/Elements/Framework/TypedElementCommonActions.h b/Engine/Source/Runtime/Engine/Public/Elements/Framework/TypedElementCommonActions.h
--- a/Engine/Source/Runtime/Engine/Public/Elements/Framework/TypedElementCommonActions.h
+++ b/Engine/Source/Runtime/Engine/Public/Elements/Framework/TypedElementCommonActions.h
@@ -18,8 +18,8 @@ public:
 	virtual ~FTypedElementCommonActionsCustomization() = default;
 
 	//~ See UTypedElementCommonActions for API docs
-	virtual bool DeleteElements(UTypedElementWorldInterface* InWorldInterface, TArrayView<const FTypedElementHandle> InElementHandles, UWorld* InWorld, UTypedElementSelectionSet* InSelectionSet, const FTypedElementDeletionOptions& InDeletionOptions);
-	virtual void DuplicateElements(UTypedElementWorldInterface* InWorldInterface, TArrayView<const FTypedElementHandle> InElementHandles, UWorld* InWorld, const FVector& InLocationOffset, TArray<FTypedElementHandle>& OutNewElements);
+	virtual bool DeleteElements(ITypedElementWorldInterface* InWorldInterface, TArrayView<const FTypedElementHandle> InElementHandles, UWorld* InWorld, UTypedElementSelectionSet* InSelectionSet, const FTypedElementDeletionOptions& InDeletionOptions);
+	virtual void DuplicateElements(ITypedElementWorldInterface* InWorldInterface, TArrayView<const FTypedElementHandle> InElementHandles, UWorld* InWorld, const FVector& InLocationOffset, TArray<FTypedElementHandle>& OutNewElements);
 };
 
 /**
@@ -30,7 +30,7 @@ struct ENGINE_API FTypedElementCommonActionsElement
 public:
 	FTypedElementCommonActionsElement() = default;
 
-	FTypedElementCommonActionsElement(TTypedElement<UTypedElementWorldInterface> InElementWorldHandle, FTypedElementCommonActionsCustomization* InCommonActionsCustomization)
+	FTypedElementCommonActionsElement(TTypedElement<ITypedElementWorldInterface> InElementWorldHandle, FTypedElementCommonActionsCustomization* InCommonActionsCustomization)
 		: ElementWorldHandle(MoveTemp(InElementWorldHandle))
 		, CommonActionsCustomization(InCommonActionsCustomization)
 	{
@@ -56,7 +56,7 @@ public:
 	//~ See UTypedElementCommonActions for API docs
 
 private:
-	TTypedElement<UTypedElementWorldInterface> ElementWorldHandle;
+	TTypedElement<ITypedElementWorldInterface> ElementWorldHandle;
 	FTypedElementCommonActionsCustomization* CommonActionsCustomization = nullptr;
 };
 
diff --git a/Engine/Source/Runtime/Engine/Public/Elements/Interfaces/TypedElementWorldInterface.h b/Engine/Source/Runtime/Engine/Public/Elements/Interfaces/TypedElementWorldInterface.h
--- a/Engine/Source/Runtime/Engine/Public/Elements/Interfaces/TypedElementWorldInterface.h
+++ b/Engine/Source/Runtime/Engine/Public/Elements/Interfaces/TypedElementWorldInterface.h
@@ -3,6 +3,8 @@
 #pragma once
 
 #include "Elements/Framework/TypedElementHandle.h"
+#include "UObject/Interface.h"
+
 #include "TypedElementWorldInterface.generated.h"
 
 class ULevel;
@@ -43,8 +45,13 @@ private:
 	bool bWarnAboutSoftReferences = true;
 };
 
-UCLASS(Abstract)
-class ENGINE_API UTypedElementWorldInterface : public UTypedElementInterface
+UINTERFACE(MinimalAPI, BlueprintType, meta = (CannotImplementInterfaceInBlueprint))
+class UTypedElementWorldInterface : public UInterface
+{
+	GENERATED_BODY()
+};
+
+class ENGINE_API ITypedElementWorldInterface
 {
 	GENERATED_BODY()
 
@@ -52,43 +59,43 @@ public:
 	/**
 	 * Is this element considered a template within its world (eg, a CDO or archetype).
 	 */
-	UFUNCTION(BlueprintPure, Category="TypedElementInterfaces|World")
+	UFUNCTION(BlueprintCallable, Category="TypedElementInterfaces|World")
 	virtual bool IsTemplateElement(const FTypedElementHandle& InElementHandle) { return false; }
 
 	/**
 	 * Can this element actually be edited in the world?
 	 */
-	UFUNCTION(BlueprintPure, Category="TypedElementInterfaces|World")
+	UFUNCTION(BlueprintCallable, Category="TypedElementInterfaces|World")
 	virtual bool CanEditElement(const FTypedElementHandle& InElementHandle) { return true; }
 
 	/**
 	 * Get the owner level associated with this element, if any.
 	 */
-	UFUNCTION(BlueprintPure, Category="TypedElementInterfaces|World")
+	UFUNCTION(BlueprintCallable, Category="TypedElementInterfaces|World")
 	virtual ULevel* GetOwnerLevel(const FTypedElementHandle& InElementHandle) { return nullptr; }
 
 	/**
 	 * Get the owner world associated with this element, if any.
 	 */
-	UFUNCTION(BlueprintPure, Category="TypedElementInterfaces|World")
+	UFUNCTION(BlueprintCallable, Category="TypedElementInterfaces|World")
 	virtual UWorld* GetOwnerWorld(const FTypedElementHandle& InElementHandle) { return nullptr; }
 
 	/**
 	 * Get the bounds of this element, if any.
 	 */
-	UFUNCTION(BlueprintPure, Category="TypedElementInterfaces|World")
+	UFUNCTION(BlueprintCallable, Category="TypedElementInterfaces|World")
 	virtual bool GetBounds(const FTypedElementHandle& InElementHandle, FBoxSphereBounds& OutBounds) { return false; }
 
 	/**
 	 * Can the given element be moved within the world?
 	 */
-	UFUNCTION(BlueprintPure, Category="TypedElementInterfaces|World")
+	UFUNCTION(BlueprintCallable, Category="TypedElementInterfaces|World")
 	virtual bool CanMoveElement(const FTypedElementHandle& InElementHandle, const ETypedElementWorldType InWorldType) { return false; }
 
 	/**
 	 * Get the transform of this element within its owner world, if any.
 	 */
-	UFUNCTION(BlueprintPure, Category="TypedElementInterfaces|World")
+	UFUNCTION(BlueprintCallable, Category="TypedElementInterfaces|World")
 	virtual bool GetWorldTransform(const FTypedElementHandle& InElementHandle, FTransform& OutTransform) { return false; }
 	
 	/**
@@ -100,7 +107,7 @@ public:
 	/**
 	 * Get the transform of this element relative to its parent, if any.
 	 */
-	UFUNCTION(BlueprintPure, Category="TypedElementInterfaces|World")
+	UFUNCTION(BlueprintCallable, Category="TypedElementInterfaces|World")
 	virtual bool GetRelativeTransform(const FTypedElementHandle& InElementHandle, FTransform& OutTransform) { return GetWorldTransform(InElementHandle, OutTransform); }
 	
 	/**
@@ -112,7 +119,7 @@ public:
 	/**
 	 * Get the local space offset of this element that should be added to its pivot location, if any.
 	 */
-	UFUNCTION(BlueprintPure, Category="TypedElementInterfaces|World")
+	UFUNCTION(BlueprintCallable, Category="TypedElementInterfaces|World")
 	virtual bool GetPivotOffset(const FTypedElementHandle& InElementHandle, FVector& OutPivotOffset) { return false; }
 
 	/**
@@ -206,7 +213,7 @@ public:
 };
 
 template <>
-struct TTypedElement<UTypedElementWorldInterface> : public TTypedElementBase<UTypedElementWorldInterface>
+struct TTypedElement<ITypedElementWorldInterface> : public TTypedElementBase<ITypedElementWorldInterface>
 {
 	bool IsTemplateElement() const { return InterfacePtr->IsTemplateElement(*this); }
 	bool CanEditElement() const { return InterfacePtr->CanEditElement(*this); }
diff --git a/Engine/Source/Runtime/Engine/Public/Elements/Object/ObjectElementAssetDataInterface.h b/Engine/Source/Runtime/Engine/Public/Elements/Object/ObjectElementAssetDataInterface.h
--- a/Engine/Source/Runtime/Engine/Public/Elements/Object/ObjectElementAssetDataInterface.h
+++ b/Engine/Source/Runtime/Engine/Public/Elements/Object/ObjectElementAssetDataInterface.h
@@ -7,7 +7,7 @@
 #include "ObjectElementAssetDataInterface.generated.h"
 
 UCLASS()
-class ENGINE_API UObjectElementAssetDataInterface : public UTypedElementAssetDataInterface
+class ENGINE_API UObjectElementAssetDataInterface : public UObject, public ITypedElementAssetDataInterface
 {
 	GENERATED_BODY()
 
diff --git a/Engine/Source/Runtime/Engine/Public/Elements/Object/ObjectElementCounterInterface.h b/Engine/Source/Runtime/Engine/Public/Elements/Object/ObjectElementCounterInterface.h
--- a/Engine/Source/Runtime/Engine/Public/Elements/Object/ObjectElementCounterInterface.h
+++ b/Engine/Source/Runtime/Engine/Public/Elements/Object/ObjectElementCounterInterface.h
@@ -6,7 +6,7 @@
 #include "ObjectElementCounterInterface.generated.h"
 
 UCLASS()
-class ENGINE_API UObjectElementCounterInterface : public UTypedElementCounterInterface
+class ENGINE_API UObjectElementCounterInterface : public UObject, public ITypedElementCounterInterface
 {
 	GENERATED_BODY()
 
diff --git a/Engine/Source/Runtime/Engine/Public/Elements/Object/ObjectElementObjectInterface.h b/Engine/Source/Runtime/Engine/Public/Elements/Object/ObjectElementObjectInterface.h
--- a/Engine/Source/Runtime/Engine/Public/Elements/Object/ObjectElementObjectInterface.h
+++ b/Engine/Source/Runtime/Engine/Public/Elements/Object/ObjectElementObjectInterface.h
@@ -6,7 +6,7 @@
 #include "ObjectElementObjectInterface.generated.h"
 
 UCLASS()
-class ENGINE_API UObjectElementObjectInterface : public UTypedElementObjectInterface
+class ENGINE_API UObjectElementObjectInterface : public UObject, public ITypedElementObjectInterface
 {
 	GENERATED_BODY()
 
diff --git a/Engine/Source/Runtime/Engine/Public/Elements/Object/ObjectElementSelectionInterface.h b/Engine/Source/Runtime/Engine/Public/Elements/Object/ObjectElementSelectionInterface.h
--- a/Engine/Source/Runtime/Engine/Public/Elements/Object/ObjectElementSelectionInterface.h
+++ b/Engine/Source/Runtime/Engine/Public/Elements/Object/ObjectElementSelectionInterface.h
@@ -6,7 +6,7 @@
 #include "ObjectElementSelectionInterface.generated.h"
 
 UCLASS()
-class ENGINE_API UObjectElementSelectionInterface : public UTypedElementSelectionInterface
+class ENGINE_API UObjectElementSelectionInterface : public UObject, public ITypedElementSelectionInterface
 {
 	GENERATED_BODY()
 
diff --git a/Engine/Source/Runtime/Engine/Public/Elements/SMInstance/SMInstanceElementAssetDataInterface.h b/Engine/Source/Runtime/Engine/Public/Elements/SMInstance/SMInstanceElementAssetDataInterface.h
--- a/Engine/Source/Runtime/Engine/Public/Elements/SMInstance/SMInstanceElementAssetDataInterface.h
+++ b/Engine/Source/Runtime/Engine/Public/Elements/SMInstance/SMInstanceElementAssetDataInterface.h
@@ -7,7 +7,7 @@
 #include "SMInstanceElementAssetDataInterface.generated.h"
 
 UCLASS()
-class ENGINE_API USMInstanceElementAssetDataInterface : public UTypedElementAssetDataInterface
+class ENGINE_API USMInstanceElementAssetDataInterface : public UObject, public ITypedElementAssetDataInterface
 {
 	GENERATED_BODY()
 
diff --git a/Engine/Source/Runtime/Engine/Public/Elements/SMInstance/SMInstanceElementHierarchyInterface.h b/Engine/Source/Runtime/Engine/Public/Elements/SMInstance/SMInstanceElementHierarchyInterface.h
--- a/Engine/Source/Runtime/Engine/Public/Elements/SMInstance/SMInstanceElementHierarchyInterface.h
+++ b/Engine/Source/Runtime/Engine/Public/Elements/SMInstance/SMInstanceElementHierarchyInterface.h
@@ -6,7 +6,7 @@
 #include "SMInstanceElementHierarchyInterface.generated.h"
 
 UCLASS()
-class ENGINE_API USMInstanceElementHierarchyInterface : public UTypedElementHierarchyInterface
+class ENGINE_API USMInstanceElementHierarchyInterface : public UObject, public ITypedElementHierarchyInterface
 {
 	GENERATED_BODY()
 
diff --git a/Engine/Source/Runtime/Engine/Public/Elements/SMInstance/SMInstanceElementSelectionInterface.h b/Engine/Source/Runtime/Engine/Public/Elements/SMInstance/SMInstanceElementSelectionInterface.h
--- a/Engine/Source/Runtime/Engine/Public/Elements/SMInstance/SMInstanceElementSelectionInterface.h
+++ b/Engine/Source/Runtime/Engine/Public/Elements/SMInstance/SMInstanceElementSelectionInterface.h
@@ -6,7 +6,7 @@
 #include "SMInstanceElementSelectionInterface.generated.h"
 
 UCLASS()
-class ENGINE_API USMInstanceElementSelectionInterface : public UTypedElementSelectionInterface
+class ENGINE_API USMInstanceElementSelectionInterface : public UObject, public ITypedElementSelectionInterface
 {
 	GENERATED_BODY()
 
diff --git a/Engine/Source/Runtime/Engine/Public/Elements/SMInstance/SMInstanceElementWorldInterface.h b/Engine/Source/Runtime/Engine/Public/Elements/SMInstance/SMInstanceElementWorldInterface.h
--- a/Engine/Source/Runtime/Engine/Public/Elements/SMInstance/SMInstanceElementWorldInterface.h
+++ b/Engine/Source/Runtime/Engine/Public/Elements/SMInstance/SMInstanceElementWorldInterface.h
@@ -6,7 +6,7 @@
 #include "SMInstanceElementWorldInterface.generated.h"
 
 UCLASS()
-class ENGINE_API USMInstanceElementWorldInterface : public UTypedElementWorldInterface
+class ENGINE_API USMInstanceElementWorldInterface : public UObject, public ITypedElementWorldInterface
 {
 	GENERATED_BODY()
 
diff --git a/Engine/Source/Runtime/TypedElementFramework/Private/Elements/Framework/TypedElementCounter.cpp b/Engine/Source/Runtime/TypedElementFramework/Private/Elements/Framework/TypedElementCounter.cpp
--- a/Engine/Source/Runtime/TypedElementFramework/Private/Elements/Framework/TypedElementCounter.cpp
+++ b/Engine/Source/Runtime/TypedElementFramework/Private/Elements/Framework/TypedElementCounter.cpp
@@ -25,7 +25,7 @@ void FTypedElementCounter::AddElement(const FTypedElementHandle& InElementHandle
 
 	IncrementCounter(GetElementTypeCategoryName(), InElementHandle.GetId().GetTypeId());
 	
-	if (const TTypedElement<UTypedElementCounterInterface> ElementCounterHandle = RegistryPtr->GetElement<UTypedElementCounterInterface>(InElementHandle))
+	if (const TTypedElement<ITypedElementCounterInterface> ElementCounterHandle = RegistryPtr->GetElement<ITypedElementCounterInterface>(InElementHandle))
 	{
 		ElementCounterHandle.IncrementCountersForElement(*this);
 	}
@@ -41,7 +41,7 @@ void FTypedElementCounter::RemoveElement(const FTypedElementHandle& InElementHan
 
 	DecrementCounter(GetElementTypeCategoryName(), InElementHandle.GetId().GetTypeId());
 
-	if (const TTypedElement<UTypedElementCounterInterface> ElementCounterHandle = RegistryPtr->GetElement<UTypedElementCounterInterface>(InElementHandle))
+	if (const TTypedElement<ITypedElementCounterInterface> ElementCounterHandle = RegistryPtr->GetElement<ITypedElementCounterInterface>(InElementHandle))
 	{
 		ElementCounterHandle.DecrementCountersForElement(*this);
 	}
diff --git a/Engine/Source/Runtime/TypedElementFramework/Private/Elements/Framework/TypedElementList.cpp b/Engine/Source/Runtime/TypedElementFramework/Private/Elements/Framework/TypedElementList.cpp
--- a/Engine/Source/Runtime/TypedElementFramework/Private/Elements/Framework/TypedElementList.cpp
+++ b/Engine/Source/Runtime/TypedElementFramework/Private/Elements/Framework/TypedElementList.cpp
@@ -6,7 +6,7 @@
 namespace TypedElementList_Private
 {
 
-void GetElementImpl(const UTypedElementRegistry* InRegistry, const FTypedElementHandle& InElementHandle, const UClass* InBaseInterfaceType, FTypedElement& OutElement)
+void GetElementImpl(const UTypedElementRegistry* InRegistry, const FTypedElementHandle& InElementHandle, const TSubclassOf<UInterface>& InBaseInterfaceType, FTypedElement& OutElement)
 {
 	if (InRegistry)
 	{
@@ -148,7 +148,7 @@ FTypedElementListRef FTypedElementList::Clone() const
 	return ClonedElementList;
 }
 
-UTypedElementInterface* FTypedElementList::GetElementInterface(const FTypedElementHandle& InElementHandle, const TSubclassOf<UTypedElementInterface>& InBaseInterfaceType) const
+UObject* FTypedElementList::GetElementInterface(const FTypedElementHandle& InElementHandle, const TSubclassOf<UInterface>& InBaseInterfaceType) const
 {
 	UTypedElementRegistry* RegistryPtr = Registry.Get();
 	return RegistryPtr
@@ -156,7 +156,7 @@ UTypedElementInterface* FTypedElementList::GetElementInterface(const FTypedEleme
 		: nullptr;
 }
 
-bool FTypedElementList::HasElements(const TSubclassOf<UTypedElementInterface>& InBaseInterfaceType) const
+bool FTypedElementList::HasElements(const TSubclassOf<UInterface>& InBaseInterfaceType) const
 {
 	bool bHasFilteredElements = false;
 
@@ -176,7 +176,7 @@ bool FTypedElementList::HasElements(const TSubclassOf<UTypedElementInterface>& I
 	return bHasFilteredElements;
 }
 
-int32 FTypedElementList::CountElements(const TSubclassOf<UTypedElementInterface>& InBaseInterfaceType) const
+int32 FTypedElementList::CountElements(const TSubclassOf<UInterface>& InBaseInterfaceType) const
 {
 	int32 NumFilteredElements = 0;
 
@@ -224,7 +224,7 @@ int32 FTypedElementList::CountElementsOfType(const FTypedHandleTypeId InElementT
 	return ElementCounts.GetCounterValue(FTypedElementCounter::GetElementTypeCategoryName(), InElementTypeId);
 }
 
-TArray<FTypedElementHandle> FTypedElementList::GetElementHandles(const TSubclassOf<UTypedElementInterface>& InBaseInterfaceType) const
+TArray<FTypedElementHandle> FTypedElementList::GetElementHandles(const TSubclassOf<UInterface>& InBaseInterfaceType) const
 {
 	TArray<FTypedElementHandle> FilteredElementHandles;
 	FilteredElementHandles.Reserve(ElementHandles.Num());
@@ -238,7 +238,7 @@ TArray<FTypedElementHandle> FTypedElementList::GetElementHandles(const TSubclass
 	return FilteredElementHandles;
 }
 
-void FTypedElementList::ForEachElementHandle(TFunctionRef<bool(const FTypedElementHandle&)> InCallback, const TSubclassOf<UTypedElementInterface>& InBaseInterfaceType) const
+void FTypedElementList::ForEachElementHandle(TFunctionRef<bool(const FTypedElementHandle&)> InCallback, const TSubclassOf<UInterface>& InBaseInterfaceType) const
 {
 	for (const FTypedElementHandle& ElementHandle : ElementHandles)
 	{
diff --git a/Engine/Source/Runtime/TypedElementFramework/Private/Elements/Framework/TypedElementListLibrary.h b/Engine/Source/Runtime/TypedElementFramework/Private/Elements/Framework/TypedElementListLibrary.h
--- a/Engine/Source/Runtime/TypedElementFramework/Private/Elements/Framework/TypedElementListLibrary.h
+++ b/Engine/Source/Runtime/TypedElementFramework/Private/Elements/Framework/TypedElementListLibrary.h
@@ -48,8 +48,8 @@ public:
 	/**
 	 * Get the element interface from the given handle.
 	 */
-	UFUNCTION(BlueprintPure, Category="TypedElementFramework|List", meta=(ScriptMethod))
-	static UTypedElementInterface* GetElementInterface(const FTypedElementListProxy ElementList, const FTypedElementHandle& ElementHandle, const TSubclassOf<UTypedElementInterface>& BaseInterfaceType)
+	UFUNCTION(BlueprintPure, Category="TypedElementFramework|List", meta=(ScriptMethod, DeterminesOutputType = "BaseInterfaceType"))
+	static UObject* GetElementInterface(const FTypedElementListProxy ElementList, const FTypedElementHandle& ElementHandle, const TSubclassOf<UInterface>& BaseInterfaceType)
 	{
 		FTypedElementListConstPtr ElementListPtr = ElementList.GetElementList();
 		return ElementListPtr ? ElementListPtr->GetElementInterface(ElementHandle, BaseInterfaceType) : nullptr;
@@ -59,7 +59,7 @@ public:
 	 * Test whether there are elements in this list, optionally filtering to elements that implement the given interface.
 	 */
 	UFUNCTION(BlueprintCallable, Category="TypedElementFramework|List", meta=(ScriptMethod))
-	static bool HasElements(const FTypedElementListProxy ElementList, const TSubclassOf<UTypedElementInterface>& BaseInterfaceType = nullptr)
+	static bool HasElements(const FTypedElementListProxy ElementList, const TSubclassOf<UInterface>& BaseInterfaceType = nullptr)
 	{
 		FTypedElementListConstPtr ElementListPtr = ElementList.GetElementList();
 		return ElementListPtr && ElementListPtr->HasElements(BaseInterfaceType);
@@ -69,7 +69,7 @@ public:
 	 * Count the number of elements in this list, optionally filtering to elements that implement the given interface.
 	 */
 	UFUNCTION(BlueprintCallable, Category="TypedElementFramework|List", meta=(ScriptMethod))
-	static int32 CountElements(const FTypedElementListProxy ElementList, const TSubclassOf<UTypedElementInterface>& BaseInterfaceType = nullptr)
+	static int32 CountElements(const FTypedElementListProxy ElementList, const TSubclassOf<UInterface>& BaseInterfaceType = nullptr)
 	{
 		FTypedElementListConstPtr ElementListPtr = ElementList.GetElementList();
 		return ElementListPtr ? ElementListPtr->CountElements(BaseInterfaceType) : 0;
@@ -99,7 +99,7 @@ public:
 	 * Get the handle of every element in this list, optionally filtering to elements that implement the given interface.
 	 */
 	UFUNCTION(BlueprintCallable, Category="TypedElementFramework|List", meta=(ScriptMethod))
-	static TArray<FTypedElementHandle> GetElementHandles(const FTypedElementListProxy ElementList, const TSubclassOf<UTypedElementInterface>& BaseInterfaceType = nullptr)
+	static TArray<FTypedElementHandle> GetElementHandles(const FTypedElementListProxy ElementList, const TSubclassOf<UInterface>& BaseInterfaceType = nullptr)
 	{
 		FTypedElementListConstPtr ElementListPtr = ElementList.GetElementList();
 		return ElementListPtr ? ElementListPtr->GetElementHandles(BaseInterfaceType) : TArray<FTypedElementHandle>();
diff --git a/Engine/Source/Runtime/TypedElementFramework/Private/Elements/Framework/TypedElementRegistry.cpp b/Engine/Source/Runtime/TypedElementFramework/Private/Elements/Framework/TypedElementRegistry.cpp
--- a/Engine/Source/Runtime/TypedElementFramework/Private/Elements/Framework/TypedElementRegistry.cpp
+++ b/Engine/Source/Runtime/TypedElementFramework/Private/Elements/Framework/TypedElementRegistry.cpp
@@ -85,7 +85,7 @@ void UTypedElementRegistry::AddReferencedObjects(UObject* InThis, FReferenceColl
 	{
 		if (RegisteredElementType)
 		{
-			for (auto& InterfacesPair : RegisteredElementType->Interfaces)
+			for (TPair<FName, UObject*>& InterfacesPair : RegisteredElementType->Interfaces)
 			{
 				Collector.AddReferencedObject(InterfacesPair.Value);
 			}
@@ -116,9 +116,9 @@ void UTypedElementRegistry::RegisterElementTypeImpl(const FName InElementTypeNam
 	AddRegisteredElementType(MoveTemp(InRegisteredElementType));
 }
 
-void UTypedElementRegistry::RegisterElementInterfaceImpl(const FName InElementTypeName, UTypedElementInterface* InElementInterface, const TSubclassOf<UTypedElementInterface>& InBaseInterfaceType, const bool InAllowOverride)
+void UTypedElementRegistry::RegisterElementInterfaceImpl(const FName InElementTypeName, UObject* InElementInterface, const TSubclassOf<UInterface>& InBaseInterfaceType, const bool InAllowOverride)
 {
-	checkf(InElementInterface->IsA(InBaseInterfaceType), TEXT("Interface '%s' of type '%s' does not derive from '%s'!"), *InElementInterface->GetPathName(), *InElementInterface->GetClass()->GetName(), *InBaseInterfaceType->GetName());
+	checkf(InElementInterface->GetClass()->ImplementsInterface(InBaseInterfaceType), TEXT("Interface '%s' of type '%s' does not derive from '%s'!"), *InElementInterface->GetPathName(), *InElementInterface->GetClass()->GetName(), *InBaseInterfaceType->GetName());
 
 	FRegisteredElementType* RegisteredElementType = GetRegisteredElementTypeFromName(InElementTypeName);
 	checkf(RegisteredElementType, TEXT("Element type '%s' has not been registered!"), *InElementTypeName.ToString());
@@ -127,7 +127,7 @@ void UTypedElementRegistry::RegisterElementInterfaceImpl(const FName InElementTy
 	RegisteredElementType->Interfaces.Add(InBaseInterfaceType->GetFName(), InElementInterface);
 }
 
-UTypedElementInterface* UTypedElementRegistry::GetElementInterfaceImpl(const FTypedHandleTypeId InElementTypeId, const TSubclassOf<UTypedElementInterface>& InBaseInterfaceType) const
+UObject* UTypedElementRegistry::GetElementInterfaceImpl(const FTypedHandleTypeId InElementTypeId, const TSubclassOf<UInterface>& InBaseInterfaceType) const
 {
 	if (!InElementTypeId)
 	{
diff --git a/Engine/Source/Runtime/TypedElementFramework/Private/TypedElementFrameworkTests.cpp b/Engine/Source/Runtime/TypedElementFramework/Private/TypedElementFrameworkTests.cpp
--- a/Engine/Source/Runtime/TypedElementFramework/Private/TypedElementFrameworkTests.cpp
+++ b/Engine/Source/Runtime/TypedElementFramework/Private/TypedElementFrameworkTests.cpp
@@ -38,6 +38,32 @@ bool UTestTypedElementInterfaceA_ImplUntyped::SetDisplayName(const FTypedElement
 	return false;
 }
 
+bool UTestTypedElementInterfaceBAndC_Typed::MarkAsTested(const FTypedElementHandle& InElementHandle)
+{
+	const FTestTypedElementData* ElementData = InElementHandle.GetData<FTestTypedElementData>();
+	if (ElementData)
+	{
+		return false;
+	}
+
+	// this shouldn't happen during the test
+	check(false);
+	return false;
+}
+
+bool UTestTypedElementInterfaceBAndC_Typed::GetIsTested(const FTypedElementHandle& InElementHandle) const
+{
+	const FTestTypedElementData* ElementData = InElementHandle.GetData<FTestTypedElementData>();
+	if (ElementData)
+	{
+		return false;
+	}
+
+	// this shouldn't happen during the test
+	check(false);
+	return false;
+}
+
 #if WITH_DEV_AUTOMATION_TESTS
 
 IMPLEMENT_SIMPLE_AUTOMATION_TEST(FTypedElementRegistrySmokeTest, "System.Runtime.TypedElementRegistry.SmokeTest", EAutomationTestFlags::EditorContext | EAutomationTestFlags::SmokeFilter)
@@ -45,7 +71,7 @@ bool FTypedElementRegistrySmokeTest::RunTest(const FString& Parameters)
 {
 	UTypedElementRegistry* Registry = NewObject<UTypedElementRegistry>();
 	
-	auto TestInterfaceAHandle = [](const TTypedElement<UTestTypedElementInterfaceA>& InElementDataInterfaceHandle)
+	auto TestInterfaceAHandle = [](const TTypedElement<ITestTypedElementInterfaceA>& InElementDataInterfaceHandle)
 	{
 		// Proxy API added via specialization
 		InElementDataInterfaceHandle.SetDisplayName(FText());
@@ -56,44 +82,77 @@ bool FTypedElementRegistrySmokeTest::RunTest(const FString& Parameters)
 		InElementDataInterfaceHandle.GetInterfaceChecked().GetDisplayName(InElementDataInterfaceHandle);
 	};
 
-	auto TestInterfaceAccess = [&Registry, &TestInterfaceAHandle](const FTypedElementHandle& InElementHandle)
+	auto TestInterfaceBHandle = [](TTypedElement<ITestTypedElementInterfaceB>& InElementDataInterfaceHandle)
+	{
+		// Proxy API added via specialization
+		InElementDataInterfaceHandle.MarkAsTested();
+
+
+		// Verbose API
+		InElementDataInterfaceHandle.GetInterfaceChecked().MarkAsTested(InElementDataInterfaceHandle);
+	};
+
+	auto TestInterfaceCHandle = [](const TTypedElement<ITestTypedElementInterfaceC>& InElementDataInterfaceHandle)
+	{
+		// Proxy API added via specialization
+		InElementDataInterfaceHandle.GetIsTested();
+
+		// Verbose API
+		InElementDataInterfaceHandle.GetInterfaceChecked().GetIsTested(InElementDataInterfaceHandle);
+	};
+
+	auto TestInterfaceAccess = [&Registry, &TestInterfaceAHandle, &TestInterfaceBHandle, &TestInterfaceCHandle](const FTypedElementHandle& InElementHandle)
 	{
 		// Get the interface and the element handle in two calls - this is how scripting might work
-		if (UTestTypedElementInterfaceA* Interface = Registry->GetElementInterface<UTestTypedElementInterfaceA>(InElementHandle))
+		if (ITestTypedElementInterfaceA* Interface = Registry->GetElementInterface<ITestTypedElementInterfaceA>(InElementHandle))
 		{
 			Interface->SetDisplayName(InElementHandle, FText());
 			Interface->GetDisplayName(InElementHandle);
 		}
 
 		// Get the interface and the element handle in a single call - this is how C++ might work
-		if (TTypedElement<UTestTypedElementInterfaceA> Element = Registry->GetElement<UTestTypedElementInterfaceA>(InElementHandle))
+		if (TTypedElement<ITestTypedElementInterfaceA> Element = Registry->GetElement<ITestTypedElementInterfaceA>(InElementHandle))
 		{
 			TestInterfaceAHandle(Element);
 		}
-	};
 
-	// Test that all the cast variants compile
-	{
-		FTypedElement DummyElement;
-		CastTypedElement<UTestTypedElementInterfaceA>(DummyElement);
-		CastTypedElement<UTestTypedElementInterfaceA>(FTypedElement());
-		CastTypedElementChecked<UTestTypedElementInterfaceA>(DummyElement);
-		CastTypedElementChecked<UTestTypedElementInterfaceA>(FTypedElement());
-
-		TTypedElement<UTestTypedElementInterfaceA> DummyCastedElement;
-		CastTypedElement<UTestTypedElementInterfaceA>(DummyElement, DummyCastedElement);
-		CastTypedElement<UTestTypedElementInterfaceA>(FTypedElement(), DummyCastedElement);
-		CastTypedElementChecked<UTestTypedElementInterfaceA>(DummyElement, DummyCastedElement);
-		CastTypedElementChecked<UTestTypedElementInterfaceA>(FTypedElement(), DummyCastedElement);
-	}
+		// Get the interface and the element handle in two calls - this is how scripting might work
+		if (ITestTypedElementInterfaceB* Interface = Registry->GetElementInterface<ITestTypedElementInterfaceB>(InElementHandle))
+		{
+			Interface->MarkAsTested(InElementHandle);
+		}
+
+		// Get the interface and the element handle in a single call - this is how C++ might work
+		if (TTypedElement<ITestTypedElementInterfaceB> Element = Registry->GetElement<ITestTypedElementInterfaceB>(InElementHandle))
+		{
+			TestInterfaceBHandle(Element);
+		}
+
+		// Get the interface and the element handle in two calls - this is how scripting might work
+		if (ITestTypedElementInterfaceC* Interface = Registry->GetElementInterface<ITestTypedElementInterfaceC>(InElementHandle))
+		{
+			Interface->GetIsTested(InElementHandle);
+		}
+
+		// Get the interface and the element handle in a single call - this is how C++ might work
+		if (TTypedElement<ITestTypedElementInterfaceC> Element = Registry->GetElement<ITestTypedElementInterfaceC>(InElementHandle))
+		{
+			TestInterfaceCHandle(Element);
+		}
+	};
 
 	const FName DummyElementType_Typed = "DummyElementType_Typed";
 	Registry->RegisterElementType<FTestTypedElementData>(DummyElementType_Typed);
-	Registry->RegisterElementInterface<UTestTypedElementInterfaceA>(DummyElementType_Typed, NewObject<UTestTypedElementInterfaceA_ImplTyped>());
+	Registry->RegisterElementInterface<ITestTypedElementInterfaceA>(DummyElementType_Typed, NewObject<UTestTypedElementInterfaceA_ImplTyped>());
+
+	UObject* TypedElementInterfaceBAndC =  NewObject<UTestTypedElementInterfaceBAndC_Typed>();
+	Registry->RegisterElementInterface<ITestTypedElementInterfaceB>(DummyElementType_Typed, TypedElementInterfaceBAndC);
+	Registry->RegisterElementInterface<ITestTypedElementInterfaceC>(DummyElementType_Typed, TypedElementInterfaceBAndC);
+
 
 	const FName DummyElementType_Untyped = "DummyElementType_Untyped";
 	Registry->RegisterElementType(DummyElementType_Untyped);
-	Registry->RegisterElementInterface<UTestTypedElementInterfaceA>(DummyElementType_Untyped, NewObject<UTestTypedElementInterfaceA_ImplUntyped>());
+	Registry->RegisterElementInterface<ITestTypedElementInterfaceA>(DummyElementType_Untyped, NewObject<UTestTypedElementInterfaceA_ImplUntyped>());
 
 	TTypedElementOwner<FTestTypedElementData> TypedElement1 = Registry->CreateElement<FTestTypedElementData>(DummyElementType_Typed);
 	TypedElement1.GetDataChecked().InternalElementId = "TypedElement1";
@@ -123,7 +182,7 @@ bool FTypedElementRegistrySmokeTest::RunTest(const FString& Parameters)
 		return true;
 	});
 	
-	ElementList->ForEachElement<UTestTypedElementInterfaceA>([&TestInterfaceAHandle](const TTypedElement<UTestTypedElementInterfaceA>& InElementHandle)
+	ElementList->ForEachElement<ITestTypedElementInterfaceA>([&TestInterfaceAHandle](const TTypedElement<ITestTypedElementInterfaceA>& InElementHandle)
 	{
 		TestInterfaceAHandle(InElementHandle);
 		return true;
@@ -154,11 +213,11 @@ bool FTypedElementRegistryPerfTest::RunTest(const FString& Parameters)
 
 	const FName DummyElementType_Typed = "DummyElementType_Typed";
 	Registry->RegisterElementType<FTestTypedElementData>(DummyElementType_Typed);
-	Registry->RegisterElementInterface<UTestTypedElementInterfaceA>(DummyElementType_Typed, NewObject<UTestTypedElementInterfaceA_ImplTyped>());
+	Registry->RegisterElementInterface<ITestTypedElementInterfaceA>(DummyElementType_Typed, NewObject<UTestTypedElementInterfaceA_ImplTyped>());
 
 	const FName DummyElementType_Untyped = "DummyElementType_Untyped";
 	Registry->RegisterElementType(DummyElementType_Untyped);
-	Registry->RegisterElementInterface<UTestTypedElementInterfaceA>(DummyElementType_Untyped, NewObject<UTestTypedElementInterfaceA_ImplUntyped>());
+	Registry->RegisterElementInterface<ITestTypedElementInterfaceA>(DummyElementType_Untyped, NewObject<UTestTypedElementInterfaceA_ImplUntyped>());
 
 	TArray<TTypedElementOwner<FTestTypedElementData>> TypedOwnerHandles;
 	TArray<FTypedElementOwner> UntypedOwnerHandles;
@@ -201,7 +260,7 @@ bool FTypedElementRegistryPerfTest::RunTest(const FString& Parameters)
 		FScopedDurationTimeLogger Timer(FString::Printf(TEXT("Finding %d interfaces from list"), ElementList->Num()));
 		ElementList->ForEachElementHandle([Registry](const FTypedElementHandle& InElementHandle)
 		{
-			Registry->GetElementInterface<UTestTypedElementInterfaceA>(InElementHandle);
+			Registry->GetElementInterface<ITestTypedElementInterfaceA>(InElementHandle);
 			return true;
 		});
 	}
@@ -211,7 +270,7 @@ bool FTypedElementRegistryPerfTest::RunTest(const FString& Parameters)
 		FScopedDurationTimeLogger Timer(FString::Printf(TEXT("Finding %d elements from list"), ElementList->Num()));
 		ElementList->ForEachElementHandle([Registry](const FTypedElementHandle& InElementHandle)
 		{
-			Registry->GetElement<UTestTypedElementInterfaceA>(InElementHandle);
+			Registry->GetElement<ITestTypedElementInterfaceA>(InElementHandle);
 			return true;
 		});
 	}
@@ -219,7 +278,7 @@ bool FTypedElementRegistryPerfTest::RunTest(const FString& Parameters)
 	// Enumerate all elements that implement an interface
 	{
 		FScopedDurationTimeLogger Timer(FString::Printf(TEXT("Enumerating %d elements in list"), ElementList->Num()));
-		ElementList->ForEachElement<UTestTypedElementInterfaceA>([Registry](const TTypedElement<UTestTypedElementInterfaceA>& InElementHandle)
+		ElementList->ForEachElement<ITestTypedElementInterfaceA>([Registry](const TTypedElement<ITestTypedElementInterfaceA>& InElementHandle)
 		{
 			return true;
 		});
diff --git a/Engine/Source/Runtime/TypedElementFramework/Private/TypedElementFrameworkTests.h b/Engine/Source/Runtime/TypedElementFramework/Private/TypedElementFrameworkTests.h
--- a/Engine/Source/Runtime/TypedElementFramework/Private/TypedElementFrameworkTests.h
+++ b/Engine/Source/Runtime/TypedElementFramework/Private/TypedElementFrameworkTests.h
@@ -3,13 +3,20 @@
 #pragma once
 
 #include "Elements/Framework/TypedElementRegistry.h"
+#include "UObject/Interface.h"
+
 #include "TypedElementFrameworkTests.generated.h"
 
+UINTERFACE(MinimalAPI, BlueprintType, meta = (CannotImplementInterfaceInBlueprint))
+class UTestTypedElementInterfaceA : public UInterface
+{
+	GENERATED_BODY()
+};
+
 /**
  * Test interfaces
  */
-UCLASS()
-class UTestTypedElementInterfaceA : public UTypedElementInterface
+class ITestTypedElementInterfaceA
 {
 	GENERATED_BODY()
 
@@ -21,13 +28,55 @@ public:
 	virtual bool SetDisplayName(const FTypedElementHandle& InElementHandle, FText InNewName, bool bNotify = true) { return false; }
 };
 
+UINTERFACE(MinimalAPI, BlueprintType, meta = (CannotImplementInterfaceInBlueprint))
+class UTestTypedElementInterfaceB : public UInterface
+{
+	GENERATED_BODY()
+};
+
+class ITestTypedElementInterfaceB
+{
+	GENERATED_BODY()
+
+public:
+	UFUNCTION(BlueprintCallable, Category = "TypedElementFramework|Testing")
+	virtual bool MarkAsTested(const FTypedElementHandle& InElementHandle) { return false; }
+};
+
+UINTERFACE(MinimalAPI, BlueprintType, meta = (CannotImplementInterfaceInBlueprint))
+class UTestTypedElementInterfaceC : public UInterface
+{
+	GENERATED_BODY()
+};
+
+class ITestTypedElementInterfaceC
+{
+	GENERATED_BODY()
+
+public:
+	UFUNCTION(BlueprintCallable, Category = "TypedElementFramework|Testing")
+	virtual bool GetIsTested(const FTypedElementHandle& InElementHandle) const { return false; }
+};
+
 template <>
-struct TTypedElement<UTestTypedElementInterfaceA> : public TTypedElementBase<UTestTypedElementInterfaceA>
+struct TTypedElement<ITestTypedElementInterfaceA> : public TTypedElementBase<ITestTypedElementInterfaceA>
 {
 	FText GetDisplayName() const { return InterfacePtr->GetDisplayName(*this); }
 	bool SetDisplayName(FText InNewName, bool bNotify = true) const { return InterfacePtr->SetDisplayName(*this, InNewName, bNotify); }
 };
 
+template <>
+struct TTypedElement<ITestTypedElementInterfaceB> : public TTypedElementBase<ITestTypedElementInterfaceB>
+{
+	bool MarkAsTested() { return InterfacePtr->MarkAsTested(*this); }
+};
+
+template <>
+struct TTypedElement<ITestTypedElementInterfaceC> : public TTypedElementBase<ITestTypedElementInterfaceC>
+{
+	bool GetIsTested() const { return InterfacePtr->GetIsTested(*this); }
+};
+
 /**
  * Test dummy type
  */
@@ -39,7 +88,7 @@ struct FTestTypedElementData
 };
 
 UCLASS()
-class UTestTypedElementInterfaceA_ImplTyped : public UTestTypedElementInterfaceA
+class UTestTypedElementInterfaceA_ImplTyped : public UObject, public ITestTypedElementInterfaceA
 {
 	GENERATED_BODY()
 
@@ -52,7 +101,7 @@ public:
  * Test untyped
  */
 UCLASS()
-class UTestTypedElementInterfaceA_ImplUntyped : public UTestTypedElementInterfaceA
+class UTestTypedElementInterfaceA_ImplUntyped : public UObject, public ITestTypedElementInterfaceA
 {
 	GENERATED_BODY()
 
@@ -60,3 +109,17 @@ public:
 	virtual FText GetDisplayName(const FTypedElementHandle& InElementHandle) override;
 	virtual bool SetDisplayName(const FTypedElementHandle& InElementHandle, FText InNewName, bool bNotify) override;
 };
+
+/**
+ * Test to two implementation in one object
+ */
+UCLASS()
+class UTestTypedElementInterfaceBAndC_Typed : public UObject, public ITestTypedElementInterfaceB, public ITestTypedElementInterfaceC
+{
+	GENERATED_BODY()
+
+public:
+	virtual bool MarkAsTested(const FTypedElementHandle& InElementHandle) override;
+	virtual bool GetIsTested(const FTypedElementHandle& InElementHandle) const override;
+};
+
diff --git a/Engine/Source/Runtime/TypedElementFramework/Public/Elements/Framework/TypedElementCounter.h b/Engine/Source/Runtime/TypedElementFramework/Public/Elements/Framework/TypedElementCounter.h
--- a/Engine/Source/Runtime/TypedElementFramework/Public/Elements/Framework/TypedElementCounter.h
+++ b/Engine/Source/Runtime/TypedElementFramework/Public/Elements/Framework/TypedElementCounter.h
@@ -3,9 +3,11 @@
 #pragma once
 
 #include <type_traits>
-#include "CoreMinimal.h"
 #include "Containers/SortedMap.h"
+#include "CoreMinimal.h"
 #include "Elements/Framework/TypedElementHandle.h"
+#include "UObject/Interface.h"
+
 #include "TypedElementCounter.generated.h"
 
 class UTypedElementRegistry;
@@ -272,8 +274,13 @@ private:
 	TSortedMap<FName, TUniquePtr<ICounterCategory>, FDefaultAllocator, FNameFastLess> CounterCategories;
 };
 
-UCLASS(Abstract)
-class TYPEDELEMENTFRAMEWORK_API UTypedElementCounterInterface : public UTypedElementInterface
+UINTERFACE(MinimalAPI, meta = (CannotImplementInterfaceInBlueprint))
+class UTypedElementCounterInterface : public UInterface
+{
+	GENERATED_BODY()
+};
+
+class TYPEDELEMENTFRAMEWORK_API ITypedElementCounterInterface
 {
 	GENERATED_BODY()
 
@@ -290,7 +297,7 @@ public:
 };
 
 template <>
-struct TTypedElement<UTypedElementCounterInterface> : public TTypedElementBase<UTypedElementCounterInterface>
+struct TTypedElement<ITypedElementCounterInterface> : public TTypedElementBase<ITypedElementCounterInterface>
 {
 	void IncrementCountersForElement(FTypedElementCounter& InOutCounter) const { InterfacePtr->IncrementCountersForElement(*this, InOutCounter); }
 	void DecrementCountersForElement(FTypedElementCounter& InOutCounter) const { InterfacePtr->DecrementCountersForElement(*this, InOutCounter); }
diff --git a/Engine/Source/Runtime/TypedElementFramework/Public/Elements/Framework/TypedElementHandle.h b/Engine/Source/Runtime/TypedElementFramework/Public/Elements/Framework/TypedElementHandle.h
--- a/Engine/Source/Runtime/TypedElementFramework/Public/Elements/Framework/TypedElementHandle.h
+++ b/Engine/Source/Runtime/TypedElementFramework/Public/Elements/Framework/TypedElementHandle.h
@@ -5,7 +5,7 @@
 #include <type_traits>
 #include "CoreMinimal.h"
 #include "Elements/Framework/TypedElementData.h"
-#include "Elements/Interfaces/TypedElementInterface.h"
+
 #include "TypedElementHandle.generated.h"
 
 /**
@@ -326,24 +326,6 @@ public:
 		Private_DestroyReleaseRef();
 	}
 
-	/**
-	 * Test to see whether the interface stored within this element is of the given type.
-	 */
-	template <typename InterfaceType>
-	FORCEINLINE bool HasInterface() const
-	{
-		return HasInterface(InterfaceType::StaticClass());
-	}
-
-	/**
-	 * Test to see whether the interface stored within this element is of the given type.
-	 */
-	FORCEINLINE bool HasInterface(const UClass* InInterfaceType) const
-	{
-		return InterfacePtr
-			&& InterfacePtr->IsA(InInterfaceType);
-	}
-
 	/**
 	 * Attempt to access the interface stored within this element, returning null if it isn't set.
 	 */
@@ -355,8 +337,10 @@ public:
 	/**
 	 * Attempt to access the interface stored within this element, asserting if it isn't set.
 	 */
-	FORCEINLINE BaseInterfaceType& GetInterfaceChecked() const
+	template <typename U = BaseInterfaceType, std::enable_if_t<!std::is_void<U>::value>* = nullptr>
+	FORCEINLINE U& GetInterfaceChecked() const
 	{
+		static_assert(std::is_same<U, BaseInterfaceType>::value, "Don't explicitly specify U!");
 		checkf(InterfacePtr, TEXT("Interface is null!"));
 		return *InterfacePtr;
 	}
@@ -414,111 +398,9 @@ template <typename BaseInterfaceType>
 struct TTypedElement : public TTypedElementBase<BaseInterfaceType>
 {
 };
-using FTypedElement = TTypedElement<UTypedElementInterface>;
-
-template <typename OtherInterfaceType, typename ThisInterfaceType>
-FORCEINLINE TTypedElement<OtherInterfaceType> CastTypedElement(const TTypedElement<ThisInterfaceType>& InElement)
-{
-	static_assert(sizeof(TTypedElement<ThisInterfaceType>) == sizeof(TTypedElement<OtherInterfaceType>), "All TTypedElement instances must be the same size for this cast implementation to work!");
-	if (InElement && InElement.template HasInterface<OtherInterfaceType>())
-	{
-		return reinterpret_cast<const TTypedElement<OtherInterfaceType>&>(InElement);
-	}
-	return TTypedElement<OtherInterfaceType>();
-}
-
-template <typename OtherInterfaceType, typename ThisInterfaceType>
-FORCEINLINE void CastTypedElement(const TTypedElement<ThisInterfaceType>& InElement, TTypedElement<OtherInterfaceType>& OutCastElement)
-{
-	static_assert(sizeof(TTypedElement<ThisInterfaceType>) == sizeof(TTypedElement<OtherInterfaceType>), "All TTypedElement instances must be the same size for this cast implementation to work!");
-	if (InElement && InElement.template HasInterface<OtherInterfaceType>())
-	{
-		OutCastElement = reinterpret_cast<const TTypedElement<OtherInterfaceType>&>(InElement);
-	}
-	else
-	{
-		OutCastElement.Private_DestroyReleaseRef();
-	}
-}
-
-template <typename OtherInterfaceType, typename ThisInterfaceType>
-FORCEINLINE TTypedElement<OtherInterfaceType> CastTypedElement(TTypedElement<ThisInterfaceType>&& InElement)
-{
-	static_assert(sizeof(TTypedElement<ThisInterfaceType>) == sizeof(TTypedElement<OtherInterfaceType>), "All TTypedElement instances must be the same size for this cast implementation to work!");
-	if (InElement && InElement.template HasInterface<OtherInterfaceType>())
-	{
-		return reinterpret_cast<TTypedElement<OtherInterfaceType>&&>(InElement);
-	}
-	return TTypedElement<OtherInterfaceType>();
-}
 
-template <typename OtherInterfaceType, typename ThisInterfaceType>
-FORCEINLINE void CastTypedElement(TTypedElement<ThisInterfaceType>&& InElement, TTypedElement<OtherInterfaceType>& OutCastElement)
-{
-	static_assert(sizeof(TTypedElement<ThisInterfaceType>) == sizeof(TTypedElement<OtherInterfaceType>), "All TTypedElement instances must be the same size for this cast implementation to work!");
-	if (InElement && InElement.template HasInterface<OtherInterfaceType>())
-	{
-		OutCastElement = reinterpret_cast<TTypedElement<OtherInterfaceType>&&>(InElement);
-	}
-	else
-	{
-		OutCastElement.Private_DestroyReleaseRef();
-	}
-}
-
-template <typename OtherInterfaceType, typename ThisInterfaceType>
-FORCEINLINE TTypedElement<OtherInterfaceType> CastTypedElementChecked(const TTypedElement<ThisInterfaceType>& InElement)
-{
-	static_assert(sizeof(TTypedElement<ThisInterfaceType>) == sizeof(TTypedElement<OtherInterfaceType>), "All TTypedElement instances must be the same size for this cast implementation to work!");
-	if (InElement)
-	{
-		checkf(InElement.template HasInterface<OtherInterfaceType>(), TEXT("Element does not implement the required interface for this cast!"));
-		return reinterpret_cast<const TTypedElement<OtherInterfaceType>&>(InElement);
-	}
-	return TTypedElement<OtherInterfaceType>();
-}
-
-template <typename OtherInterfaceType, typename ThisInterfaceType>
-FORCEINLINE void CastTypedElementChecked(const TTypedElement<ThisInterfaceType>& InElement, TTypedElement<OtherInterfaceType>& OutCastElement)
-{
-	static_assert(sizeof(TTypedElement<ThisInterfaceType>) == sizeof(TTypedElement<OtherInterfaceType>), "All TTypedElement instances must be the same size for this cast implementation to work!");
-	if (InElement)
-	{
-		checkf(InElement.template HasInterface<OtherInterfaceType>(), TEXT("Element does not implement the required interface for this cast!"));
-		OutCastElement = reinterpret_cast<const TTypedElement<OtherInterfaceType>&>(InElement);
-	}
-	else
-	{
-		OutCastElement.Private_DestroyReleaseRef();
-	}
-}
-
-template <typename OtherInterfaceType, typename ThisInterfaceType>
-FORCEINLINE TTypedElement<OtherInterfaceType> CastTypedElementChecked(TTypedElement<ThisInterfaceType>&& InElement)
-{
-	static_assert(sizeof(TTypedElement<ThisInterfaceType>) == sizeof(TTypedElement<OtherInterfaceType>), "All TTypedElement instances must be the same size for this cast implementation to work!");
-	if (InElement)
-	{
-		checkf(InElement.template HasInterface<OtherInterfaceType>(), TEXT("Element does not implement the required interface for this cast!"));
-		return reinterpret_cast<TTypedElement<OtherInterfaceType>&&>(InElement);
-	}
-	return TTypedElement<OtherInterfaceType>();
-}
-
-template <typename OtherInterfaceType, typename ThisInterfaceType>
-FORCEINLINE void CastTypedElementChecked(TTypedElement<ThisInterfaceType>&& InElement, TTypedElement<OtherInterfaceType>& OutCastElement)
-{
-	static_assert(sizeof(TTypedElement<ThisInterfaceType>) == sizeof(TTypedElement<OtherInterfaceType>), "All TTypedElement instances must be the same size for this cast implementation to work!");
-	if (InElement)
-	{
-		checkf(InElement.template HasInterface<OtherInterfaceType>(), TEXT("Element does not implement the required interface for this cast!"));
-		OutCastElement = reinterpret_cast<TTypedElement<OtherInterfaceType>&&>(InElement);
-	}
-	else
-	{
-		OutCastElement.Private_DestroyReleaseRef();
-	}
-}
+// The typed element InterfacePtr will points toward the vtable of the interface
+using FTypedElement = TTypedElement<void>;
 
 /**
  * A representation of the owner of an element that includes its mutable handle data.
diff --git a/Engine/Source/Runtime/TypedElementFramework/Public/Elements/Framework/TypedElementList.h b/Engine/Source/Runtime/TypedElementFramework/Public/Elements/Framework/TypedElementList.h
--- a/Engine/Source/Runtime/TypedElementFramework/Public/Elements/Framework/TypedElementList.h
+++ b/Engine/Source/Runtime/TypedElementFramework/Public/Elements/Framework/TypedElementList.h
@@ -17,13 +17,13 @@ class UTypedElementRegistry;
 namespace TypedElementList_Private
 {
 
-TYPEDELEMENTFRAMEWORK_API void GetElementImpl(const UTypedElementRegistry* InRegistry, const FTypedElementHandle& InElementHandle, const UClass* InBaseInterfaceType, FTypedElement& OutElement);
+TYPEDELEMENTFRAMEWORK_API void GetElementImpl(const UTypedElementRegistry* InRegistry, const FTypedElementHandle& InElementHandle, const TSubclassOf<UInterface>& InBaseInterfaceType, FTypedElement& OutElement);
 
 template <typename BaseInterfaceType>
 FORCEINLINE void GetElement(const UTypedElementRegistry* InRegistry, const FTypedElementHandle& InElementHandle, TTypedElement<BaseInterfaceType>& OutElement)
 {
 	static_assert(sizeof(TTypedElement<BaseInterfaceType>) == sizeof(FTypedElement), "All TTypedElement instances must be the same size for this cast implementation to work!");
-	GetElementImpl(InRegistry, InElementHandle, BaseInterfaceType::StaticClass(), reinterpret_cast<FTypedElement&>(OutElement));
+	GetElementImpl(InRegistry, InElementHandle, BaseInterfaceType::UClassType::StaticClass(), reinterpret_cast<FTypedElement&>(OutElement));
 }
 
 template <typename BaseInterfaceType>
@@ -236,29 +236,30 @@ public:
 		return TempElement;
 	}
 
+
 	/**
 	 * Get the element interface from the given handle.
 	 */
 	template <typename BaseInterfaceType>
 	BaseInterfaceType* GetElementInterface(const FTypedElementHandle& InElementHandle) const
 	{
-		return static_cast<BaseInterfaceType*>(GetElementInterface(InElementHandle, BaseInterfaceType::StaticClass()));
+		return Cast<BaseInterfaceType>(GetElementInterface(InElementHandle, BaseInterfaceType::UClassType::StaticClass()));
 	}
 
 	/**
 	 * Get the element interface from the given handle.
 	 */
-	UTypedElementInterface* GetElementInterface(const FTypedElementHandle& InElementHandle, const TSubclassOf<UTypedElementInterface>& InBaseInterfaceType) const;
+	UObject* GetElementInterface(const FTypedElementHandle& InElementHandle, const TSubclassOf<UInterface>& InBaseInterfaceType) const;
 
 	/**
 	 * Test whether there are elements in this list, optionally filtering to elements that implement the given interface.
 	 */
-	bool HasElements(const TSubclassOf<UTypedElementInterface>& InBaseInterfaceType = nullptr) const;
+	bool HasElements(const TSubclassOf<UInterface>& InBaseInterfaceType = nullptr) const;
 
 	/**
 	 * Count the number of elements in this list, optionally filtering to elements that implement the given interface.
 	 */
-	int32 CountElements(const TSubclassOf<UTypedElementInterface>& InBaseInterfaceType = nullptr) const;
+	int32 CountElements(const TSubclassOf<UInterface>& InBaseInterfaceType = nullptr) const;
 
 	/**
 	 * Test whether there are elements in this list of the given type.
@@ -275,13 +276,13 @@ public:
 	/**
 	 * Get the handle of every element in this list, optionally filtering to elements that implement the given interface.
 	 */
-	TArray<FTypedElementHandle> GetElementHandles(const TSubclassOf<UTypedElementInterface>& InBaseInterfaceType = nullptr) const;
+	TArray<FTypedElementHandle> GetElementHandles(const TSubclassOf<UInterface>& InBaseInterfaceType = nullptr) const;
 
 	/**
 	 * Get the handle of every element in this list, optionally filtering to elements that implement the given interface.
 	 */
 	template <typename ArrayAllocator>
-	void GetElementHandles(TArray<FTypedElementHandle, ArrayAllocator>& OutArray, const TSubclassOf<UTypedElementInterface>& InBaseInterfaceType = nullptr) const
+	void GetElementHandles(TArray<FTypedElementHandle, ArrayAllocator>& OutArray, const TSubclassOf<UInterface>& InBaseInterfaceType = nullptr) const
 	{
 		OutArray.Reset();
 		OutArray.Reserve(ElementHandles.Num());
@@ -296,7 +297,7 @@ public:
 	 * Enumerate the handle of every element in this list, optionally filtering to elements that implement the given interface.
 	 * @note Return true from the callback to continue enumeration.
 	 */
-	void ForEachElementHandle(TFunctionRef<bool(const FTypedElementHandle&)> InCallback, const TSubclassOf<UTypedElementInterface>& InBaseInterfaceType = nullptr) const;
+	void ForEachElementHandle(TFunctionRef<bool(const FTypedElementHandle&)> InCallback, const TSubclassOf<UInterface>& InBaseInterfaceType = nullptr) const;
 
 	/**
 	 * Enumerate the elements in this list that implement the given interface.
diff --git a/Engine/Source/Runtime/TypedElementFramework/Public/Elements/Framework/TypedElementRegistry.h b/Engine/Source/Runtime/TypedElementFramework/Public/Elements/Framework/TypedElementRegistry.h
--- a/Engine/Source/Runtime/TypedElementFramework/Public/Elements/Framework/TypedElementRegistry.h
+++ b/Engine/Source/Runtime/TypedElementFramework/Public/Elements/Framework/TypedElementRegistry.h
@@ -2,12 +2,16 @@
 
 #pragma once
 
-#include "CoreMinimal.h"
-#include "Misc/ScopeRWLock.h"
+#include <type_traits>
+
 #include "Containers/ArrayView.h"
 #include "Containers/SortedMap.h"
+#include "CoreMinimal.h"
 #include "Elements/Framework/TypedElementHandle.h"
 #include "Elements/Framework/TypedElementList.h"
+#include "Misc/ScopeRWLock.h"
+#include "UObject/ScriptInterface.h"
+
 #include "TypedElementRegistry.generated.h"
 
 /**
@@ -110,9 +114,10 @@ public:
 	 * Register that an element interface is supported for the given type, which must have previously been registered via RegisterElementType.
 	 */
 	template <typename BaseInterfaceType>
-	FORCEINLINE void RegisterElementInterface(const FName InElementTypeName, UTypedElementInterface* InElementInterface, const bool InAllowOverride = false)
+	FORCEINLINE void RegisterElementInterface(const FName InElementTypeName, UObject* InElementInterface, const bool InAllowOverride = false)
 	{
-		RegisterElementInterfaceImpl(InElementTypeName, InElementInterface, BaseInterfaceType::StaticClass(), InAllowOverride);
+		checkf(InElementInterface->GetClass()->ImplementsInterface(BaseInterfaceType::UClassType::StaticClass()), TEXT("The InElementInterface pass must implement the interface to register."));
+		RegisterElementInterfaceImpl(InElementTypeName, InElementInterface, BaseInterfaceType::UClassType::StaticClass(), InAllowOverride);
 	}
 
 	/**
@@ -121,7 +126,7 @@ public:
 	template <typename BaseInterfaceType>
 	FORCEINLINE BaseInterfaceType* GetElementInterface(const FTypedHandleTypeId InElementTypeId) const
 	{
-		return static_cast<BaseInterfaceType*>(GetElementInterfaceImpl(InElementTypeId, BaseInterfaceType::StaticClass()));
+		return Cast<BaseInterfaceType>(GetElementInterfaceImpl(InElementTypeId, BaseInterfaceType::UClassType::StaticClass()));
 	}
 
 	/**
@@ -130,14 +135,14 @@ public:
 	template <typename BaseInterfaceType>
 	FORCEINLINE BaseInterfaceType* GetElementInterface(const FTypedElementHandle& InElementHandle) const
 	{
-		return static_cast<BaseInterfaceType*>(GetElementInterfaceImpl(InElementHandle.GetId().GetTypeId(), BaseInterfaceType::StaticClass()));
+		return Cast<BaseInterfaceType>(GetElementInterfaceImpl(InElementHandle.GetId().GetTypeId(), BaseInterfaceType::UClassType::StaticClass()));
 	}
 
 	/**
 	 * Get the element interface supported by the given handle, or null if there is no support for this interface.
 	 */
-	UFUNCTION(BlueprintPure, Category="TypedElementFramework|Registry")
-	FORCEINLINE UTypedElementInterface* GetElementInterface(const FTypedElementHandle& InElementHandle, const TSubclassOf<UTypedElementInterface>& InBaseInterfaceType) const
+	UFUNCTION(BlueprintPure, Category="TypedElementFramework|Registry", meta = (DeterminesOutputType = "InBaseInterfaceType"))
+	FORCEINLINE UObject* GetElementInterface(const FTypedElementHandle& InElementHandle, const TSubclassOf<UInterface>& InBaseInterfaceType) const
 	{
 		return GetElementInterfaceImpl(InElementHandle.GetId().GetTypeId(), InBaseInterfaceType);
 	}
@@ -200,7 +205,7 @@ public:
 	/**
 	 * Get an element that implements the given interface from its minimal ID.
 	 */
-	FORCEINLINE FTypedElement GetElement(const FTypedElementId& InElementId, const TSubclassOf<UTypedElementInterface>& InBaseInterfaceType) const
+	FORCEINLINE FTypedElement GetElement(const FTypedElementId& InElementId, const TSubclassOf<UInterface>& InBaseInterfaceType) const
 	{
 		FTypedElement Element;
 		GetElementImpl(InElementId, InBaseInterfaceType, Element);
@@ -214,14 +219,14 @@ public:
 	FORCEINLINE TTypedElement<BaseInterfaceType> GetElement(const FTypedElementId& InElementId) const
 	{
 		TTypedElement<BaseInterfaceType> Element;
-		GetElementImpl(InElementId, BaseInterfaceType::StaticClass(), Element);
+		GetElementImpl(InElementId, BaseInterfaceType::UClassType::StaticClass(), Element);
 		return Element;
 	}
 
 	/**
 	 * Get an element that implements the given interface from its handle.
 	 */
-	FORCEINLINE FTypedElement GetElement(const FTypedElementHandle& InElementHandle, const TSubclassOf<UTypedElementInterface>& InBaseInterfaceType) const
+	FORCEINLINE FTypedElement GetElement(const FTypedElementHandle& InElementHandle, const TSubclassOf<UInterface>& InBaseInterfaceType) const
 	{
 		FTypedElement Element;
 		GetElementImpl(InElementHandle, InBaseInterfaceType, Element);
@@ -235,7 +240,7 @@ public:
 	FORCEINLINE TTypedElement<BaseInterfaceType> GetElement(const FTypedElementHandle& InElementHandle) const
 	{
 		TTypedElement<BaseInterfaceType> Element;
-		GetElementImpl(InElementHandle, BaseInterfaceType::StaticClass(), Element);
+		GetElementImpl(InElementHandle, BaseInterfaceType::UClassType::StaticClass(), Element);
 		return Element;
 	}
 
@@ -290,7 +295,7 @@ public:
 	}
 
 	// Note: Access for FTypedElementList
-	FORCEINLINE void Private_GetElementImpl(const FTypedElementHandle& InElementHandle, const UClass* InBaseInterfaceType, FTypedElement& OutElement) const
+	FORCEINLINE void Private_GetElementImpl(const FTypedElementHandle& InElementHandle, const TSubclassOf<UInterface>& InBaseInterfaceType, FTypedElement& OutElement) const
 	{
 		GetElementImpl(InElementHandle, InBaseInterfaceType, OutElement);
 	}
@@ -336,7 +341,8 @@ private:
 
 		FTypedHandleTypeId TypeId = 0;
 		FName TypeName;
-		TSortedMap<FName, UTypedElementInterface*, FDefaultAllocator, FNameFastLess> Interfaces;
+
+		TSortedMap<FName, UObject*, FDefaultAllocator, FNameFastLess> Interfaces;
 	};
 
 	template <typename ElementDataType>
@@ -402,8 +408,8 @@ private:
 	};
 
 	void RegisterElementTypeImpl(const FName InElementTypeName, TUniquePtr<FRegisteredElementType>&& InRegisteredElementType);
-	void RegisterElementInterfaceImpl(const FName InElementTypeName, UTypedElementInterface* InElementInterface, const TSubclassOf<UTypedElementInterface>& InBaseInterfaceType, const bool InAllowOverride);
-	UTypedElementInterface* GetElementInterfaceImpl(const FTypedHandleTypeId InElementTypeId, const TSubclassOf<UTypedElementInterface>& InBaseInterfaceType) const;
+	void RegisterElementInterfaceImpl(const FName InElementTypeName, UObject* InElementInterface, const TSubclassOf<UInterface>& InBaseInterfaceType, const bool InAllowOverride);
+	UObject* GetElementInterfaceImpl(const FTypedHandleTypeId InElementTypeId, const TSubclassOf<UInterface>& InBaseInterfaceType) const;
 
 	template <typename ElementDataType>
 	TTypedElementOwner<ElementDataType> CreateElementImpl(const FName InElementTypeName, const FTypedHandleElementId InElementId)
@@ -433,7 +439,7 @@ private:
 	}
 
 	template <typename BaseInterfaceType>
-	void GetElementImpl(const FTypedElementId& InElementId, const UClass* InBaseInterfaceType, TTypedElement<BaseInterfaceType>& OutElement) const
+	void GetElementImpl(const FTypedElementId& InElementId, const TSubclassOf<UInterface>& InBaseInterfaceType, TTypedElement<BaseInterfaceType>& OutElement) const
 	{
 		OutElement.Private_DestroyReleaseRef();
 
@@ -442,12 +448,21 @@ private:
 			FRegisteredElementType* RegisteredElementType = GetRegisteredElementTypeFromId(InElementId.GetTypeId());
 			checkf(RegisteredElementType, TEXT("Element type ID '%d' has not been registered!"), InElementId.GetTypeId());
 
-			OutElement.Private_InitializeAddRef(RegisteredElementType->GetDataForElement(InElementId.GetElementId()), static_cast<BaseInterfaceType*>(RegisteredElementType->Interfaces.FindRef(InBaseInterfaceType->GetFName())));
+			UObject* InterfaceObject = RegisteredElementType->Interfaces.FindRef(InBaseInterfaceType->GetFName());
+			const FTypedElementInternalData& TypedElementInternalData = RegisteredElementType->GetDataForElement(InElementId.GetElementId());
+			if constexpr (std::is_void<BaseInterfaceType>::value)
+			{ 
+				OutElement.Private_InitializeAddRef(TypedElementInternalData, InterfaceObject->GetInterfaceAddress(InBaseInterfaceType));
+			}
+			else
+			{
+				OutElement.Private_InitializeAddRef(TypedElementInternalData, Cast<BaseInterfaceType>(InterfaceObject));
+			}
 		}
 	}
 
 	template <typename BaseInterfaceType>
-	void GetElementImpl(const FTypedElementHandle& InElementHandle, const UClass* InBaseInterfaceType, TTypedElement<BaseInterfaceType>& OutElement) const
+	void GetElementImpl(const FTypedElementHandle& InElementHandle, const TSubclassOf<UInterface>& InBaseInterfaceType, TTypedElement<BaseInterfaceType>& OutElement) const
 	{
 		OutElement.Private_DestroyReleaseRef();
 
@@ -456,7 +471,16 @@ private:
 			FRegisteredElementType* RegisteredElementType = GetRegisteredElementTypeFromId(InElementHandle.GetId().GetTypeId());
 			checkf(RegisteredElementType, TEXT("Element type ID '%d' has not been registered!"), InElementHandle.GetId().GetTypeId());
 
-			OutElement.Private_InitializeAddRef(*InElementHandle.Private_GetInternalData(), static_cast<BaseInterfaceType*>(RegisteredElementType->Interfaces.FindRef(InBaseInterfaceType->GetFName())));
+			UObject* InterfaceObject = RegisteredElementType->Interfaces.FindRef(InBaseInterfaceType->GetFName());
+			const FTypedElementInternalData& TypedElementInternalData = *InElementHandle.Private_GetInternalData();
+			if constexpr (std::is_void<BaseInterfaceType>::value)
+			{
+				OutElement.Private_InitializeAddRef(TypedElementInternalData, InterfaceObject->GetInterfaceAddress(InBaseInterfaceType));
+			}
+			else
+			{
+				OutElement.Private_InitializeAddRef(TypedElementInternalData, Cast<BaseInterfaceType>(InterfaceObject));
+			}
 		}
 	}
 
diff --git a/Engine/Source/Runtime/TypedElementFramework/Public/Elements/Interfaces/TypedElementInterface.h b/Engine/Source/Runtime/TypedElementFramework/Public/Elements/Interfaces/TypedElementInterface.h
deleted file mode 100644
diff --git a/Engine/Source/Runtime/TypedElementRuntime/Private/Elements/Framework/TypedElementListObjectUtil.cpp b/Engine/Source/Runtime/TypedElementRuntime/Private/Elements/Framework/TypedElementListObjectUtil.cpp
--- a/Engine/Source/Runtime/TypedElementRuntime/Private/Elements/Framework/TypedElementListObjectUtil.cpp
+++ b/Engine/Source/Runtime/TypedElementRuntime/Private/Elements/Framework/TypedElementListObjectUtil.cpp
@@ -6,10 +6,10 @@
 namespace TypedElementListObjectUtil
 {
 
-UObject* GetObjectOfType(const TTypedElement<UTypedElementObjectInterface>& InObjectElement, const UClass* InRequiredClass)
+UObject* GetObjectOfType(const TTypedElement<ITypedElementObjectInterface>& InObjectElement, const UClass* InRequiredClass)
 {
 	UObject* ElementObject = InObjectElement ? InObjectElement.GetObject() : nullptr;
-	return (ElementObject && (!InRequiredClass || ElementObject->IsA(InRequiredClass)))
+	return (ElementObject && (!InRequiredClass || ElementObject->GetClass()->ImplementsInterface(InRequiredClass)))
 		? ElementObject
 		: nullptr;
 }
@@ -42,7 +42,7 @@ int32 CountObjects(FTypedElementListConstRef InElementList, const UClass* InRequ
 
 void ForEachObject(FTypedElementListConstRef InElementList, TFunctionRef<bool(UObject*)> InCallback, const UClass* InRequiredClass)
 {
-	InElementList->ForEachElement<UTypedElementObjectInterface>([&InCallback, InRequiredClass](const TTypedElement<UTypedElementObjectInterface>& InObjectElement)
+	InElementList->ForEachElement<ITypedElementObjectInterface>([&InCallback, InRequiredClass](const TTypedElement<ITypedElementObjectInterface>& InObjectElement)
 	{
 		if (UObject* ElementObject = GetObjectOfType(InObjectElement, InRequiredClass))
 		{
@@ -68,7 +68,7 @@ TArray<UObject*> GetObjects(FTypedElementListConstRef InElementList, const UClas
 
 UObject* GetTopObject(FTypedElementListConstRef InElementList, const UClass* InRequiredClass)
 {
-	TTypedElement<UTypedElementObjectInterface> TempElement;
+	TTypedElement<ITypedElementObjectInterface> TempElement;
 	for (int32 ElementIndex = 0; ElementIndex < InElementList->Num(); ++ElementIndex)
 	{
 		InElementList->GetElementAt(ElementIndex, TempElement);
@@ -84,7 +84,7 @@ UObject* GetTopObject(FTypedElementListConstRef InElementList, const UClass* InR
 
 UObject* GetBottomObject(FTypedElementListConstRef InElementList, const UClass* InRequiredClass)
 {
-	TTypedElement<UTypedElementObjectInterface> TempElement;
+	TTypedElement<ITypedElementObjectInterface> TempElement;
 	for (int32 ElementIndex = InElementList->Num() - 1; ElementIndex >= 0; --ElementIndex)
 	{
 		InElementList->GetElementAt(ElementIndex, TempElement);
diff --git a/Engine/Source/Runtime/TypedElementRuntime/Private/Elements/Framework/TypedElementSelectionSet.cpp b/Engine/Source/Runtime/TypedElementRuntime/Private/Elements/Framework/TypedElementSelectionSet.cpp
--- a/Engine/Source/Runtime/TypedElementRuntime/Private/Elements/Framework/TypedElementSelectionSet.cpp
+++ b/Engine/Source/Runtime/TypedElementRuntime/Private/Elements/Framework/TypedElementSelectionSet.cpp
@@ -8,12 +8,12 @@
 #include "Serialization/ObjectReader.h"
 #include "Serialization/ObjectWriter.h"
 
-void FTypedElementSelectionCustomization::GetNormalizedElements(const TTypedElement<UTypedElementSelectionInterface>& InElementSelectionHandle, FTypedElementListConstRef InSelectionSet, const FTypedElementSelectionNormalizationOptions& InNormalizationOptions, FTypedElementListRef OutNormalizedElements)
+void FTypedElementSelectionCustomization::GetNormalizedElements(const TTypedElement<ITypedElementSelectionInterface>& InElementSelectionHandle, FTypedElementListConstRef InSelectionSet, const FTypedElementSelectionNormalizationOptions& InNormalizationOptions, FTypedElementListRef OutNormalizedElements)
 {
 	// Don't include this element in the normalized selection if it has a parent element that is also selected
 	{
 		FTypedElementHandle ElementToTest = InElementSelectionHandle;
-		while (TTypedElement<UTypedElementHierarchyInterface> ElementHierarchyHandle = InSelectionSet->GetElement<UTypedElementHierarchyInterface>(ElementToTest))
+		while (TTypedElement<ITypedElementHierarchyInterface> ElementHierarchyHandle = InSelectionSet->GetElement<ITypedElementHierarchyInterface>(ElementToTest))
 		{
 			ElementToTest = ElementHierarchyHandle.GetParentElement();
 			if (ElementToTest && InSelectionSet->Contains(ElementToTest))
@@ -68,7 +68,7 @@ bool UTypedElementSelectionSet::Modify(bool bAlwaysMarkDirty)
 	if (GUndo && CanModify())
 	{
 		bool bCanModify = true;
-		ElementList->ForEachElement<UTypedElementSelectionInterface>([&bCanModify](const TTypedElement<UTypedElementSelectionInterface>& InSelectionElement)
+		ElementList->ForEachElement<ITypedElementSelectionInterface>([&bCanModify](const TTypedElement<ITypedElementSelectionInterface>& InSelectionElement)
 		{
 			bCanModify = !InSelectionElement.ShouldPreventTransactions();
 			return bCanModify;
@@ -138,7 +138,7 @@ void UTypedElementSelectionSet::Serialize(FArchive& Ar)
 			FTypedHandleTypeId TypeId = 0;
 			Ar << TypeId;
 
-			UTypedElementSelectionInterface* ElementTypeSelectionInterface = Registry->GetElementInterface<UTypedElementSelectionInterface>(TypeId);
+			ITypedElementSelectionInterface* ElementTypeSelectionInterface = Registry->GetElementInterface<ITypedElementSelectionInterface>(TypeId);
 			checkf(ElementTypeSelectionInterface, TEXT("Failed to find selection interface for a previously transacted element type!"));
 
 			TUniquePtr<ITypedElementTransactedElement> TransactedElement = ElementTypeSelectionInterface->CreateTransactedElement(TypeId);
@@ -187,7 +187,7 @@ bool UTypedElementSelectionSet::SelectElement(const FTypedElementHandle& InEleme
 
 	if (InSelectionOptions.GetChildElementInclusionMethod() != ETypedElementChildInclusionMethod::None)
 	{
-		if (TTypedElement<UTypedElementHierarchyInterface> ElementHierarchyHandle = ElementList->GetElement<UTypedElementHierarchyInterface>(InElementHandle))
+		if (TTypedElement<ITypedElementHierarchyInterface> ElementHierarchyHandle = ElementList->GetElement<ITypedElementHierarchyInterface>(InElementHandle))
 		{
 			const FTypedElementSelectionOptions ChildSelectionOptions = FTypedElementSelectionOptions(InSelectionOptions)
 				.SetChildElementInclusionMethod(InSelectionOptions.GetChildElementInclusionMethod() == ETypedElementChildInclusionMethod::Immediate ? ETypedElementChildInclusionMethod::None : InSelectionOptions.GetChildElementInclusionMethod());
@@ -252,7 +252,7 @@ bool UTypedElementSelectionSet::DeselectElement(const FTypedElementHandle& InEle
 
 	if (InSelectionOptions.GetChildElementInclusionMethod() != ETypedElementChildInclusionMethod::None)
 	{
-		if (TTypedElement<UTypedElementHierarchyInterface> ElementHierarchyHandle = ElementList->GetElement<UTypedElementHierarchyInterface>(InElementHandle))
+		if (TTypedElement<ITypedElementHierarchyInterface> ElementHierarchyHandle = ElementList->GetElement<ITypedElementHierarchyInterface>(InElementHandle))
 		{
 			const FTypedElementSelectionOptions ChildSelectionOptions = FTypedElementSelectionOptions(InSelectionOptions)
 				.SetChildElementInclusionMethod(InSelectionOptions.GetChildElementInclusionMethod() == ETypedElementChildInclusionMethod::Immediate ? ETypedElementChildInclusionMethod::None : InSelectionOptions.GetChildElementInclusionMethod());
@@ -395,7 +395,7 @@ FTypedElementListRef UTypedElementSelectionSet::GetNormalizedElementList(FTypedE
 void UTypedElementSelectionSet::GetNormalizedElementList(FTypedElementListConstRef InElementList, const FTypedElementSelectionNormalizationOptions& InNormalizationOptions, FTypedElementListRef OutNormalizedElements) const
 {
 	OutNormalizedElements->Reset();
-	InElementList->ForEachElement<UTypedElementSelectionInterface>([this, InElementList, &InNormalizationOptions, OutNormalizedElements](const TTypedElement<UTypedElementSelectionInterface>& InElementSelectionHandle)
+	InElementList->ForEachElement<ITypedElementSelectionInterface>([this, InElementList, &InNormalizationOptions, OutNormalizedElements](const TTypedElement<ITypedElementSelectionInterface>& InElementSelectionHandle)
 	{
 		// Note: FTypedElementSelectionSetElement needs a non-const FTypedElementList, but we won't call anything that will modify it...
 		FTypedElementSelectionSetElement SelectionSetElement(InElementSelectionHandle, ConstCastSharedRef<FTypedElementList>(InElementList), GetInterfaceCustomizationByTypeId(InElementSelectionHandle.GetId().GetTypeId()));
@@ -412,7 +412,7 @@ FTypedElementSelectionSetState UTypedElementSelectionSet::GetCurrentSelectionSta
 	CurrentState.CreatedFromSelectionSet = this;
 
 	CurrentState.TransactedElements.Reserve(ElementList->Num());
-	ElementList->ForEachElement<UTypedElementSelectionInterface>([&CurrentState](const TTypedElement<UTypedElementSelectionInterface>& InSelectionElement)
+	ElementList->ForEachElement<ITypedElementSelectionInterface>([&CurrentState](const TTypedElement<ITypedElementSelectionInterface>& InSelectionElement)
 	{
 		if (TUniquePtr<ITypedElementTransactedElement> TransactedElement = InSelectionElement.CreateTransactedElement())
 		{
@@ -457,7 +457,7 @@ void UTypedElementSelectionSet::RestoreSelectionState(const FTypedElementSelecti
 FTypedElementSelectionSetElement UTypedElementSelectionSet::ResolveSelectionSetElement(const FTypedElementHandle& InElementHandle) const
 {
 	return InElementHandle
-		? FTypedElementSelectionSetElement(ElementList->GetElement<UTypedElementSelectionInterface>(InElementHandle), ElementList, GetInterfaceCustomizationByTypeId(InElementHandle.GetId().GetTypeId()))
+		? FTypedElementSelectionSetElement(ElementList->GetElement<ITypedElementSelectionInterface>(InElementHandle), ElementList, GetInterfaceCustomizationByTypeId(InElementHandle.GetId().GetTypeId()))
 		: FTypedElementSelectionSetElement();
 }
 
diff --git a/Engine/Source/Runtime/TypedElementRuntime/Private/Elements/Interfaces/TypedElementAssetDataInterface.cpp b/Engine/Source/Runtime/TypedElementRuntime/Private/Elements/Interfaces/TypedElementAssetDataInterface.cpp
--- a/Engine/Source/Runtime/TypedElementRuntime/Private/Elements/Interfaces/TypedElementAssetDataInterface.cpp
+++ b/Engine/Source/Runtime/TypedElementRuntime/Private/Elements/Interfaces/TypedElementAssetDataInterface.cpp
@@ -2,7 +2,7 @@
 
 #include "Elements/Interfaces/TypedElementAssetDataInterface.h"
 
-TArray<FAssetData> UTypedElementAssetDataInterface::GetAllReferencedAssetDatas(const FTypedElementHandle& InElementHandle)
+TArray<FAssetData> ITypedElementAssetDataInterface::GetAllReferencedAssetDatas(const FTypedElementHandle& InElementHandle)
 {
 	TArray<FAssetData> AssetDatas;
 
@@ -15,7 +15,7 @@ TArray<FAssetData> UTypedElementAssetDataInterface::GetAllReferencedAssetDatas(c
 	return AssetDatas;
 }
 
-FAssetData UTypedElementAssetDataInterface::GetAssetData(const FTypedElementHandle& InElementHandle)
+FAssetData ITypedElementAssetDataInterface::GetAssetData(const FTypedElementHandle& InElementHandle)
 {
 	return FAssetData();
 }
diff --git a/Engine/Source/Runtime/TypedElementRuntime/Private/Elements/Interfaces/TypedElementObjectInterface.cpp b/Engine/Source/Runtime/TypedElementRuntime/Private/Elements/Interfaces/TypedElementObjectInterface.cpp
--- a/Engine/Source/Runtime/TypedElementRuntime/Private/Elements/Interfaces/TypedElementObjectInterface.cpp
+++ b/Engine/Source/Runtime/TypedElementRuntime/Private/Elements/Interfaces/TypedElementObjectInterface.cpp
@@ -4,12 +4,12 @@
 
 #include "Elements/Framework/TypedElementRegistry.h"
 
-UObject* UTypedElementObjectInterface::GetObject(const FTypedElementHandle& InElementHandle)
+UObject* ITypedElementObjectInterface::GetObject(const FTypedElementHandle& InElementHandle)
 {
 	return nullptr;
 }
 
-UClass* UTypedElementObjectInterface::GetObjectClass(const FTypedElementHandle& InElementHandle)
+UClass* ITypedElementObjectInterface::GetObjectClass(const FTypedElementHandle& InElementHandle)
 {
 	UObject* HandleAsObject = GetObject(InElementHandle);
 	return HandleAsObject ? HandleAsObject->GetClass() : nullptr;
diff --git a/Engine/Source/Runtime/TypedElementRuntime/Private/Elements/Interfaces/TypedElementSelectionInterface.cpp b/Engine/Source/Runtime/TypedElementRuntime/Private/Elements/Interfaces/TypedElementSelectionInterface.cpp
--- a/Engine/Source/Runtime/TypedElementRuntime/Private/Elements/Interfaces/TypedElementSelectionInterface.cpp
+++ b/Engine/Source/Runtime/TypedElementRuntime/Private/Elements/Interfaces/TypedElementSelectionInterface.cpp
@@ -3,19 +3,19 @@
 #include "Elements/Interfaces/TypedElementSelectionInterface.h"
 #include "Elements/Framework/TypedElementList.h"
 
-bool UTypedElementSelectionInterface::IsElementSelected(const FTypedElementHandle& InElementHandle, const FTypedElementListProxy InSelectionSet, const FTypedElementIsSelectedOptions& InSelectionOptions)
+bool ITypedElementSelectionInterface::IsElementSelected(const FTypedElementHandle& InElementHandle, const FTypedElementListProxy InSelectionSet, const FTypedElementIsSelectedOptions& InSelectionOptions)
 {
 	FTypedElementListConstPtr SelectionSetPtr = InSelectionSet.GetElementList();
 	return SelectionSetPtr && SelectionSetPtr->Contains(InElementHandle);
 }
 
-bool UTypedElementSelectionInterface::SelectElement(const FTypedElementHandle& InElementHandle, FTypedElementListProxy InSelectionSet, const FTypedElementSelectionOptions& InSelectionOptions)
+bool ITypedElementSelectionInterface::SelectElement(const FTypedElementHandle& InElementHandle, FTypedElementListProxy InSelectionSet, const FTypedElementSelectionOptions& InSelectionOptions)
 {
 	FTypedElementListPtr SelectionSetPtr = InSelectionSet.GetElementList();
 	return SelectionSetPtr && SelectionSetPtr->Add(InElementHandle);
 }
 
-bool UTypedElementSelectionInterface::DeselectElement(const FTypedElementHandle& InElementHandle, FTypedElementListProxy InSelectionSet, const FTypedElementSelectionOptions& InSelectionOptions)
+bool ITypedElementSelectionInterface::DeselectElement(const FTypedElementHandle& InElementHandle, FTypedElementListProxy InSelectionSet, const FTypedElementSelectionOptions& InSelectionOptions)
 {
 	FTypedElementListPtr SelectionSetPtr = InSelectionSet.GetElementList();
 	return SelectionSetPtr && SelectionSetPtr->Remove(InElementHandle);
diff --git a/Engine/Source/Runtime/TypedElementRuntime/Public/Elements/Framework/TypedElementSelectionSet.h b/Engine/Source/Runtime/TypedElementRuntime/Public/Elements/Framework/TypedElementSelectionSet.h
--- a/Engine/Source/Runtime/TypedElementRuntime/Public/Elements/Framework/TypedElementSelectionSet.h
+++ b/Engine/Source/Runtime/TypedElementRuntime/Public/Elements/Framework/TypedElementSelectionSet.h
@@ -37,15 +37,15 @@ class TYPEDELEMENTRUNTIME_API FTypedElementSelectionCustomization
 public:
 	virtual ~FTypedElementSelectionCustomization() = default;
 
-	//~ See UTypedElementSelectionInterface for API docs
-	virtual bool IsElementSelected(const TTypedElement<UTypedElementSelectionInterface>& InElementSelectionHandle, FTypedElementListConstRef InSelectionSet, const FTypedElementIsSelectedOptions& InSelectionOptions) { return InElementSelectionHandle.IsElementSelected(InSelectionSet, InSelectionOptions); }
-	virtual bool CanSelectElement(const TTypedElement<UTypedElementSelectionInterface>& InElementSelectionHandle, const FTypedElementSelectionOptions& InSelectionOptions) { return InElementSelectionHandle.CanSelectElement(InSelectionOptions); }
-	virtual bool CanDeselectElement(const TTypedElement<UTypedElementSelectionInterface>& InElementSelectionHandle, const FTypedElementSelectionOptions& InSelectionOptions) { return InElementSelectionHandle.CanDeselectElement(InSelectionOptions); }
-	virtual bool SelectElement(const TTypedElement<UTypedElementSelectionInterface>& InElementSelectionHandle, FTypedElementListRef InSelectionSet, const FTypedElementSelectionOptions& InSelectionOptions) { return InElementSelectionHandle.SelectElement(InSelectionSet, InSelectionOptions); }
-	virtual bool DeselectElement(const TTypedElement<UTypedElementSelectionInterface>& InElementSelectionHandle, FTypedElementListRef InSelectionSet, const FTypedElementSelectionOptions& InSelectionOptions) { return InElementSelectionHandle.DeselectElement(InSelectionSet, InSelectionOptions); }
-	virtual bool AllowSelectionModifiers(const TTypedElement<UTypedElementSelectionInterface>& InElementSelectionHandle, FTypedElementListConstRef InSelectionSet) { return InElementSelectionHandle.AllowSelectionModifiers(InSelectionSet); }
-	virtual FTypedElementHandle GetSelectionElement(const TTypedElement<UTypedElementSelectionInterface>& InElementSelectionHandle, FTypedElementListConstRef InCurrentSelection, const ETypedElementSelectionMethod InSelectionMethod) { return InElementSelectionHandle.GetSelectionElement(InCurrentSelection, InSelectionMethod); }
-	virtual void GetNormalizedElements(const TTypedElement<UTypedElementSelectionInterface>& InElementSelectionHandle, FTypedElementListConstRef InSelectionSet, const FTypedElementSelectionNormalizationOptions& InNormalizationOptions, FTypedElementListRef OutNormalizedElements);
+	//~ See ITypedElementSelectionInterface for API docs
+	virtual bool IsElementSelected(const TTypedElement<ITypedElementSelectionInterface>& InElementSelectionHandle, FTypedElementListConstRef InSelectionSet, const FTypedElementIsSelectedOptions& InSelectionOptions) { return InElementSelectionHandle.IsElementSelected(InSelectionSet, InSelectionOptions); }
+	virtual bool CanSelectElement(const TTypedElement<ITypedElementSelectionInterface>& InElementSelectionHandle, const FTypedElementSelectionOptions& InSelectionOptions) { return InElementSelectionHandle.CanSelectElement(InSelectionOptions); }
+	virtual bool CanDeselectElement(const TTypedElement<ITypedElementSelectionInterface>& InElementSelectionHandle, const FTypedElementSelectionOptions& InSelectionOptions) { return InElementSelectionHandle.CanDeselectElement(InSelectionOptions); }
+	virtual bool SelectElement(const TTypedElement<ITypedElementSelectionInterface>& InElementSelectionHandle, FTypedElementListRef InSelectionSet, const FTypedElementSelectionOptions& InSelectionOptions) { return InElementSelectionHandle.SelectElement(InSelectionSet, InSelectionOptions); }
+	virtual bool DeselectElement(const TTypedElement<ITypedElementSelectionInterface>& InElementSelectionHandle, FTypedElementListRef InSelectionSet, const FTypedElementSelectionOptions& InSelectionOptions) { return InElementSelectionHandle.DeselectElement(InSelectionSet, InSelectionOptions); }
+	virtual bool AllowSelectionModifiers(const TTypedElement<ITypedElementSelectionInterface>& InElementSelectionHandle, FTypedElementListConstRef InSelectionSet) { return InElementSelectionHandle.AllowSelectionModifiers(InSelectionSet); }
+	virtual FTypedElementHandle GetSelectionElement(const TTypedElement<ITypedElementSelectionInterface>& InElementSelectionHandle, FTypedElementListConstRef InCurrentSelection, const ETypedElementSelectionMethod InSelectionMethod) { return InElementSelectionHandle.GetSelectionElement(InCurrentSelection, InSelectionMethod); }
+	virtual void GetNormalizedElements(const TTypedElement<ITypedElementSelectionInterface>& InElementSelectionHandle, FTypedElementListConstRef InSelectionSet, const FTypedElementSelectionNormalizationOptions& InNormalizationOptions, FTypedElementListRef OutNormalizedElements);
 };
 
 /**
@@ -56,7 +56,7 @@ struct TYPEDELEMENTRUNTIME_API FTypedElementSelectionSetElement
 public:
 	FTypedElementSelectionSetElement() = default;
 
-	FTypedElementSelectionSetElement(TTypedElement<UTypedElementSelectionInterface> InElementSelectionHandle, FTypedElementListPtr InElementList, FTypedElementSelectionCustomization* InSelectionCustomization)
+	FTypedElementSelectionSetElement(TTypedElement<ITypedElementSelectionInterface> InElementSelectionHandle, FTypedElementListPtr InElementList, FTypedElementSelectionCustomization* InSelectionCustomization)
 		: ElementSelectionHandle(MoveTemp(InElementSelectionHandle))
 		, ElementList(InElementList)
 		, SelectionCustomization(InSelectionCustomization)
@@ -81,7 +81,7 @@ public:
 			&& SelectionCustomization;
 	}
 
-	//~ See UTypedElementSelectionInterface for API docs
+	//~ See ITypedElementSelectionInterface for API docs
 	bool IsElementSelected(const FTypedElementIsSelectedOptions& InSelectionOptions) const { return SelectionCustomization->IsElementSelected(ElementSelectionHandle, ElementList.ToSharedRef(), InSelectionOptions); }
 	bool CanSelectElement(const FTypedElementSelectionOptions& InSelectionOptions) const { return SelectionCustomization->CanSelectElement(ElementSelectionHandle, InSelectionOptions); }
 	bool CanDeselectElement(const FTypedElementSelectionOptions& InSelectionOptions) const { return SelectionCustomization->CanDeselectElement(ElementSelectionHandle, InSelectionOptions); }
@@ -92,7 +92,7 @@ public:
 	void GetNormalizedElements(const FTypedElementSelectionNormalizationOptions& InNormalizationOptions, FTypedElementListRef OutNormalizedElements) const { return SelectionCustomization->GetNormalizedElements(ElementSelectionHandle, ElementList.ToSharedRef(), InNormalizationOptions, OutNormalizedElements); }
 
 private:
-	TTypedElement<UTypedElementSelectionInterface> ElementSelectionHandle;
+	TTypedElement<ITypedElementSelectionInterface> ElementSelectionHandle;
 	FTypedElementListPtr ElementList;
 	FTypedElementSelectionCustomization* SelectionCustomization = nullptr;
 };
@@ -268,7 +268,7 @@ public:
 	 * Test whether there selected elements, optionally filtering to elements that implement the given interface.
 	 */
 	UFUNCTION(BlueprintCallable, BlueprintPure=false, Category="TypedElementFramework|Selection")
-	bool HasSelectedElements(const TSubclassOf<UTypedElementInterface>& InBaseInterfaceType = nullptr) const
+	bool HasSelectedElements(const TSubclassOf<UInterface>& InBaseInterfaceType = nullptr) const
 	{
 		return ElementList->HasElements(InBaseInterfaceType);
 	}
@@ -277,7 +277,7 @@ public:
 	 * Count the number of selected elements, optionally filtering to elements that implement the given interface.
 	 */
 	UFUNCTION(BlueprintCallable, BlueprintPure=false, Category="TypedElementFramework|Selection")
-	int32 CountSelectedElements(const TSubclassOf<UTypedElementInterface>& InBaseInterfaceType = nullptr) const
+	int32 CountSelectedElements(const TSubclassOf<UInterface>& InBaseInterfaceType = nullptr) const
 	{
 		return ElementList->CountElements(InBaseInterfaceType);
 	}
@@ -286,7 +286,7 @@ public:
 	 * Get the handle of every selected element, optionally filtering to elements that implement the given interface.
 	 */
 	UFUNCTION(BlueprintCallable, BlueprintPure=false, Category="TypedElementFramework|Selection")
-	TArray<FTypedElementHandle> GetSelectedElementHandles(const TSubclassOf<UTypedElementInterface>& InBaseInterfaceType = nullptr) const
+	TArray<FTypedElementHandle> GetSelectedElementHandles(const TSubclassOf<UInterface>& InBaseInterfaceType = nullptr) const
 	{
 		return ElementList->GetElementHandles(InBaseInterfaceType);
 	}
@@ -295,7 +295,7 @@ public:
 	 * Get the handle of every selected element, optionally filtering to elements that implement the given interface.
 	 */
 	template <typename ArrayAllocator>
-	void GetSelectedElementHandles(TArray<FTypedElementHandle, ArrayAllocator>& OutArray, const TSubclassOf<UTypedElementInterface>& InBaseInterfaceType = nullptr) const
+	void GetSelectedElementHandles(TArray<FTypedElementHandle, ArrayAllocator>& OutArray, const TSubclassOf<UInterface>& InBaseInterfaceType = nullptr) const
 	{
 		ElementList->GetElementHandles(OutArray, InBaseInterfaceType);
 	}
@@ -304,7 +304,7 @@ public:
 	 * Enumerate the handle of every selected element, optionally filtering to elements that implement the given interface.
 	 * @note Return true from the callback to continue enumeration.
 	 */
-	void ForEachSelectedElementHandle(TFunctionRef<bool(const FTypedElementHandle&)> InCallback, const TSubclassOf<UTypedElementInterface>& InBaseInterfaceType = nullptr) const
+	void ForEachSelectedElementHandle(TFunctionRef<bool(const FTypedElementHandle&)> InCallback, const TSubclassOf<UInterface>& InBaseInterfaceType = nullptr) const
 	{
 		ElementList->ForEachElementHandle(InCallback, InBaseInterfaceType);
 	}
diff --git a/Engine/Source/Runtime/TypedElementRuntime/Public/Elements/Interfaces/TypedElementAssetDataInterface.h b/Engine/Source/Runtime/TypedElementRuntime/Public/Elements/Interfaces/TypedElementAssetDataInterface.h
--- a/Engine/Source/Runtime/TypedElementRuntime/Public/Elements/Interfaces/TypedElementAssetDataInterface.h
+++ b/Engine/Source/Runtime/TypedElementRuntime/Public/Elements/Interfaces/TypedElementAssetDataInterface.h
@@ -2,13 +2,19 @@
 
 #pragma once
 
-#include "Elements/Framework/TypedElementHandle.h"
 #include "AssetRegistry/AssetData.h"
+#include "Elements/Framework/TypedElementHandle.h"
+#include "UObject/Interface.h"
 
 #include "TypedElementAssetDataInterface.generated.h"
 
-UCLASS(Abstract)
-class TYPEDELEMENTRUNTIME_API UTypedElementAssetDataInterface : public UTypedElementInterface
+UINTERFACE(MinimalAPI, BlueprintType, meta = (CannotImplementInterfaceInBlueprint))
+class UTypedElementAssetDataInterface : public UInterface
+{
+	GENERATED_BODY()
+};
+
+class TYPEDELEMENTRUNTIME_API ITypedElementAssetDataInterface
 {
 	GENERATED_BODY()
 
@@ -19,18 +25,18 @@ public:
 	 *
 	 * @returns An array of valid asset datas.
 	 */
-	UFUNCTION(BlueprintPure, Category = "TypedElementInterfaces|AssetData")
+	UFUNCTION(BlueprintCallable, Category = "TypedElementInterfaces|AssetData")
 	virtual TArray<FAssetData> GetAllReferencedAssetDatas(const FTypedElementHandle& InElementHandle);
 
 	/**
 	 * Returns the asset data for the given handle, if it exists.
 	 */
-	UFUNCTION(BlueprintPure, Category = "TypedElementInterfaces|AssetData")
+	UFUNCTION(BlueprintCallable, Category = "TypedElementInterfaces|AssetData")
 	virtual FAssetData GetAssetData(const FTypedElementHandle& InElementHandle);
 };
 
 template <>
-struct TTypedElement<UTypedElementAssetDataInterface> : public TTypedElementBase<UTypedElementAssetDataInterface>
+struct TTypedElement<ITypedElementAssetDataInterface> : public TTypedElementBase<ITypedElementAssetDataInterface>
 {
 	TArray<FAssetData> GetAllReferencedAssetDatas() const { return InterfacePtr->GetAllReferencedAssetDatas(*this); }
 	FAssetData GetAssetData() const { return InterfacePtr->GetAssetData(*this); }
diff --git a/Engine/Source/Runtime/TypedElementRuntime/Public/Elements/Interfaces/TypedElementHierarchyInterface.h b/Engine/Source/Runtime/TypedElementRuntime/Public/Elements/Interfaces/TypedElementHierarchyInterface.h
--- a/Engine/Source/Runtime/TypedElementRuntime/Public/Elements/Interfaces/TypedElementHierarchyInterface.h
+++ b/Engine/Source/Runtime/TypedElementRuntime/Public/Elements/Interfaces/TypedElementHierarchyInterface.h
@@ -3,10 +3,17 @@
 #pragma once
 
 #include "Elements/Framework/TypedElementHandle.h"
+#include "UObject/Interface.h"
+
 #include "TypedElementHierarchyInterface.generated.h"
 
-UCLASS(Abstract)
-class TYPEDELEMENTRUNTIME_API UTypedElementHierarchyInterface : public UTypedElementInterface
+UINTERFACE(MinimalAPI, BlueprintType, meta = (CannotImplementInterfaceInBlueprint))
+class UTypedElementHierarchyInterface : public UInterface
+{
+	GENERATED_BODY()
+};
+
+class TYPEDELEMENTRUNTIME_API ITypedElementHierarchyInterface
 {
 	GENERATED_BODY()
 
@@ -15,7 +22,7 @@ public:
 	 * Get the logical parent of this element, if any.
 	 * eg) A component might return its actor, or a static mesh instance might return its ISM component.
 	 */
-	UFUNCTION(BlueprintPure, Category="TypedElementInterfaces|Hierarchy")
+	UFUNCTION(BlueprintCallable, Category="TypedElementInterfaces|Hierarchy")
 	virtual FTypedElementHandle GetParentElement(const FTypedElementHandle& InElementHandle, const bool bAllowCreate = true) { return FTypedElementHandle(); }
 
 	/**
@@ -24,12 +31,12 @@ public:
 	 *
 	 * @note Appends to OutElementHandles.
 	 */
-	UFUNCTION(BlueprintPure, Category="TypedElementInterfaces|Hierarchy")
+	UFUNCTION(BlueprintCallable, Category="TypedElementInterfaces|Hierarchy")
 	virtual void GetChildElements(const FTypedElementHandle& InElementHandle, TArray<FTypedElementHandle>& OutElementHandles, const bool bAllowCreate = true) {}
 };
 
 template <>
-struct TTypedElement<UTypedElementHierarchyInterface> : public TTypedElementBase<UTypedElementHierarchyInterface>
+struct TTypedElement<ITypedElementHierarchyInterface> : public TTypedElementBase<ITypedElementHierarchyInterface>
 {
 	FTypedElementHandle GetParentElement(const bool bAllowCreate = true) const { return InterfacePtr->GetParentElement(*this, bAllowCreate); }
 	void GetChildElements(TArray<FTypedElementHandle>& OutElementHandles, const bool bAllowCreate = true) const { return InterfacePtr->GetChildElements(*this, OutElementHandles, bAllowCreate); }
diff --git a/Engine/Source/Runtime/TypedElementRuntime/Public/Elements/Interfaces/TypedElementObjectInterface.h b/Engine/Source/Runtime/TypedElementRuntime/Public/Elements/Interfaces/TypedElementObjectInterface.h
--- a/Engine/Source/Runtime/TypedElementRuntime/Public/Elements/Interfaces/TypedElementObjectInterface.h
+++ b/Engine/Source/Runtime/TypedElementRuntime/Public/Elements/Interfaces/TypedElementObjectInterface.h
@@ -3,10 +3,17 @@
 #pragma once
 
 #include "Elements/Framework/TypedElementHandle.h"
+#include "UObject/Interface.h"
+
 #include "TypedElementObjectInterface.generated.h"
 
-UCLASS(Abstract)
-class TYPEDELEMENTRUNTIME_API UTypedElementObjectInterface : public UTypedElementInterface
+UINTERFACE(MinimalAPI, BlueprintType, meta = (CannotImplementInterfaceInBlueprint))
+class UTypedElementObjectInterface : public UInterface
+{
+	GENERATED_BODY()
+};
+
+class TYPEDELEMENTRUNTIME_API ITypedElementObjectInterface
 {
 	GENERATED_BODY()
 
@@ -14,14 +21,14 @@ public:
 	/**
 	 * Get the object instance that this handle represents, if any.
 	 */
-	UFUNCTION(BlueprintPure, Category="TypedElementInterfaces|Object")
+	UFUNCTION(BlueprintCallable, Category="TypedElementInterfaces|Object")
 	virtual UObject* GetObject(const FTypedElementHandle& InElementHandle);
 
 	/**
 	 * Gets the object instance's class that the handle represents, if any. 
 	 */
-	UFUNCTION(BlueprintPure, Category = "TypedElementInterfaces|Object")
-	UClass* GetObjectClass(const FTypedElementHandle& InElementHandle);
+	UFUNCTION(BlueprintCallable, Category = "TypedElementInterfaces|Object")
+	virtual UClass* GetObjectClass(const FTypedElementHandle& InElementHandle);
 
 	/**
 	 * Attempts to cast the given handle to another class.
@@ -59,7 +66,7 @@ public:
 };
 
 template <>
-struct TTypedElement<UTypedElementObjectInterface> : public TTypedElementBase<UTypedElementObjectInterface>
+struct TTypedElement<ITypedElementObjectInterface> : public TTypedElementBase<ITypedElementObjectInterface>
 {
 	UObject* GetObject() const { return InterfacePtr->GetObject(*this); }
 	UClass* GetObjectClass() const { return InterfacePtr->GetObjectClass(*this); }
diff --git a/Engine/Source/Runtime/TypedElementRuntime/Public/Elements/Interfaces/TypedElementSelectionInterface.h b/Engine/Source/Runtime/TypedElementRuntime/Public/Elements/Interfaces/TypedElementSelectionInterface.h
--- a/Engine/Source/Runtime/TypedElementRuntime/Public/Elements/Interfaces/TypedElementSelectionInterface.h
+++ b/Engine/Source/Runtime/TypedElementRuntime/Public/Elements/Interfaces/TypedElementSelectionInterface.h
@@ -4,6 +4,8 @@
 
 #include "Elements/Framework/TypedElementHandle.h"
 #include "Elements/Framework/TypedElementListProxy.h"
+#include "UObject/Interface.h"
+
 #include "TypedElementSelectionInterface.generated.h"
 
 UENUM()
@@ -127,8 +129,13 @@ private:
 	FTypedHandleTypeId TypeId = 0;
 };
 
-UCLASS(Abstract)
-class TYPEDELEMENTRUNTIME_API UTypedElementSelectionInterface : public UTypedElementInterface
+UINTERFACE(MinimalAPI, BlueprintType, meta = (CannotImplementInterfaceInBlueprint))
+class UTypedElementSelectionInterface : public UInterface
+{
+	GENERATED_BODY()
+};
+
+class TYPEDELEMENTRUNTIME_API ITypedElementSelectionInterface
 {
 	GENERATED_BODY()
 
@@ -136,19 +143,19 @@ public:
 	/**
 	 * Test to see whether the given element is currently considered selected.
 	 */
-	UFUNCTION(BlueprintPure, Category="TypedElementInterfaces|Selection")
+	UFUNCTION(BlueprintCallable, Category="TypedElementInterfaces|Selection")
 	virtual bool IsElementSelected(const FTypedElementHandle& InElementHandle, const FTypedElementListProxy InSelectionSet, const FTypedElementIsSelectedOptions& InSelectionOptions);
 
 	/**
 	 * Test to see whether the given element can be selected.
 	 */
-	UFUNCTION(BlueprintPure, Category="TypedElementInterfaces|Selection")
+	UFUNCTION(BlueprintCallable, Category="TypedElementInterfaces|Selection")
 	virtual bool CanSelectElement(const FTypedElementHandle& InElementHandle, const FTypedElementSelectionOptions& InSelectionOptions) { return true; }
 
 	/**
 	 * Test to see whether the given element can be deselected.
 	 */
-	UFUNCTION(BlueprintPure, Category="TypedElementInterfaces|Selection")
+	UFUNCTION(BlueprintCallable, Category="TypedElementInterfaces|Selection")
 	virtual bool CanDeselectElement(const FTypedElementHandle& InElementHandle, const FTypedElementSelectionOptions& InSelectionOptions) { return true; }
 
 	/**
@@ -168,13 +175,13 @@ public:
 	/**
 	 * Test to see whether selection modifiers (Ctrl or Shift) are allowed while selecting this element.
 	 */
-	UFUNCTION(BlueprintPure, Category="TypedElementInterfaces|Selection")
+	UFUNCTION(BlueprintCallable, Category="TypedElementInterfaces|Selection")
 	virtual bool AllowSelectionModifiers(const FTypedElementHandle& InElementHandle, const FTypedElementListProxy InSelectionSet) { return true; }
 
 	/**
 	 * Given an element, return the element that should actually perform a selection operation.
 	 */
-	UFUNCTION(BlueprintPure, Category="TypedElementInterfaces|Selection")
+	UFUNCTION(BlueprintCallable, Category="TypedElementInterfaces|Selection")
 	virtual FTypedElementHandle GetSelectionElement(const FTypedElementHandle& InElementHandle, const FTypedElementListProxy InCurrentSelection, const ETypedElementSelectionMethod InSelectionMethod) { return InElementHandle; }
 
 	/**
@@ -217,7 +224,7 @@ protected:
 };
 
 template <>
-struct TTypedElement<UTypedElementSelectionInterface> : public TTypedElementBase<UTypedElementSelectionInterface>
+struct TTypedElement<ITypedElementSelectionInterface> : public TTypedElementBase<ITypedElementSelectionInterface>
 {
 	bool IsElementSelected(FTypedElementListConstRef InSelectionSet, const FTypedElementIsSelectedOptions& InSelectionOptions) const { return InterfacePtr->IsElementSelected(*this, InSelectionSet, InSelectionOptions); }
 	bool CanSelectElement(const FTypedElementSelectionOptions& InSelectionOptions) const { return InterfacePtr->CanSelectElement(*this, InSelectionOptions); }
