Reskin the Level Editor Main Toolbar,
Move Editor Modes to a Combo Button,
Combine Content and Create into one dropdown,
Reorder items in the Content Browser add and context menu to be more meaningful

diff --git a/Engine/Plugins/Bridge/Source/Bridge/Private/UI/BridgeUIManager.cpp b/Engine/Plugins/Bridge/Source/Bridge/Private/UI/BridgeUIManager.cpp
--- a/Engine/Plugins/Bridge/Source/Bridge/Private/UI/BridgeUIManager.cpp
+++ b/Engine/Plugins/Bridge/Source/Bridge/Private/UI/BridgeUIManager.cpp
@@ -69,9 +69,9 @@ void FBridgeUIManagerImpl::SetupMenuItem()
 	// 	}
 	// }
 
-	// Adding Bridge entry to Quick Content menu.
-	UToolMenu* ContentMenu = UToolMenus::Get()->ExtendMenu("LevelEditor.LevelEditorToolBar.ContentQuickMenu");
-	FToolMenuSection& Section = ContentMenu->FindOrAddSection("ExternalContent");
+	// Adding Bridge entry to Add Content menu.
+	UToolMenu* ContentMenu = UToolMenus::Get()->ExtendMenu("LevelEditor.LevelEditorToolBar.AddQuickMenu");
+	FToolMenuSection& Section = ContentMenu->FindOrAddSection("Content");
 	Section.AddMenuEntry("OpenBridgeTab",
 		LOCTEXT("OpenBridgeTab_Label", "Quixel Bridge"),
 		LOCTEXT("OpenBridgeTab_Desc", "Opens the Quixel Bridge."),
@@ -99,13 +99,12 @@ void FBridgeUIManagerImpl::SetupMenuItem()
 	
 	ContextMenuSection.AddMenuEntry(
 		"GetMegascans",
-		LOCTEXT("OpenBridgeTabText", "Add Quixel Content"),
+		LOCTEXT("OpenBridgeTabText", "Quixel Bridge"),
 		LOCTEXT("GetBridgeTooltip", "Add Megascans and DHI assets to project."),
 		FSlateIcon(FBridgeStyle::GetStyleSetName(), "Bridge.MenuLogo"),
 		FUIAction(FExecuteAction::CreateRaw(this, &FBridgeUIManagerImpl::CreateWIndow), FCanExecuteAction())
 	);
 
-
 }
 
 void FBridgeUIManagerImpl::AddPluginMenu(FMenuBuilder& MenuBuilder)
diff --git a/Engine/Plugins/Editor/ContentBrowser/ContentBrowserAssetDataSource/Source/ContentBrowserAssetDataSource/Private/NewAssetContextMenu.cpp b/Engine/Plugins/Editor/ContentBrowser/ContentBrowserAssetDataSource/Source/ContentBrowserAssetDataSource/Private/NewAssetContextMenu.cpp
--- a/Engine/Plugins/Editor/ContentBrowser/ContentBrowserAssetDataSource/Source/ContentBrowserAssetDataSource/Private/NewAssetContextMenu.cpp
+++ b/Engine/Plugins/Editor/ContentBrowser/ContentBrowserAssetDataSource/Source/ContentBrowserAssetDataSource/Private/NewAssetContextMenu.cpp
@@ -233,7 +233,7 @@ void FNewAssetContextMenu::MakeContextMenu(
 	if (InOnImportAssetRequested.IsBound() && !FirstSelectedPath.IsNone())
 	{
 		{
-			FToolMenuSection& Section = Menu->AddSection("ContentBrowserImportAsset", LOCTEXT("ImportAssetMenuHeading", "Import Asset"));
+			FToolMenuSection& Section = Menu->FindOrAddSection("ContentBrowserGetContent");
 			Section.AddMenuEntry(
 				"ImportAsset",
 				FText::Format(LOCTEXT("ImportAsset", "Import to {0}..."), FText::FromName(FirstSelectedPath)),
@@ -243,7 +243,7 @@ void FNewAssetContextMenu::MakeContextMenu(
 					FExecuteAction::CreateStatic(&FNewAssetContextMenu::ExecuteImportAsset, InOnImportAssetRequested, FirstSelectedPath),
 					CanExecuteAssetActionsDelegate
 					)
-				);
+				).InsertPosition = FToolMenuInsert(NAME_None, EToolMenuInsertType::First);
 		}
 	}
 
diff --git a/Engine/Plugins/Editor/SequencerAnimTools/Source/SequencerAnimTools/Private/SequencerToolsEditMode.cpp b/Engine/Plugins/Editor/SequencerAnimTools/Source/SequencerAnimTools/Private/SequencerToolsEditMode.cpp
--- a/Engine/Plugins/Editor/SequencerAnimTools/Source/SequencerAnimTools/Private/SequencerToolsEditMode.cpp
+++ b/Engine/Plugins/Editor/SequencerAnimTools/Source/SequencerAnimTools/Private/SequencerToolsEditMode.cpp
@@ -33,7 +33,8 @@ USequencerToolsEditMode::~USequencerToolsEditMode()
 
 bool USequencerToolsEditMode::IsCompatibleWith(FEditorModeID OtherModeID) const
 {
-	return OtherModeID == FName(TEXT("EM_SequencerMode"), FNAME_Find) || OtherModeID == FName(TEXT("EditMode.ControlRig"), FNAME_Find) || OtherModeID == FName(TEXT("EditMode.ControlRigEditor"), FNAME_Find);
+	// Compatible with all modes similar to FSequencerEdMode
+	return true;
 }
 
 //If we have one of our own active tools we pass input to it's commands.
diff --git a/Engine/Plugins/Enterprise/DatasmithImporter/Source/DatasmithImporter/Private/DatasmithImporterModule.cpp b/Engine/Plugins/Enterprise/DatasmithImporter/Source/DatasmithImporter/Private/DatasmithImporterModule.cpp
--- a/Engine/Plugins/Enterprise/DatasmithImporter/Source/DatasmithImporter/Private/DatasmithImporterModule.cpp
+++ b/Engine/Plugins/Enterprise/DatasmithImporter/Source/DatasmithImporter/Private/DatasmithImporterModule.cpp
@@ -168,7 +168,7 @@ void FDatasmithImporterModule::SetupMenuEntry()
 {
 	if (!IsRunningCommandlet())
 	{
-		UToolMenu* ContentMenu = UToolMenus::Get()->ExtendMenu("LevelEditor.LevelEditorToolBar.ContentQuickMenu");
+		UToolMenu* ContentMenu = UToolMenus::Get()->ExtendMenu("LevelEditor.LevelEditorToolBar.AddQuickMenu");
 		check(ContentMenu);
 
 		FToolMenuSection& Section = ContentMenu->FindOrAddSection("ImportAssets");
diff --git a/Engine/Source/Editor/ContentBrowser/Private/ContentBrowserSingleton.cpp b/Engine/Source/Editor/ContentBrowser/Private/ContentBrowserSingleton.cpp
--- a/Engine/Source/Editor/ContentBrowser/Private/ContentBrowserSingleton.cpp
+++ b/Engine/Source/Editor/ContentBrowser/Private/ContentBrowserSingleton.cpp
@@ -44,52 +44,22 @@ FContentBrowserSingleton::FContentBrowserSingleton()
 	// We're going to call a static function in the editor style module, so we need to make sure the module has actually been loaded
 	FModuleManager::Get().LoadModuleChecked("EditorStyle");
 
-	// Register the tab spawners for all content browsers
-	const FSlateIcon ContentBrowserIcon(FEditorStyle::GetStyleSetName(), "ContentBrowser.TabIcon");
+	const FSlateIcon ContentBrowserIcon(FAppStyle::Get().GetStyleSetName(), "ContentBrowser.TabIcon");
 	const IWorkspaceMenuStructure& MenuStructure = WorkspaceMenu::GetMenuStructure();
 	TSharedRef<FWorkspaceItem> ContentBrowserGroup = MenuStructure.GetLevelEditorCategory()->AddGroup(
-		LOCTEXT( "WorkspaceMenu_ContentBrowserCategory", "Content Browser" ),
-		LOCTEXT( "ContentBrowserMenuTooltipText", "Open a Content Browser tab." ),
+		LOCTEXT("WorkspaceMenu_ContentBrowserCategory", "Content Browser"),
+		LOCTEXT("ContentBrowserMenuTooltipText", "Open a Content Browser tab."),
 		ContentBrowserIcon,
 		true);
 
-	UToolMenu* ContentMenu = UToolMenus::Get()->ExtendMenu("LevelEditor.LevelEditorToolBar.ContentQuickMenu");
-	FToolMenuSection& Section = ContentMenu->FindOrAddSection("ProjectContent");
-	Section.AddMenuEntry("FocusContentBrowser",
-		LOCTEXT("FocusContentBrowser_Label", "Focus Content Browser"),
-		LOCTEXT("FocusContentBrowser_Desc", "Focuses the most recently active content browser tab."),
-		FSlateIcon(),
-		FUIAction(FExecuteAction::CreateRaw(this, &FContentBrowserSingleton::FocusPrimaryContentBrowser, true), FCanExecuteAction())
-	);
-	Section.AddSeparator(NAME_None);
 
-	auto AddSpawnerEntryToMenuSection = [](FToolMenuSection& InSection, FTabSpawnerEntry& InSpawnerNode, FName InTabID)
-	{
-		InSection.AddMenuEntry(
-			InTabID,
-			InSpawnerNode.GetDisplayName().IsEmpty() ? FText::FromName(InTabID) : InSpawnerNode.GetDisplayName(),
-			InSpawnerNode.GetTooltipText(),
-			FSlateIcon(),
-			FGlobalTabmanager::Get()->GetUIActionForTabSpawnerMenuEntry(InSpawnerNode.AsSpawnerEntry()),
-			EUserInterfaceActionType::Check
-			);
-	};
-
-	for ( int32 BrowserIdx = 0; BrowserIdx < UE_ARRAY_COUNT(ContentBrowserTabIDs); BrowserIdx++ )
-	{
-		const FName TabID = FName(*FString::Printf(TEXT("ContentBrowserTab%d"), BrowserIdx + 1));
-		ContentBrowserTabIDs[BrowserIdx] = TabID;
+	UToolMenu* ContentMenu = UToolMenus::Get()->ExtendMenu("LevelEditor.LevelEditorToolBar.AddQuickMenu");
+	FToolMenuSection& Section = ContentMenu->FindOrAddSection("Content");
 
-		const FText DefaultDisplayName = GetContentBrowserLabelWithIndex( BrowserIdx );
-
-		FTabSpawnerEntry& ContentBrowserTabSpawner = FGlobalTabmanager::Get()->RegisterNomadTabSpawner( TabID, FOnSpawnTab::CreateRaw(this, &FContentBrowserSingleton::SpawnContentBrowserTab, BrowserIdx) )
-			.SetDisplayName(DefaultDisplayName)
-			.SetTooltipText( LOCTEXT( "ContentBrowserMenuTooltipText", "Open a Content Browser tab." ) )
-			.SetGroup( ContentBrowserGroup )
-			.SetIcon(ContentBrowserIcon);
-
-		AddSpawnerEntryToMenuSection(Section, ContentBrowserTabSpawner, TabID);
-	}
+	Section.AddSubMenu("ContentBrowser", LOCTEXT("ContentBrowserMenu", "Content Browser"), LOCTEXT("ContentBrowserTooltip", "Actions related to the Content Browser"),
+						FNewToolMenuDelegate::CreateRaw(this, &FContentBrowserSingleton::GetContentBrowserSubMenu, ContentBrowserGroup), false, 
+						FSlateIcon(FAppStyle::Get().GetStyleSetName(), "LevelEditor.OpenContentBrowser"))
+						.InsertPosition = FToolMenuInsert("OpenMarketplace", EToolMenuInsertType::After);
 
 	// Register to be notified when properties are edited
 	FEditorDelegates::LoadSelectedAssetsIfNeeded.AddRaw(this, &FContentBrowserSingleton::OnEditorLoadSelectedAssetsIfNeeded);
@@ -988,4 +958,47 @@ void FContentBrowserSingleton::PopulateConfigValues()
 	}
 }
 
+void FContentBrowserSingleton::GetContentBrowserSubMenu(UToolMenu* Menu, TSharedRef<FWorkspaceItem> ContentBrowserGroup)
+{
+	// Register the tab spawners for all content browsers
+	const FSlateIcon ContentBrowserIcon(FAppStyle::Get().GetStyleSetName(), "ContentBrowser.TabIcon");
+
+	FToolMenuSection& Section = Menu->AddSection("ContentBrowser");
+
+	Section.AddMenuEntry("FocusContentBrowser",
+		LOCTEXT("FocusContentBrowser_Label", "Focus Content Browser"),
+		LOCTEXT("FocusContentBrowser_Desc", "Focuses the most recently active content browser tab."),
+		FSlateIcon(),
+		FUIAction(FExecuteAction::CreateRaw(this, &FContentBrowserSingleton::FocusPrimaryContentBrowser, true), FCanExecuteAction())
+	);
+
+	auto AddSpawnerEntryToMenuSection = [](FToolMenuSection& InSection, FTabSpawnerEntry& InSpawnerNode, FName InTabID)
+	{
+		InSection.AddMenuEntry(
+			InTabID,
+			InSpawnerNode.GetDisplayName().IsEmpty() ? FText::FromName(InTabID) : InSpawnerNode.GetDisplayName(),
+			InSpawnerNode.GetTooltipText(),
+			FSlateIcon(),
+			FGlobalTabmanager::Get()->GetUIActionForTabSpawnerMenuEntry(InSpawnerNode.AsSpawnerEntry()),
+			EUserInterfaceActionType::Check
+		);
+	};
+
+	for (int32 BrowserIdx = 0; BrowserIdx < UE_ARRAY_COUNT(ContentBrowserTabIDs); BrowserIdx++)
+	{
+		const FName TabID = FName(*FString::Printf(TEXT("ContentBrowserTab%d"), BrowserIdx + 1));
+		ContentBrowserTabIDs[BrowserIdx] = TabID;
+
+		const FText DefaultDisplayName = GetContentBrowserLabelWithIndex(BrowserIdx);
+
+		FTabSpawnerEntry& ContentBrowserTabSpawner = FGlobalTabmanager::Get()->RegisterNomadTabSpawner(TabID, FOnSpawnTab::CreateRaw(this, &FContentBrowserSingleton::SpawnContentBrowserTab, BrowserIdx))
+			.SetDisplayName(DefaultDisplayName)
+			.SetTooltipText(LOCTEXT("ContentBrowserMenuTooltipText", "Open a Content Browser tab."))
+			.SetGroup(ContentBrowserGroup)
+			.SetIcon(ContentBrowserIcon);
+
+		AddSpawnerEntryToMenuSection(Section, ContentBrowserTabSpawner, TabID);
+	}
+}
+
 #undef LOCTEXT_NAMESPACE
diff --git a/Engine/Source/Editor/ContentBrowser/Private/ContentBrowserSingleton.h b/Engine/Source/Editor/ContentBrowser/Private/ContentBrowserSingleton.h
--- a/Engine/Source/Editor/ContentBrowser/Private/ContentBrowserSingleton.h
+++ b/Engine/Source/Editor/ContentBrowser/Private/ContentBrowserSingleton.h
@@ -14,6 +14,8 @@ class FTabManager;
 class FViewport;
 class SContentBrowser;
 class UFactory;
+class UToolMenu;
+class FWorkspaceItem;
 
 #define MAX_CONTENT_BROWSERS 4
 
@@ -130,6 +132,9 @@ private:
 	/** Populates properties that come from ini files */
 	void PopulateConfigValues();
 
+	/** Creates the Content Browser submenu in the Level Editor Toolbar Add menu */
+	void GetContentBrowserSubMenu(UToolMenu* Menu, TSharedRef<FWorkspaceItem> ContentBrowserGroup);
+
 public:
 	/** The tab identifier/instance name for content browser tabs */
 	FName ContentBrowserTabIDs[MAX_CONTENT_BROWSERS];
diff --git a/Engine/Source/Editor/ContentBrowser/Private/NewAssetOrClassContextMenu.cpp b/Engine/Source/Editor/ContentBrowser/Private/NewAssetOrClassContextMenu.cpp
--- a/Engine/Source/Editor/ContentBrowser/Private/NewAssetOrClassContextMenu.cpp
+++ b/Engine/Source/Editor/ContentBrowser/Private/NewAssetOrClassContextMenu.cpp
@@ -39,14 +39,22 @@ void FNewAssetOrClassContextMenu::MakeContextMenu(
 		// We can execute folder actions when we only have a single path selected, and that path is a valid path for creating a folder
 		return bHasSinglePathSelected && bIsValidNewFolderPath;
 	};
+
+	// If a single folder is selected and it is not valid for creating a new folder (The All Folder), don't create a context menu
+	if (bHasSinglePathSelected && !bIsValidNewFolderPath)
+	{
+		return;
+	}
+
 	const FCanExecuteAction CanExecuteFolderActionsDelegate = FCanExecuteAction::CreateLambda(CanExecuteFolderActions);
 
 	// Get Content
+	FToolMenuSection& GetContentSection = Menu->AddSection("ContentBrowserGetContent", LOCTEXT("GetContentMenuHeading", "Get Content"));
+
 	if ( InOnGetContentRequested.IsBound() )
 	{
 		{
-			FToolMenuSection& Section = Menu->AddSection( "ContentBrowserGetContent", LOCTEXT( "GetContentMenuHeading", "Content" ) );
-			Section.AddMenuEntry(
+			GetContentSection.AddMenuEntry(
 				"GetContent",
 				LOCTEXT( "GetContentText", "Add Feature or Content Pack..." ),
 				LOCTEXT( "GetContentTooltip", "Add features and content packs to the project." ),
diff --git a/Engine/Source/Editor/ContentBrowser/Private/SContentBrowser.cpp b/Engine/Source/Editor/ContentBrowser/Private/SContentBrowser.cpp
--- a/Engine/Source/Editor/ContentBrowser/Private/SContentBrowser.cpp
+++ b/Engine/Source/Editor/ContentBrowser/Private/SContentBrowser.cpp
@@ -2271,10 +2271,8 @@ void SContentBrowser::PopulateAddNewContextMenu(class UToolMenu* Menu)
 
 	// New feature packs don't depend on the current paths, so we always add this item if it was requested
 	FNewAssetOrClassContextMenu::FOnGetContentRequested OnGetContentRequested;
-	if (ContextObject->OwnerDomain == EContentBrowserDataMenuContext_AddNewMenuDomain::Toolbar)
-	{
-		OnGetContentRequested = FNewAssetOrClassContextMenu::FOnGetContentRequested::CreateSP(this, &SContentBrowser::OnAddContentRequested);
-	}
+	
+	OnGetContentRequested = FNewAssetOrClassContextMenu::FOnGetContentRequested::CreateSP(this, &SContentBrowser::OnAddContentRequested);
 
 	FNewAssetOrClassContextMenu::MakeContextMenu(
 		Menu,
diff --git a/Engine/Source/Editor/EditorStyle/Private/StarshipStyle.cpp b/Engine/Source/Editor/EditorStyle/Private/StarshipStyle.cpp
--- a/Engine/Source/Editor/EditorStyle/Private/StarshipStyle.cpp
+++ b/Engine/Source/Editor/EditorStyle/Private/StarshipStyle.cpp
@@ -4540,9 +4540,13 @@ void FStarshipEditorStyle::FStyle::SetupLevelEditorStyle()
 		Set( "LevelEditor.OpenPlaceActors", new IMAGE_BRUSH_SVG( "Starship/Common/PlaceActors", Icon20x20 ) );
 		Set( "LevelEditor.OpenContentBrowser", new IMAGE_BRUSH_SVG( "Starship/MainToolbar/content", Icon20x20 ) );
 		Set( "LevelEditor.OpenMarketplace", new IMAGE_BRUSH_SVG( "Starship/MainToolbar/marketplace", Icon20x20) );
+		Set( "LevelEditor.ImportContent", new CORE_IMAGE_BRUSH_SVG("Starship/Common/import", Icon20x20));
 		Set( "LevelEditor.OpenLevelBlueprint", new IMAGE_BRUSH_SVG( "Starship/MainToolbar/blueprints", Icon20x20) );
 		Set( "LevelEditor.OpenCinematic", new IMAGE_BRUSH_SVG("Starship/MainToolbar/cinematics", Icon20x20));
 
+		Set( "LevelEditor.OpenAddContent.Background", new IMAGE_BRUSH_SVG("Starship/MainToolbar/PlaceActorsBase", Icon20x20));
+		Set( "LevelEditor.OpenAddContent.Overlay", new IMAGE_BRUSH_SVG("Starship/MainToolbar/ToolBadgePlus", Icon20x20, FStyleColors::AccentGreen));
+
 		Set( "LevelEditor.CreateClassBlueprint", new IMAGE_BRUSH("Icons/icon_class_Blueprint_New_16x", Icon16x16));
 		Set( "LevelEditor.OpenClassBlueprint", new IMAGE_BRUSH("Icons/icon_class_Blueprint_Open_16x", Icon16x16));
 
@@ -5433,7 +5437,7 @@ void FStarshipEditorStyle::FStyle::SetupPersonaStyle()
 
 
 		MainToolbarLeftButton.SetButtonStyle(LeftToolbarButton);
-		MainToolbarLeftButton.SetButtonPadding(FMargin(8.f, 0.0f, 0.0f, 0.0f));
+		MainToolbarLeftButton.SetButtonPadding(FMargin(10.f, 0.0f, 0.0f, 0.0f));
 		MainToolbarLeftButton.SetSeparatorPadding(FMargin(0.f, 0.f, 8.f, 0.f));
 
 		Set("Toolbar.BackplateLeft", MainToolbarLeftButton);
@@ -5500,8 +5504,8 @@ void FStarshipEditorStyle::FStyle::SetupPersonaStyle()
 			.SetPressedPadding(FMargin(2.f, 2.f, 8.f, 2.f));
 
 		MainToolbarRightButton.SetButtonStyle(RightToolbarButton);
-		MainToolbarRightButton.SetButtonPadding(FMargin(0.0f, 0.0f, 8.0f, 0.0f));
-		MainToolbarRightButton.SetSeparatorPadding(FMargin(8.f, 0.f, 8.f, 0.f));
+		MainToolbarRightButton.SetButtonPadding(FMargin(0.0f, 0.0f, 4.0f, 0.0f));
+		MainToolbarRightButton.SetSeparatorPadding(FMargin(6.f, 0.f, 8.f, 0.f));
 
 		Set("Toolbar.BackplateRight", MainToolbarRightButton);
 
@@ -6371,7 +6375,7 @@ void FStarshipEditorStyle::FStyle::SetupContentBrowserStyle()
 			.SetFont(DEFAULT_FONT("Regular", 7)));
 
 
-		Set( "ContentBrowser.AddContent", new IMAGE_BRUSH( "Icons/icon_AddContent_40x", Icon25x25 ) );
+		Set( "ContentBrowser.AddContent", new IMAGE_BRUSH_SVG( "Starship/ContentBrowser/ContentPack", Icon20x20 ) );
 		Set( "ContentBrowser.ImportPackage", new IMAGE_BRUSH( "Icons/icon_Import_40x", Icon25x25 ) );
 	
 		// Asset Context Menu
@@ -6440,7 +6444,7 @@ void FStarshipEditorStyle::FStyle::SetupContentBrowserStyle()
 		Set( "ContentBrowser.AssetDragDropTooltipBackground", new BOX_BRUSH( "Old/Menu_Background", FMargin(8.0f/64.0f) ) );
 		Set( "ContentBrowser.CollectionTreeDragDropBorder", new BOX_BRUSH( "Old/Window/ViewportDebugBorder", 0.8f ) );
 		Set( "ContentBrowser.PopupMessageIcon", new IMAGE_BRUSH( "Icons/alert", Icon32x32) );
-		Set( "ContentBrowser.NewFolderIcon", new IMAGE_BRUSH("Icons/icon_AddFolder_16x", Icon16x16 ) );
+		Set( "ContentBrowser.NewFolderIcon", new CORE_IMAGE_BRUSH_SVG("Starship/Common/folder-plus", Icon16x16 ) );
 
 		Set( "ContentBrowser.Local", new IMAGE_BRUSH( "ContentBrowser/Content_Local_12x", Icon12x12 ) );
 		Set( "ContentBrowser.Local.Small", new IMAGE_BRUSH( "ContentBrowser/Content_Local_16x", Icon16x16 ) );
diff --git a/Engine/Source/Editor/LevelEditor/Private/LevelEditorActions.cpp b/Engine/Source/Editor/LevelEditor/Private/LevelEditorActions.cpp
--- a/Engine/Source/Editor/LevelEditor/Private/LevelEditorActions.cpp
+++ b/Engine/Source/Editor/LevelEditor/Private/LevelEditorActions.cpp
@@ -103,6 +103,8 @@
 #include "PlatformInfo.h"
 #include "Misc/CoreMisc.h"
 #include "Misc/ScopeExit.h"
+#include "Dialogs/DlgPickPath.h"
+#include "AssetToolsModule.h"
 
 #include "Internationalization/Culture.h"
 
@@ -2218,6 +2220,26 @@ void FLevelEditorActionCallbacks::OpenMarketplace()
 	FUnrealEdMisc::Get().OpenMarketplace();
 }
 
+void FLevelEditorActionCallbacks::ImportContent()
+{
+	FString Path = "/Game";
+
+	//Ask the user for the root path where they want to any content to be placed
+	TSharedRef<SDlgPickPath> PickContentPathDlg =
+		SNew(SDlgPickPath)
+		.Title(LOCTEXT("ChooseImportRootContentPath", "Choose a location to import the content into"));
+
+	if (PickContentPathDlg->ShowModal() == EAppReturnType::Cancel)
+	{
+		return;
+	}
+
+	Path = PickContentPathDlg->GetPath().ToString();
+
+	FAssetToolsModule& AssetToolsModule = FModuleManager::Get().LoadModuleChecked<FAssetToolsModule>(TEXT("AssetTools"));
+
+	AssetToolsModule.Get().ImportAssetsWithDialog(Path);
+}
 
 void FLevelEditorActionCallbacks::ToggleVR()
 {
@@ -3481,7 +3503,8 @@ void FLevelEditorCommands::RegisterCommands()
 	UI_COMMAND( WorldProperties, "World Settings", "Displays the world settings", EUserInterfaceActionType::Button, FInputChord() );
 	UI_COMMAND( OpenPlaceActors, "Place Actors Panel", "Opens the Place Actors Panel", EUserInterfaceActionType::Button, FInputChord(EModifierKey::Alt, EKeys::P) );
 	UI_COMMAND( OpenContentBrowser, "Open Content Browser", "Opens the Content Browser", EUserInterfaceActionType::Button, FInputChord(EModifierKey::Control|EModifierKey::Shift, EKeys::F) );
-	UI_COMMAND( OpenMarketplace, "Open Marketplace", "Opens the Marketplace", EUserInterfaceActionType::Button, FInputChord() );
+	UI_COMMAND( OpenMarketplace, "Unreal Marketplace", "Opens the Marketplace", EUserInterfaceActionType::Button, FInputChord() );
+	UI_COMMAND( ImportContent, "Import Content...", "Import Content into a specified location", EUserInterfaceActionType::Button, FInputChord());
 
 	UI_COMMAND( ToggleVR, "Toggle VR", "Toggles VR (Virtual Reality) mode", EUserInterfaceActionType::ToggleButton, FInputChord(EModifierKey::Shift, EKeys::V ) );
 
diff --git a/Engine/Source/Editor/LevelEditor/Private/LevelEditorModesActions.cpp b/Engine/Source/Editor/LevelEditor/Private/LevelEditorModesActions.cpp
--- a/Engine/Source/Editor/LevelEditor/Private/LevelEditorModesActions.cpp
+++ b/Engine/Source/Editor/LevelEditor/Private/LevelEditorModesActions.cpp
@@ -51,7 +51,7 @@ void FLevelEditorModesCommands::RegisterCommands()
 				ModeName,
 				Tooltip,
 				ModeIcon,
-				EUserInterfaceActionType::ToggleButton,
+				EUserInterfaceActionType::CollapsedButton,
 				DefaultKeyBinding);
 
 			EditorModeCommands.Add(EditorModeCommand);
diff --git a/Engine/Source/Editor/LevelEditor/Private/LevelEditorToolBar.cpp b/Engine/Source/Editor/LevelEditor/Private/LevelEditorToolBar.cpp
--- a/Engine/Source/Editor/LevelEditor/Private/LevelEditorToolBar.cpp
+++ b/Engine/Source/Editor/LevelEditor/Private/LevelEditorToolBar.cpp
@@ -1164,6 +1164,37 @@ void LevelEditorActionHelpers::OnSelectPlayerControllerClassPicked(UClass* InCho
 	FSlateApplication::Get().DismissAllMenus();
 }
 
+FText FLevelEditorToolBar::GetActiveModeName(TWeakPtr<SLevelEditor> LevelEditorPtr)
+{
+#define LOCTEXT_NAMESPACE "LevelEditorToolBar"
+
+	for (const FEditorModeInfo& Mode : GEditor->GetEditorSubsystem<UAssetEditorSubsystem>()->GetEditorModeInfoOrderedByPriority())
+	{
+		TSharedPtr<SLevelEditor> LevelEditorPin = LevelEditorPtr.Pin();
+		if (LevelEditorPin.IsValid() && LevelEditorPin->GetEditorModeManager().IsModeActive(Mode.ID) && Mode.IsVisible())
+		{
+			return FText::Format(LOCTEXT("ActiveMode", "{0} Mode"), Mode.Name);
+		}
+	}
+
+	return LOCTEXT("ActiveMode", "No Active Mode");
+
+#undef LOCTEXT_NAMESPACE
+}
+
+const FSlateBrush* FLevelEditorToolBar::GetActiveModeIcon(TWeakPtr<SLevelEditor> LevelEditorPtr)
+{
+	for (const FEditorModeInfo& Mode : GEditor->GetEditorSubsystem<UAssetEditorSubsystem>()->GetEditorModeInfoOrderedByPriority())
+	{
+		TSharedPtr<SLevelEditor> LevelEditorPin = LevelEditorPtr.Pin();
+		if (LevelEditorPin.IsValid() && LevelEditorPin->GetEditorModeManager().IsModeActive(Mode.ID) && Mode.IsVisible())
+		{
+			return Mode.IconBrush.GetIcon();
+		}
+	}
+	return nullptr;
+}
+
 void FLevelEditorToolBar::RegisterLevelEditorToolBar( const TSharedRef<FUICommandList>& InCommandList, const TSharedRef<SLevelEditor> InLevelEditor)
 {
 	RegisterSourceControlMenu();
@@ -1173,14 +1204,15 @@ void FLevelEditorToolBar::RegisterLevelEditorToolBar( const TSharedRef<FUIComman
 	RegisterOpenBlueprintMenu();
 	RegisterContentMenu();
 	RegisterCreateMenu();
+	RegisterAddMenu();
 
 #define LOCTEXT_NAMESPACE "LevelEditorToolBar"
 
-	UToolMenu* AssetsToolBar = UToolMenus::Get()->RegisterMenu("LevelEditor.LevelEditorToolBar.AssetsToolBar", NAME_None, EMultiBoxType::SlimHorizontalToolBar);
-	AssetsToolBar->StyleName = "AssetEditorToolbar";
+	UToolMenu* ModesToolbar = UToolMenus::Get()->RegisterMenu("LevelEditor.LevelEditorToolBar.ModesToolBar", NAME_None, EMultiBoxType::SlimHorizontalToolBar);
+	ModesToolbar->StyleName = "AssetEditorToolbar";
 	{
 		{
-			FToolMenuSection& Section = AssetsToolBar->AddSection("File");
+			FToolMenuSection& Section = ModesToolbar->AddSection("File");
 
 			// Save All Levels
 			Section.AddEntry(FToolMenuEntry::InitToolBarButton(
@@ -1190,33 +1222,142 @@ void FLevelEditorToolBar::RegisterLevelEditorToolBar( const TSharedRef<FUIComman
 				FSlateIcon(FEditorStyle::GetStyleSetName(), "AssetEditor.SaveAsset"),
 				NAME_None,
 				FName("SaveAllLevels")
-				));
+			));
 		}
 
+		TWeakPtr<SLevelEditor> LevelEditorPtr = InLevelEditor;
+		ModesToolbar->AddDynamicSection("EditorModes", FNewToolMenuDelegate::CreateLambda([LevelEditorPtr](UToolMenu* ToolMenu)
+			{
+				FToolMenuSection& Section = ToolMenu->AddSection("EditorModes");
+				
+				// Combo Button to swap editor modes
+				TSharedRef<SComboButton> EditorModesComboButton = SNew(SComboButton)
+					.OnGetMenuContent_Lambda([LevelEditorPtr]()
+					{
+						const FLevelEditorModule& LevelEditorModule = FModuleManager::LoadModuleChecked<FLevelEditorModule>("LevelEditor");
+						const FLevelEditorModesCommands& Commands = LevelEditorModule.GetLevelEditorModesCommands();
+
+						TArray<FEditorModeInfo> DefaultModes;
+						TArray<FEditorModeInfo> NonDefaultModes;
+						TArray<TSharedPtr<FUICommandInfo>> CommandInfos;
+
+						for (const FEditorModeInfo& Mode : GEditor->GetEditorSubsystem<UAssetEditorSubsystem>()->GetEditorModeInfoOrderedByPriority())
+						{
+							TSharedPtr<SLevelEditor> LevelEditorPin = LevelEditorPtr.Pin();
+							if (LevelEditorPin.IsValid() && LevelEditorPin->GetEditorModeManager().IsDefaultMode(Mode.ID))
+							{
+								DefaultModes.Add(Mode.ID);
+							}
+							else
+							{
+								NonDefaultModes.Add(Mode.ID);
+							}
+
+						}
+
+						auto GetCommandForModes = [&CommandInfos, &Commands](TArrayView<FEditorModeInfo> Modes)
+						{
+							for (const FEditorModeInfo& Mode : Modes)
+							{
+								FName EditorModeCommandName = FName(*(FString("EditorMode.") + Mode.ID.ToString()));
+
+								TSharedPtr<FUICommandInfo> EditorModeCommand =
+									FInputBindingManager::Get().FindCommandInContext(Commands.GetContextName(), EditorModeCommandName);
+
+								if (EditorModeCommand->GetDefaultChord(EMultipleKeyBindingIndex(0)).IsValidChord())
+								{
+									CommandInfos.Add(EditorModeCommand);
+								}
+									
+							}
+						};
+
+						// Default Modes first
+						GetCommandForModes(DefaultModes);
+
+						GetCommandForModes(NonDefaultModes);
+
+						FMenuBuilder MenuBuilder(true, LevelEditorModule.GetGlobalLevelEditorActions());
+
+						MenuBuilder.BeginSection("EditorModes");
+
+						for (TSharedPtr<FUICommandInfo> Command : CommandInfos)
+						{
+							MenuBuilder.AddMenuEntry(Command);
+						}
+
+						MenuBuilder.EndSection();
+
+						return MenuBuilder.MakeWidget();
+					})
+					.ButtonContent()
+					[
+						SNew(SHorizontalBox)
+						+ SHorizontalBox::Slot()
+						.AutoWidth()
+						.Padding(FMargin(0.f, 0.f, 6.f, 0.f))
+						[
+							SNew(SBox)
+							.WidthOverride(16.f)
+							.HeightOverride(16.f)
+							[
+								SNew(SImage)
+								.Image_Static(&FLevelEditorToolBar::GetActiveModeIcon, LevelEditorPtr)
+							]
+							
+						]
+						+ SHorizontalBox::Slot()
+						[
+							SNew(STextBlock)
+							.Text_Static(&FLevelEditorToolBar::GetActiveModeName, LevelEditorPtr)
+						]
+						
+					];
+
+				// Horizontal Box to add some spacing beside the modes combo button
+				TSharedRef<SHorizontalBox> EditorModesWidget = 
+					SNew(SHorizontalBox)
+					+ SHorizontalBox::Slot()
+					.FillWidth(1.0)
+					[
+						SNew(SSpacer)
+						.Size(FVector2D(10.f, 1.0f))
+					]
+					+ SHorizontalBox::Slot()
+					.AutoWidth()
+					[
+						EditorModesComboButton
+					]
+					+ SHorizontalBox::Slot()
+					.FillWidth(1.0)
+					[
+						SNew(SSpacer)
+						.Size(FVector2D(10.f, 1.0f))
+					];
+					
+
+				Section.AddEntry(FToolMenuEntry::InitWidget("Editor Modes", EditorModesWidget, LOCTEXT("EditorModesLabel", "Editor Modes")));
+
+				Section.AddSeparator(NAME_None);
+			}));
+	}
+
+	UToolMenu* AssetsToolBar = UToolMenus::Get()->RegisterMenu("LevelEditor.LevelEditorToolBar.AssetsToolBar", NAME_None, EMultiBoxType::SlimHorizontalToolBar);
+	AssetsToolBar->StyleName = "AssetEditorToolbar";
+	{
 		{
 			FToolMenuSection& Section = AssetsToolBar->AddSection("Content");
 
-			FToolMenuEntry CreateEntry = FToolMenuEntry::InitComboButton(
-				"Create",
+			FToolMenuEntry AddContentEntry = FToolMenuEntry::InitComboButton(
+				"AddContent",
 				FUIAction(),
-				FOnGetContent::CreateStatic(&FLevelEditorToolBar::GenerateCreateMenuWidget, InCommandList, TWeakPtr<SLevelEditor>(InLevelEditor)),
-				LOCTEXT("OpenCreate_Label", "Create"),
-				LOCTEXT("OpenCreate_Tooltip", "Quickly add to the project."),
-				FSlateIcon(FEditorStyle::GetStyleSetName(), "LevelEditor.OpenPlaceActors")
+				FOnGetContent::CreateStatic(&FLevelEditorToolBar::GenerateAddMenuWidget, InCommandList, TWeakPtr<SLevelEditor>(InLevelEditor)),
+				LOCTEXT("AddContent_Label", "Add"),
+				LOCTEXT("AddContent_Tooltip", "Quickly add to the project."),
+				FSlateIcon(FAppStyle::Get().GetStyleSetName(), "LevelEditor.OpenAddContent.Background", NAME_None, "LevelEditor.OpenAddContent.Overlay")
 			);
-			CreateEntry.StyleNameOverride = "CalloutToolbar";
-			Section.AddEntry(CreateEntry);
-
-			FToolMenuEntry ContentEntry = FToolMenuEntry::InitComboButton(
-				"Content",
-				FUIAction(),
-				FOnGetContent::CreateStatic(&FLevelEditorToolBar::GenerateContentMenuWidget, InCommandList, TWeakPtr<SLevelEditor>(InLevelEditor)),
-				LOCTEXT("OpenContent_Label", "Content"),
-				LOCTEXT("OpenContent_Tooltip", "Quickly find content available to add to the project."),
-				FSlateIcon(FEditorStyle::GetStyleSetName(), "LevelEditor.OpenContentBrowser")
-			);
-			ContentEntry.StyleNameOverride = "CalloutToolbar";
-			Section.AddEntry(ContentEntry);
+			AddContentEntry.StyleNameOverride = "AssetEditorToolbar";
+			Section.AddEntry(AddContentEntry);
 
 			FToolMenuEntry BlueprintEntry = FToolMenuEntry::InitComboButton(
 				"OpenBlueprint",
@@ -1224,9 +1365,9 @@ void FLevelEditorToolBar::RegisterLevelEditorToolBar( const TSharedRef<FUIComman
 				FOnGetContent::CreateStatic(&FLevelEditorToolBar::GenerateOpenBlueprintMenuContent, InCommandList, TWeakPtr<SLevelEditor>(InLevelEditor)),
 				LOCTEXT("OpenBlueprint_Label", "Blueprints"),
 				LOCTEXT("OpenBlueprint_ToolTip", "List of world Blueprints available to the user for editing or creation."),
-				FSlateIcon(FEditorStyle::GetStyleSetName(), "LevelEditor.OpenLevelBlueprint")
+				FSlateIcon(FAppStyle::Get().GetStyleSetName(), "LevelEditor.OpenLevelBlueprint")
 			);
-			BlueprintEntry.StyleNameOverride = "CalloutToolbar";
+			BlueprintEntry.StyleNameOverride = "AssetEditorToolbar";
 			Section.AddEntry(BlueprintEntry);
 
 			FToolMenuEntry CinematicsEntry = FToolMenuEntry::InitComboButton(
@@ -1237,7 +1378,7 @@ void FLevelEditorToolBar::RegisterLevelEditorToolBar( const TSharedRef<FUIComman
 				LOCTEXT("EditCinematics_Tooltip", "Displays a list of Level Sequence objects to open in their respective editors"),
 				FSlateIcon(FAppStyle::Get().GetStyleSetName(), "LevelEditor.OpenCinematic")
 			);
-			CinematicsEntry.StyleNameOverride = "CalloutToolbar";
+			CinematicsEntry.StyleNameOverride = "AssetEditorToolbar";
 			Section.AddEntry(CinematicsEntry);
 
 			Section.AddEntry(FToolMenuEntry::InitToolBarButton(FLevelEditorCommands::Get().ToggleVR, LOCTEXT("ToggleVR", "VR Mode")));
@@ -1245,64 +1386,6 @@ void FLevelEditorToolBar::RegisterLevelEditorToolBar( const TSharedRef<FUIComman
 
 	}
 
-	UToolMenu* ModesToolbar = UToolMenus::Get()->RegisterMenu("LevelEditor.LevelEditorToolBar.ModesToolBar", NAME_None, EMultiBoxType::SlimHorizontalToolBar);
-	ModesToolbar->StyleName = "AssetEditorToolbar";
-	{
-		TWeakPtr<SLevelEditor> LevelEditorPtr = InLevelEditor;
-		ModesToolbar->AddDynamicSection("EditorModes", FNewToolMenuDelegate::CreateLambda([LevelEditorPtr](UToolMenu* ToolMenu)
-		{
-			FToolMenuSection& Section = ToolMenu->AddSection("EditorModes");
-			Section.AddSeparator(NAME_None);
-			const FLevelEditorModule& LevelEditorModule = FModuleManager::LoadModuleChecked<FLevelEditorModule>("LevelEditor");
-			const FLevelEditorModesCommands& Commands = LevelEditorModule.GetLevelEditorModesCommands();
-
-			TArray<FEditorModeInfo, TInlineAllocator<1>> DefaultModes;
-
-			TArray<FEditorModeInfo, TInlineAllocator<10>> NonDefaultModes;
-
-			for (const FEditorModeInfo& Mode : GEditor->GetEditorSubsystem<UAssetEditorSubsystem>()->GetEditorModeInfoOrderedByPriority())
-			{
-				TSharedPtr<SLevelEditor> LevelEditorPin = LevelEditorPtr.Pin();
-				if (LevelEditorPin.IsValid() && LevelEditorPin->GetEditorModeManager().IsDefaultMode(Mode.ID))
-				{
-					DefaultModes.Add(Mode.ID);
-				}
-				else
-				{
-					NonDefaultModes.Add(Mode.ID);
-				}
-
-			}
-
-			auto BuildEditorModes =
-				[&Commands, &Section](TArrayView<FEditorModeInfo> Modes)
-			{
-				for (const FEditorModeInfo& Mode : Modes)
-				{
-					FName EditorModeCommandName = FName(*(FString("EditorMode.") + Mode.ID.ToString()));
-
-					TSharedPtr<FUICommandInfo> EditorModeCommand =
-						FInputBindingManager::Get().FindCommandInContext(Commands.GetContextName(), EditorModeCommandName);
-
-					// If a command isn't yet registered for this mode, we need to register one.
-					if (!EditorModeCommand.IsValid())
-					{
-						continue;
-					}
-
-					Section.AddEntry(FToolMenuEntry::InitToolBarButton(EditorModeCommand));
-				}
-
-			};
-
-			// Build Default Modes first
-			BuildEditorModes(MakeArrayView(DefaultModes));
-
-			// Build non-default modes second
-			BuildEditorModes(MakeArrayView(NonDefaultModes));
-		}));
-	}
-
 	UToolMenu* PlayToolBar = UToolMenus::Get()->RegisterMenu("LevelEditor.LevelEditorToolBar.PlayToolBar", NAME_None, EMultiBoxType::SlimHorizontalToolBar);
 	PlayToolBar->StyleName = "AssetEditorToolbar";
 	{
@@ -1441,37 +1524,14 @@ TSharedRef< SWidget > FLevelEditorToolBar::MakeLevelEditorToolBar( const TShared
 			SNew(SHorizontalBox)
 			.IsEnabled(FSlateApplication::Get().GetNormalExecutionAttribute())
 			+ SHorizontalBox::Slot()
-			.Padding(0.0f)
 			.AutoWidth()
 			[
-				UToolMenus::Get()->GenerateWidget("LevelEditor.LevelEditorToolBar.AssetsToolBar", MenuContext)
+				UToolMenus::Get()->GenerateWidget("LevelEditor.LevelEditorToolBar.ModesToolBar", MenuContext)
 			]
-			
 			+ SHorizontalBox::Slot()
 			.AutoWidth()
 			[
-				SNew(SSpacer)
-				.Size(FVector2D(120.f, 1.0f))
-			]
-			
-			+ SHorizontalBox::Slot()
-			.AutoWidth()
-			.HAlign(HAlign_Left)
-			[
-				SNew(SBox)
-				.MinDesiredWidth_Lambda(
-					[LevelEditorWeakPtr]()
-					{
-						if (TSharedPtr<SLevelEditor> LevelEditor = LevelEditorWeakPtr.Pin())
-						{
-							return LevelEditor->GetTickSpaceGeometry().GetLocalSize().X / 4.0f;
-						}
-						return 350.f;
-					}
-				)
-				[
-					UToolMenus::Get()->GenerateWidget("LevelEditor.LevelEditorToolBar.ModesToolBar", MenuContext)
-				]
+				UToolMenus::Get()->GenerateWidget("LevelEditor.LevelEditorToolBar.AssetsToolBar", MenuContext)
 			]
 			+ SHorizontalBox::Slot()
 			.HAlign(HAlign_Left)
@@ -1865,12 +1925,8 @@ void FLevelEditorToolBar::RegisterContentMenu()
 #define LOCTEXT_NAMESPACE "LevelToolBarViewMenu"
 	UToolMenu* ContentMenu = UToolMenus::Get()->RegisterMenu("LevelEditor.LevelEditorToolBar.ContentQuickMenu");
 	{
-		FToolMenuSection& Section = ContentMenu->FindOrAddSection("ProjectContent");
-		Section.InitSection("ProjectContent", LOCTEXT("ProjectContent_Label", "Project Content"), FToolMenuInsert("ContentBrowser", EToolMenuInsertType::First));
-	}
-	{
-		FToolMenuSection& Section = ContentMenu->FindOrAddSection("ExternalContent");
-		Section.InitSection("ExternalContent", LOCTEXT("ExternalContent_Label", "Get Content"), FToolMenuInsert());
+		FToolMenuSection& Section = ContentMenu->FindOrAddSection("Content");
+		Section.InitSection("ExternalContent", LOCTEXT("Content_Label", "Get Content"), FToolMenuInsert());
 		if (FLauncherPlatformModule::Get()->CanOpenLauncher(true))
 		{
 			Section.AddMenuEntry(FLevelEditorCommands::Get().OpenMarketplace);
@@ -1891,6 +1947,26 @@ void FLevelEditorToolBar::RegisterCreateMenu()
 #undef LOCTEXT_NAMESPACE
 }
 
+void FLevelEditorToolBar::RegisterAddMenu()
+{
+#define LOCTEXT_NAMESPACE "LevelToolBarViewMenu"
+	UToolMenu* AddMenu = UToolMenus::Get()->RegisterMenu("LevelEditor.LevelEditorToolBar.AddQuickMenu");
+	AddMenu->bSearchable = true;
+	{
+		FToolMenuSection& Section = AddMenu->FindOrAddSection("Content");
+
+		Section.InitSection("Content", LOCTEXT("Content_Label", "Get Content"), FToolMenuInsert(NAME_None, EToolMenuInsertType::First));
+
+		Section.AddMenuEntry(FLevelEditorCommands::Get().ImportContent);
+
+		if (FLauncherPlatformModule::Get()->CanOpenLauncher(true))
+		{
+			Section.AddMenuEntry(FLevelEditorCommands::Get().OpenMarketplace).InsertPosition = FToolMenuInsert("OpenBridgeTab", EToolMenuInsertType::After);
+		}
+	}
+#undef LOCTEXT_NAMESPACE
+}
+
 void FLevelEditorToolBar::OnOpenSubLevelBlueprint( ULevel* InLevel )
 {
 	ULevelScriptBlueprint* LevelScriptBlueprint = InLevel->GetLevelScriptBlueprint();
@@ -1937,6 +2013,16 @@ TSharedRef< SWidget > FLevelEditorToolBar::GenerateCreateMenuWidget(TSharedRef<F
 	return UToolMenus::Get()->GenerateWidget("LevelEditor.LevelEditorToolBar.CreateQuickMenu", MenuContext);
 }
 
+TSharedRef< SWidget > FLevelEditorToolBar::GenerateAddMenuWidget(TSharedRef<FUICommandList> InCommandList, TWeakPtr<SLevelEditor> InLevelEditor)
+{
+	FToolMenuContext MenuContext(InCommandList);
+	ULevelEditorMenuContext* LevelEditorMenuContext = NewObject<ULevelEditorMenuContext>();
+	LevelEditorMenuContext->LevelEditor = InLevelEditor;
+	MenuContext.AddObject(LevelEditorMenuContext);
+
+	return UToolMenus::Get()->GenerateWidget("LevelEditor.LevelEditorToolBar.AddQuickMenu", MenuContext);
+}
+
 void FLevelEditorToolBar::RegisterCinematicsMenu()
 {
 #define LOCTEXT_NAMESPACE "LevelToolBarCinematicsMenu"
diff --git a/Engine/Source/Editor/LevelEditor/Private/LevelEditorToolBar.h b/Engine/Source/Editor/LevelEditor/Private/LevelEditorToolBar.h
--- a/Engine/Source/Editor/LevelEditor/Private/LevelEditorToolBar.h
+++ b/Engine/Source/Editor/LevelEditor/Private/LevelEditorToolBar.h
@@ -68,6 +68,8 @@ protected:
 
 	static TSharedRef< SWidget > GenerateCreateMenuWidget(TSharedRef<FUICommandList> InCommandList, TWeakPtr<SLevelEditor> InLevelEditor);
 
+	static TSharedRef< SWidget > GenerateAddMenuWidget(TSharedRef<FUICommandList> InCommandList, TWeakPtr<SLevelEditor> InLevelEditor);
+
 	/**
 	 * Delegate for actor selection within the Cinematics popup menu's SceneOutliner.
 	 * Opens the editor for the selected actor and dismisses all popup menus.
@@ -89,4 +91,9 @@ private:
 	static void RegisterOpenBlueprintMenu();
 	static void RegisterContentMenu();
 	static void RegisterCreateMenu();
+	static void RegisterAddMenu();
+
+	static FText GetActiveModeName(TWeakPtr<SLevelEditor> LevelEditorPtr);
+	static const FSlateBrush* GetActiveModeIcon(TWeakPtr<SLevelEditor> LevelEditorPtr);
+
 };
diff --git a/Engine/Source/Editor/LevelEditor/Private/SLevelEditor.cpp b/Engine/Source/Editor/LevelEditor/Private/SLevelEditor.cpp
--- a/Engine/Source/Editor/LevelEditor/Private/SLevelEditor.cpp
+++ b/Engine/Source/Editor/LevelEditor/Private/SLevelEditor.cpp
@@ -138,6 +138,10 @@ void SLevelEditor::BindCommands()
 		Actions.OpenMarketplace,
 		FExecuteAction::CreateStatic( &FLevelEditorActionCallbacks::OpenMarketplace ) );
 
+	LevelEditorCommands->MapAction(
+		Actions.ImportContent,
+		FExecuteAction::CreateStatic(&FLevelEditorActionCallbacks::ImportContent));
+
 	LevelEditorCommands->MapAction(
 		Actions.ToggleVR,
 		FExecuteAction::CreateStatic( &FLevelEditorActionCallbacks::ToggleVR ),
diff --git a/Engine/Source/Editor/LevelEditor/Public/LevelEditorActions.h b/Engine/Source/Editor/LevelEditor/Public/LevelEditorActions.h
--- a/Engine/Source/Editor/LevelEditor/Public/LevelEditorActions.h
+++ b/Engine/Source/Editor/LevelEditor/Public/LevelEditorActions.h
@@ -554,6 +554,7 @@ public:
 	TSharedPtr< FUICommandInfo > OpenContentBrowser;
 	TSharedPtr< FUICommandInfo > OpenMarketplace;
 	TSharedPtr< FUICommandInfo > ToggleVR;
+	TSharedPtr< FUICommandInfo > ImportContent;
 
 	/**
 	 * Blueprints commands
@@ -1155,6 +1156,9 @@ public:
 	/** Open the Marketplace */
 	static void OpenMarketplace();
 
+	/** Import content into a chosen location*/
+	static void ImportContent();
+
 	/** Checks out the Project Settings config */
 	static void CheckOutProjectSettingsConfig();
 
diff --git a/Engine/Source/Editor/PlacementMode/Private/PlacementModeModule.cpp b/Engine/Source/Editor/PlacementMode/Private/PlacementModeModule.cpp
--- a/Engine/Source/Editor/PlacementMode/Private/PlacementModeModule.cpp
+++ b/Engine/Source/Editor/PlacementMode/Private/PlacementModeModule.cpp
@@ -297,9 +297,11 @@ void FPlacementModeModule::StartupModule()
 		};
 
 		
-		UToolMenu* ContentMenu = UToolMenus::Get()->ExtendMenu("LevelEditor.LevelEditorToolBar.CreateQuickMenu");
+		UToolMenu* ContentMenu = UToolMenus::Get()->ExtendMenu("LevelEditor.LevelEditorToolBar.AddQuickMenu");
 		FName CreateSectionName = TEXT("PMQCreateMenu");
-		FText CreateSectionDisplayName = NSLOCTEXT("PlacementMode", "PMQCreateMenu", "Create");
+		FText CreateSectionDisplayName = NSLOCTEXT("PlacementMode", "PMQCreateMenu", "Place Actors");
+
+		GenerateCategorySubMenu(ContentMenu, CreateSectionName, CreateSectionDisplayName, FBuiltInPlacementCategories::Basic());
 
 		// All Subcategories as submenus
 		FName CategoriesSectionName = TEXT("CreateAllCategories");
@@ -318,28 +320,24 @@ void FPlacementModeModule::StartupModule()
 					GenerateCategorySubMenu(InMenu, CreateSectionName, CreateSectionDisplayName, CategoryInfo.UniqueHandle);
 				}
 			}
-		));
+		), FToolMenuInsert(NAME_None, EToolMenuInsertType::First));
 
 
-		// Basics Section
-		GenerateQuickCreateSection(FBuiltInPlacementCategories::Basic(), ContentMenu);
-
-		// Recents Section, limit to 5 items, then put in an all Recents submenu
+		// Recents Section, limit to 5 items
 		const FName& RecentName = FBuiltInPlacementCategories::RecentlyPlaced();
+
 		ContentMenu->AddDynamicSection(RecentName,
-			FNewToolMenuDelegate::CreateLambda([this, GenerateQuickCreateSection, GenerateCategorySubMenu](UToolMenu* InMenu)
+			FNewToolMenuDelegate::CreateLambda([this, GenerateQuickCreateSection, GenerateCategorySubMenu, CreateSectionName, CreateSectionDisplayName](UToolMenu* InMenu)
 			{
 				const FName& RecentName = FBuiltInPlacementCategories::RecentlyPlaced();
 				FPlacementCategory* RecentCategory = Categories.Find(RecentName);
 				RefreshRecentlyPlaced();
 				GenerateQuickCreateSection(RecentName, InMenu, 5);
-				GenerateCategorySubMenu(InMenu, RecentName, RecentCategory->DisplayName, RecentName);
 			}
 		));
 
 		// Open Placement Browser Panel
-		FToolMenuSection& BrowserSection = ContentMenu->AddSection(TEXT("PlacementBrowserMenuSection"));
-		BrowserSection.AddSeparator(NAME_None);
+		FToolMenuSection& BrowserSection = ContentMenu->AddSection(TEXT("PlacementBrowserMenuSection"), FText::GetEmpty(), FToolMenuInsert(RecentName, EToolMenuInsertType::Before));
 
 		BrowserSection.AddMenuEntry(FLevelEditorCommands::Get().OpenPlaceActors);
 
diff --git a/Engine/Source/Editor/PlacementMode/Private/SPlacementModeTools.cpp b/Engine/Source/Editor/PlacementMode/Private/SPlacementModeTools.cpp
--- a/Engine/Source/Editor/PlacementMode/Private/SPlacementModeTools.cpp
+++ b/Engine/Source/Editor/PlacementMode/Private/SPlacementModeTools.cpp
@@ -390,7 +390,7 @@ void SPlacementAssetMenuEntry::Construct(const FArguments& InArgs, const TShared
 		.BorderImage( this, &SPlacementAssetMenuEntry::GetBorder )
 		.Cursor( EMouseCursor::GrabHand )
 		.ToolTip( AssetEntryToolTip )
-		.Padding(FMargin(12.f, 3.f, 5.f, 3.f))
+		.Padding(FMargin(27.f, 3.f, 5.f, 3.f))
 		[
 			SNew( SHorizontalBox )
 
@@ -412,6 +412,7 @@ void SPlacementAssetMenuEntry::Construct(const FArguments& InArgs, const TShared
 
 			+ SHorizontalBox::Slot()
 			.FillWidth(1.f)
+			.Padding(1.f, 0.f, 0.f, 0.f)
 			.VAlign(VAlign_Center)
 			.HAlign(HAlign_Left)
 			[
diff --git a/Engine/Source/Runtime/Slate/Private/Framework/MultiBox/SToolBarComboButtonBlock.cpp b/Engine/Source/Runtime/Slate/Private/Framework/MultiBox/SToolBarComboButtonBlock.cpp
--- a/Engine/Source/Runtime/Slate/Private/Framework/MultiBox/SToolBarComboButtonBlock.cpp
+++ b/Engine/Source/Runtime/Slate/Private/Framework/MultiBox/SToolBarComboButtonBlock.cpp
@@ -7,6 +7,7 @@
 #include "Widgets/Input/SComboButton.h"
 #include "Widgets/Layout/SBox.h"
 #include "Styling/ToolBarStyle.h"
+#include "Widgets/Images/SLayeredImage.h"
 
 
 FToolBarComboButtonBlock::FToolBarComboButtonBlock( const FUIAction& InAction, const FOnGetContent& InMenuContentGenerator, const TAttribute<FText>& InLabel, const TAttribute<FText>& InToolTip, const TAttribute<FSlateIcon>& InIcon, bool bInSimpleComboBox )
@@ -85,6 +86,13 @@ void SToolBarComboButtonBlock::BuildMultiBlockWidget(const ISlateStyle* StyleSet
 	TSharedRef<SWidget> IconWidget = SNullWidget::NullWidget;
 	if (!ToolBarComboButtonBlock->bSimpleComboBox)
 	{
+		TSharedRef<SLayeredImage> ActualIconWidget =
+			SNew(SLayeredImage)
+			.ColorAndOpacity(this, &SToolBarComboButtonBlock::GetIconForegroundColor)
+			.Image(this, &SToolBarComboButtonBlock::GetIconBrush);
+
+		ActualIconWidget->AddLayer(TAttribute<const FSlateBrush*>(this, &SToolBarComboButtonBlock::GetOverlayIconBrush));
+
 		if (MultiBox->GetType() == EMultiBoxType::SlimHorizontalToolBar)
 		{
 			const FVector2D IconSize = ToolBarStyle.IconSize;
@@ -94,17 +102,12 @@ void SToolBarComboButtonBlock::BuildMultiBlockWidget(const ISlateStyle* StyleSet
 				.WidthOverride(IconSize.X)
 				.HeightOverride(IconSize.Y)
 				[
-					SNew(SImage)
-					.ColorAndOpacity(FSlateColor::UseForeground())
-					.Image(this, &SToolBarComboButtonBlock::GetIconBrush)
+					ActualIconWidget
 				];
 		}
 		else
 		{
-			IconWidget = 
-				SNew(SImage)
-				.ColorAndOpacity(FSlateColor::UseForeground())
-				.Image(this, &SToolBarComboButtonBlock::GetIconBrush);
+			IconWidget = ActualIconWidget;
 		}
 
 		Label = ToolBarComboButtonBlock->Label;
@@ -277,6 +280,30 @@ EVisibility SToolBarComboButtonBlock::GetIconVisibility(bool bIsASmallIcon) cons
 	return ((bForceSmallIcons || FMultiBoxSettings::UseSmallToolBarIcons.Get()) ^ bIsASmallIcon) ? EVisibility::Collapsed : EVisibility::Visible;
 }
 
+FSlateColor SToolBarComboButtonBlock::GetIconForegroundColor() const
+{
+	// If any brush has a tint, don't assume it should be subdued
+	const FSlateBrush* Brush = GetIconBrush();
+	if (Brush && Brush->TintColor != FLinearColor::White)
+	{
+		return FLinearColor::White;
+	}
+
+	return FSlateColor::UseForeground();
+}
+
+const FSlateBrush* SToolBarComboButtonBlock::GetOverlayIconBrush() const
+{
+	const FSlateIcon& ActualIcon = Icon.Get();
+
+	if (ActualIcon.IsSet())
+	{
+		return ActualIcon.GetOverlayIcon();
+	}
+
+	return nullptr;
+}
+
 FSlateColor SToolBarComboButtonBlock::OnGetForegroundColor() const
 {
 	if (ComboButtonWidget->IsOpen())
diff --git a/Engine/Source/Runtime/Slate/Public/Framework/MultiBox/SToolBarComboButtonBlock.h b/Engine/Source/Runtime/Slate/Public/Framework/MultiBox/SToolBarComboButtonBlock.h
--- a/Engine/Source/Runtime/Slate/Public/Framework/MultiBox/SToolBarComboButtonBlock.h
+++ b/Engine/Source/Runtime/Slate/Public/Framework/MultiBox/SToolBarComboButtonBlock.h
@@ -164,6 +164,10 @@ private:
 	/** Called by Slate to determine whether icons/labels are visible */
 	EVisibility GetIconVisibility(bool bIsASmallIcon) const;
 
+	FSlateColor GetIconForegroundColor() const;
+
+	const FSlateBrush* GetOverlayIconBrush() const;
+
 	FSlateColor OnGetForegroundColor() const;
 
 private:
