Moved FocusAllViewportsToSelection to FEditorViewportCommands so that it doesn't trigger during PIE sessions

diff --git a/Engine/Source/Editor/LevelEditor/Private/LevelEditorActions.cpp b/Engine/Source/Editor/LevelEditor/Private/LevelEditorActions.cpp
--- a/Engine/Source/Editor/LevelEditor/Private/LevelEditorActions.cpp
+++ b/Engine/Source/Editor/LevelEditor/Private/LevelEditorActions.cpp
@@ -3548,8 +3548,6 @@ void FLevelEditorCommands::RegisterCommands()
 
 	UI_COMMAND( SelectActorsInLayers, "Select all actors in selected actor's layers", "Selects all actors belonging to the layers of the currently selected actors", EUserInterfaceActionType::Button, FInputChord( EModifierKey::Control, EKeys::L ) );
 
-	UI_COMMAND( FocusAllViewportsToSelection, "Focus Selected Actors in All Viewports", "Moves the camera in front of the selected actors in all open viewports", EUserInterfaceActionType::Button, FInputChord( EKeys::F, EModifierKey::Shift ) );
-
 	UI_COMMAND(MaterialQualityLevel_Low, "Low", "Sets material quality in the scene to low.", EUserInterfaceActionType::RadioButton, FInputChord());
 	UI_COMMAND(MaterialQualityLevel_Medium, "Medium", "Sets material quality in the scene to medium.", EUserInterfaceActionType::RadioButton, FInputChord());
 	UI_COMMAND(MaterialQualityLevel_High, "High", "Sets material quality in the scene to high.", EUserInterfaceActionType::RadioButton, FInputChord());
diff --git a/Engine/Source/Editor/LevelEditor/Private/SLevelEditor.cpp b/Engine/Source/Editor/LevelEditor/Private/SLevelEditor.cpp
--- a/Engine/Source/Editor/LevelEditor/Private/SLevelEditor.cpp
+++ b/Engine/Source/Editor/LevelEditor/Private/SLevelEditor.cpp
@@ -154,7 +154,7 @@ void SLevelEditor::BindCommands()
 		FExecuteAction::CreateStatic< TWeakPtr< SLevelEditor > >( &FLevelEditorActionCallbacks::OnShowWorldProperties, SharedThis( this ) ) );
 	
 	LevelEditorCommands->MapAction( 
-		Actions.FocusAllViewportsToSelection, 
+		FEditorViewportCommands::Get().FocusAllViewportsToSelection,
 		FExecuteAction::CreateStatic( &FLevelEditorActionCallbacks::ExecuteExecCommand, FString( TEXT("CAMERA ALIGN") ) )
 		);
 
diff --git a/Engine/Source/Editor/LevelEditor/Public/LevelEditorActions.h b/Engine/Source/Editor/LevelEditor/Public/LevelEditorActions.h
--- a/Engine/Source/Editor/LevelEditor/Public/LevelEditorActions.h
+++ b/Engine/Source/Editor/LevelEditor/Public/LevelEditorActions.h
@@ -619,9 +619,6 @@ public:
 	TSharedPtr< FUICommandInfo > ResetSelectedParticleSystem;
 	TSharedPtr< FUICommandInfo > SelectActorsInLayers;
 
-	TSharedPtr< FUICommandInfo > FocusAllViewportsToSelection;
-	TSharedPtr< FUICommandInfo > FocusViewportToSelection;
-
         // Open merge actor command
 	TSharedPtr< FUICommandInfo > OpenMergeActor;
 
diff --git a/Engine/Source/Editor/UnrealEd/Private/EditorViewportCommands.cpp b/Engine/Source/Editor/UnrealEd/Private/EditorViewportCommands.cpp
--- a/Engine/Source/Editor/UnrealEd/Private/EditorViewportCommands.cpp
+++ b/Engine/Source/Editor/UnrealEd/Private/EditorViewportCommands.cpp
@@ -153,6 +153,7 @@ void FEditorViewportCommands::RegisterCommands()
 #endif
 	UI_COMMAND( CycleTransformGizmos, "Cycle Between Translate, Rotate, and Scale", "Cycles the transform gizmos between translate, rotate, and scale", EUserInterfaceActionType::Button, FInputChord(EKeys::SpaceBar) );
 	
+	UI_COMMAND( FocusAllViewportsToSelection, "Focus Selected Actors in All Viewports", "Moves the camera in front of the selected actors in all open viewports", EUserInterfaceActionType::Button, FInputChord(EKeys::F, EModifierKey::Shift) );
 	UI_COMMAND( FocusViewportToSelection, "Focus Selected", "Moves the camera in front of the selection", EUserInterfaceActionType::Button, FInputChord( EKeys::F ) );
 
 	UI_COMMAND( LocationGridSnap, "Grid Snap", "Enables or disables snapping to the grid when dragging objects around", EUserInterfaceActionType::ToggleButton, FInputChord() );
diff --git a/Engine/Source/Editor/UnrealEd/Public/EditorViewportCommands.h b/Engine/Source/Editor/UnrealEd/Public/EditorViewportCommands.h
--- a/Engine/Source/Editor/UnrealEd/Public/EditorViewportCommands.h
+++ b/Engine/Source/Editor/UnrealEd/Public/EditorViewportCommands.h
@@ -195,6 +195,7 @@ public:
 	TSharedPtr< FUICommandInfo > CycleTransformGizmos;
 	TSharedPtr< FUICommandInfo > CycleTransformGizmoCoordSystem;
 
+	TSharedPtr< FUICommandInfo > FocusAllViewportsToSelection;
 	TSharedPtr< FUICommandInfo > FocusViewportToSelection;
 
 	/** Toggle automatic exposure */
