In-viewport details - reskin appearance and fix drag behavior

- Update in-viewport details panel to match latest mocks
- Fix drag-to-move behavior so that panel tracks mouse position properly
  - Originally, we measured the grab position's fraction within the header widget and then computed the decorator offset by multiplying by the entire panel's size
  - Now, we just measure the screen-space offset from the top-left corner. This makes the window position follow the drag properly and also behaves better when the window is dragged close to the right edge of the viewport.

diff --git a/Engine/Source/Editor/EditorStyle/Private/StarshipStyle.cpp b/Engine/Source/Editor/EditorStyle/Private/StarshipStyle.cpp
--- a/Engine/Source/Editor/EditorStyle/Private/StarshipStyle.cpp
+++ b/Engine/Source/Editor/EditorStyle/Private/StarshipStyle.cpp
@@ -3088,22 +3088,22 @@ void FStarshipEditorStyle::FStyle::SetupPropertyEditorStyles()
 			);
 
 		FWindowStyle InViewportDecoratorWindow = FStarshipCoreStyle::GetCoreStyle().GetWidgetStyle<FWindowStyle>("Window");
-		InViewportDecoratorWindow.SetCornerRadius(20);
+		InViewportDecoratorWindow.SetCornerRadius(4);
 
 		Set("InViewportDecoratorWindow", InViewportDecoratorWindow);
-		FLinearColor TransparentHeader = FStyleColors::Panel.GetSpecifiedColor();
-		TransparentHeader.A = 0.5f;
-		FLinearColor TransparentToolbar = FStyleColors::Header.GetSpecifiedColor();
-		TransparentToolbar.A = 0.5f;
 		FLinearColor TransparentBackground = FStyleColors::Background.GetSpecifiedColor();
-		TransparentBackground.A = 0.35f;
-		Set("PropertyTable.InViewport.Header", new FSlateRoundedBoxBrush(FSlateColor(TransparentHeader), FVector4(10.0f, 10.0f, 0.0f, 0.0f)));
-		Set("PropertyTable.InViewport.Background", new FSlateRoundedBoxBrush(FSlateColor(TransparentBackground), 10.0f));
+		TransparentBackground.A = 0.8f;
+		Set("PropertyTable.InViewport.Header", new FSlateRoundedBoxBrush(FStyleColors::Title, FVector4(4.0f, 4.0f, 0.0f, 0.0f)));
+		Set("PropertyTable.InViewport.Background", new FSlateRoundedBoxBrush(FSlateColor(TransparentBackground), 4.0f));
 		// InViewportToolbar
 		{
 			FToolBarStyle InViewportToolbar = FStarshipCoreStyle::GetCoreStyle().GetWidgetStyle<FToolBarStyle>("SlimToolBar");
-			InViewportToolbar.SetBackground(FSlateColorBrush(FSlateColor(TransparentToolbar)));
-			InViewportToolbar.SetBackgroundPadding(0.0f);
+			InViewportToolbar.SetBackground(FSlateColorBrush(FStyleColors::Panel));
+			InViewportToolbar.SetBackgroundPadding(FMargin(4.0f, 0.0f));
+			InViewportToolbar.SetButtonPadding(0.0f);
+			InViewportToolbar.SetIconSize(Icon16x16);
+			InViewportToolbar.ButtonStyle.SetNormalPadding(FMargin(4, 4, 4, 4));
+			InViewportToolbar.ButtonStyle.SetPressedPadding(FMargin(4, 5, 4, 3));
 			Set("InViewportToolbar", InViewportToolbar);
 		}
 		const FTableViewStyle InViewportViewStyle = FTableViewStyle()
diff --git a/Engine/Source/Editor/LevelEditor/Private/SInViewportDetails.cpp b/Engine/Source/Editor/LevelEditor/Private/SInViewportDetails.cpp
--- a/Engine/Source/Editor/LevelEditor/Private/SInViewportDetails.cpp
+++ b/Engine/Source/Editor/LevelEditor/Private/SInViewportDetails.cpp
@@ -63,25 +63,20 @@ void FInViewportUIDragOperation::OnDragged(const FDragDropEvent& DragDropEvent)
 	CursorDecoratorWindow->MoveWindowTo(TargetPosition);
 }
 
-TSharedRef<FInViewportUIDragOperation> FInViewportUIDragOperation::New(const TSharedRef<class SInViewportDetails>& InUIToBeDragged, const FVector2D InTabGrabOffset, const FVector2D& OwnerAreaSize)
+TSharedRef<FInViewportUIDragOperation> FInViewportUIDragOperation::New(const TSharedRef<class SInViewportDetails>& InUIToBeDragged, const FVector2D InDecoratorOffset, const FVector2D& OwnerAreaSize)
 {
-	const TSharedRef<FInViewportUIDragOperation> Operation = MakeShareable(new FInViewportUIDragOperation(InUIToBeDragged, InTabGrabOffset, OwnerAreaSize));
+	const TSharedRef<FInViewportUIDragOperation> Operation = MakeShareable(new FInViewportUIDragOperation(InUIToBeDragged, InDecoratorOffset, OwnerAreaSize));
 	return Operation;
 }
 
-FVector2D FInViewportUIDragOperation::GetTabGrabOffsetFraction() const
-{
-	return TabGrabOffsetFraction;
-}
-
 FInViewportUIDragOperation::~FInViewportUIDragOperation()
 {
 
 }
 
-FInViewportUIDragOperation::FInViewportUIDragOperation(const TSharedRef<class SInViewportDetails>& InUIToBeDragged, const FVector2D InTabGrabOffsetFraction, const FVector2D& OwnerAreaSize)
+FInViewportUIDragOperation::FInViewportUIDragOperation(const TSharedRef<class SInViewportDetails>& InUIToBeDragged, const FVector2D InDecoratorOffset, const FVector2D& OwnerAreaSize)
 	: UIBeingDragged(InUIToBeDragged)
-	, TabGrabOffsetFraction(InTabGrabOffsetFraction)
+	, DecoratorOffsetFromCursor(InDecoratorOffset)
 	, LastContentSize(OwnerAreaSize)
 {
 	// Create the decorator window that we will use during this drag and drop to make the user feel like
@@ -100,8 +95,7 @@ FInViewportUIDragOperation::FInViewportUIDragOperation(const TSharedRef<class SI
 
 const FVector2D FInViewportUIDragOperation::GetDecoratorOffsetFromCursor()
 {
-	const FVector2D TabDesiredSize = UIBeingDragged->GetDesiredSize();
-	return TabGrabOffsetFraction * TabDesiredSize;
+	return DecoratorOffsetFromCursor;
 }
 
 class SInViewportDetailsRow : public STableRow< TSharedPtr<IDetailTreeNode> >
@@ -141,46 +135,47 @@ public:
 				.PhysicalSplitterHandleSize(1.0f)
  				.HitDetectionSplitterHandleSize(5.0f)
  				+ SSplitter::Slot()
-				.SizeRule(SSplitter::ESizeRule::FractionOfParent)
  				.Value(ColumnSizeData.NameColumnWidth)
 				.OnSlotResized(ColumnSizeData.OnNameColumnResized)
  				[
 					SNew(SBox)
-					.HAlign(HAlign_Right)
- 					.Padding(5.0f)
+					.HAlign(HAlign_Left)
+ 					.Padding(2.0f)
+					.Clipping(EWidgetClipping::OnDemand)
  					[
 						NameWidget.ToSharedRef()
 					]
  				]
  				+ SSplitter::Slot()
-				.SizeRule(SSplitter::ESizeRule::FractionOfParent)
  				.Value(ColumnSizeData.ValueColumnWidth)
  				.OnSlotResized(ColumnSizeData.OnValueColumnResized)
  				[
  					SNew(SBox)
- 					.Padding(5.0f)
+ 					.Padding(2.0f)
  					[
  						ValueWidget.ToSharedRef()
  					]
 				]
 				+ SSplitter::Slot()
-					.SizeRule(SSplitter::ESizeRule::SizeToContent)
+				.Value(ColumnSizeData.RightColumnWidth)
+				.OnSlotResized(ColumnSizeData.OnRightColumnResized)
+				[
+					SNew(SHorizontalBox)
+					+ SHorizontalBox::Slot()
+					.VAlign(VAlign_Center)
+					.HAlign(HAlign_Center)
+					.Padding(2.0f)
 					[
-						SNew(SHorizontalBox)
-						+ SHorizontalBox::Slot()
-						.VAlign(VAlign_Center)
-						.HAlign(HAlign_Center)
-						.Padding(2.0f)
-						[
-							ResetWidget.ToSharedRef()
-						]
-					];
+						ResetWidget.ToSharedRef()
+					]
+				];
  		}
 
 		ChildSlot
 			[
 				SNew(SBox)
-				.MinDesiredWidth(350.0f)
+				.MinDesiredWidth(250.0f)
+				.Padding(FMargin(14.f, 0.f, 10.f, 0.f))
 				[
 					RowWidget.ToSharedRef()
 				]
@@ -206,7 +201,7 @@ void SInViewportDetails::Construct(const FArguments& InArgs)
 	FPropertyRowGeneratorArgs Args;
 	Args.DefaultsOnlyVisibility = EEditDefaultsOnlyNodeVisibility::Hide;
 	Args.NotifyHook = GUnrealEd;
-	ColumnSizeData.ValueColumnWidth = 0.5f;
+	ColumnSizeData.SetValueColumnWidth(0.5f);
 	PropertyRowGenerator = PropertyEditorModule.CreatePropertyRowGenerator(Args);
 
 	USelection::SelectionChangedEvent.AddRaw(this, &SInViewportDetails::OnEditorSelectionChanged);
@@ -234,7 +229,8 @@ void SInViewportDetails::GenerateWidget()
 		}
 		else
 		{
-			NameString = SelectedActor->GetHumanReadableName();
+			// Use the actor label because that's the friendly name used in other editor UI.
+			NameString = SelectedActor->GetActorLabel();
 		}
 
 		// Get the common base class of the selected objects
@@ -261,9 +257,9 @@ void SInViewportDetails::GenerateWidget()
 		ChildSlot
 			[
 				SNew(SBackgroundBlur)
-				.BlurStrength(20)
+				.BlurStrength(4)
 				.Padding(0.0f)
-				.CornerRadius(FVector4(15.0f, 15.0f, 15.0f, 15.0f))
+				.CornerRadius(FVector4(4.0f, 4.0f, 4.0f, 4.0f))
 				[
 					SNew(SBorder)
 					.BorderImage(FAppStyle::Get().GetBrush("PropertyTable.InViewport.Background"))
@@ -288,13 +284,15 @@ void SInViewportDetails::GenerateWidget()
 								[
 									SNew(SImage)
 									.Image(ActorIcon)
+									.ColorAndOpacity(FSlateColor::UseForeground())
 								]
 								+ SHorizontalBox::Slot()
-								.Padding(10.0f, 2.0f)
+								.VAlign(VAlign_Center)
+								.Padding(4.0f, 0.0f)
 								[
 									SNew(STextBlock)
 									.Text(FText::FromString(NameString))
-									.TextStyle(FAppStyle::Get(), "DetailsView.CategoryTextStyle")
+									.Font(FAppStyle::Get().GetFontStyle("PropertyWindow.NormalFont"))
 								]
 							]
 						]
@@ -307,7 +305,7 @@ void SInViewportDetails::GenerateWidget()
 						]
 						+ SVerticalBox::Slot()
 						.AutoHeight()
-						.Padding(0.0f)
+						.Padding(FMargin(0.0f, 5.0f))
 						[
 							MakeDetailsWidget()
 						]
@@ -457,13 +455,13 @@ TSharedRef<ITableRow> SInViewportDetails::GenerateListRow(TSharedPtr<IDetailTree
 
 }
 
-FReply SInViewportDetails::StartDraggingDetails(FVector2D InTabGrabOffsetFraction, const FPointerEvent& MouseEvent)
+FReply SInViewportDetails::StartDraggingDetails(FVector2D InTabGrabScreenSpaceOffset, const FPointerEvent& MouseEvent)
 {
 	// Start dragging.
 	TSharedRef<FInViewportUIDragOperation> DragDropOperation =
 		FInViewportUIDragOperation::New(
 			SharedThis(this),
-			InTabGrabOffsetFraction,
+			InTabGrabScreenSpaceOffset,
 			GetDesiredSize()
 		);
 	if (OwningViewport.IsValid())
@@ -480,7 +478,7 @@ void SInViewportDetailsHeader::Construct(const FArguments& InArgs)
 		[
 			SNew(SBorder)
 			.BorderImage(FAppStyle::Get().GetBrush("PropertyTable.InViewport.Header"))
-			.Padding(5.0f)
+			.Padding(FMargin(8.0f, 5.0f))
 			[
 				InArgs._Content.Widget
 			]
@@ -490,16 +488,12 @@ void SInViewportDetailsHeader::Construct(const FArguments& InArgs)
 FReply SInViewportDetailsHeader::OnDragDetected(const FGeometry& MyGeometry, const FPointerEvent& MouseEvent)
 {
 	// Need to remember where within a tab we grabbed
-	const FVector2D TabGrabOffset = MyGeometry.AbsoluteToLocal(MouseEvent.GetScreenSpacePosition());
-	const FVector2D TabSize = MyGeometry.GetLocalSize();
-	const FVector2D TabGrabOffsetFraction = FVector2D(
-		FMath::Clamp(TabGrabOffset.X / TabSize.X, 0.0f, 1.0f),
-		FMath::Clamp(TabGrabOffset.Y / TabSize.Y, 0.0f, 1.0f));
+	const FVector2D TabGrabScreenSpaceOffset = MouseEvent.GetScreenSpacePosition() - MyGeometry.GetAbsolutePosition();
 
 	TSharedPtr<SInViewportDetails> PinnedParent = ParentPtr.Pin();
 	if (PinnedParent.IsValid())
 	{
-		return PinnedParent->StartDraggingDetails(TabGrabOffsetFraction, MouseEvent);
+		return PinnedParent->StartDraggingDetails(TabGrabScreenSpaceOffset, MouseEvent);
 	}
 
 	return FReply::Unhandled();
diff --git a/Engine/Source/Editor/LevelEditor/Private/SInViewportDetails.h b/Engine/Source/Editor/LevelEditor/Private/SInViewportDetails.h
--- a/Engine/Source/Editor/LevelEditor/Private/SInViewportDetails.h
+++ b/Engine/Source/Editor/LevelEditor/Private/SInViewportDetails.h
@@ -53,15 +53,12 @@ public:
 	 * Create this Drag and Drop Content
 	 *
 	 * @param InUIToBeDragged	  The UI being dragged 
-	 * @param InTabGrabOffset     Where within the tab we grabbed, so we're not dragging by the upper left of the tab.
+	 * @param InDecoratorOffset   Where within the UI we grabbed, so we're not dragging by the upper left of the UI.
 	 * @param OwnerAreaSize       Size of the DockArea at the time when we start dragging.
 	 *
 	 * @return a new FDockingDragOperation
 	 */
-	static TSharedRef<FInViewportUIDragOperation> New(const TSharedRef<class SInViewportDetails>& InUIToBeDragged, const FVector2D InTabGrabOffset, const FVector2D& OwnerAreaSize);
-
-	/** @return location where the user grabbed within the tab as a fraction of the tab's size */
-	FVector2D GetTabGrabOffsetFraction() const;
+	static TSharedRef<FInViewportUIDragOperation> New(const TSharedRef<class SInViewportDetails>& InUIToBeDragged, const FVector2D InDecoratorOffset, const FVector2D& OwnerAreaSize);
 
 	virtual ~FInViewportUIDragOperation();
 
@@ -70,13 +67,13 @@ public:
 
 protected:
 	/** The constructor is protected, so that this class can only be instantiated as a shared pointer. */
-	FInViewportUIDragOperation(const TSharedRef<class SInViewportDetails>& InUIToBeDragged, const FVector2D InTabGrabOffsetFraction, const FVector2D& OwnerAreaSize);
+	FInViewportUIDragOperation(const TSharedRef<class SInViewportDetails>& InUIToBeDragged, const FVector2D InDecoratorOffsetFromCursor, const FVector2D& OwnerAreaSize);
 
 	/** What is actually being dragged in this operation */
 	TSharedPtr<class SInViewportDetails> UIBeingDragged;
 
-	/** Where the user grabbed the tab as a fraction of the tab's size */
-	FVector2D TabGrabOffsetFraction;
+	/** Where the user grabbed the UI measured in screen space from its top-left corner */
+	FVector2D DecoratorOffsetFromCursor;
 
 	/** Decorator widget where we add temp doc tabs to */
 	TSharedPtr<SDockingTabStack> CursorDecoratorStackNode;
@@ -115,7 +112,7 @@ public:
 	virtual void PostUndo(bool bSuccess) override;
 	virtual void PostRedo(bool bSuccess) override;
 
-	FReply StartDraggingDetails(FVector2D InTabGrabOffsetFraction, const FPointerEvent& MouseEvent);
+	FReply StartDraggingDetails(FVector2D InTabGrabScreenSpaceOffset, const FPointerEvent& MouseEvent);
 	FDetailColumnSizeData& GetColumnSizeData() { return ColumnSizeData; }
 	AActor* GetSelectedActorInEditor() const;
 	UToolMenu* GetGeneratedToolbarMenu() const;
