Fixed an issue causing the modes dropdown items to stop working if the layout of the editor was switched

diff --git a/Engine/Source/Editor/LevelEditor/Private/LevelEditorToolBar.cpp b/Engine/Source/Editor/LevelEditor/Private/LevelEditorToolBar.cpp
--- a/Engine/Source/Editor/LevelEditor/Private/LevelEditorToolBar.cpp
+++ b/Engine/Source/Editor/LevelEditor/Private/LevelEditorToolBar.cpp
@@ -1279,6 +1279,12 @@ void FLevelEditorToolBar::RegisterLevelEditorToolBar( const TSharedRef<FUIComman
 
 						FMenuBuilder MenuBuilder(true, LevelEditorModule.GetGlobalLevelEditorActions());
 
+						TSharedPtr<SLevelEditor> LevelEditorPin = LevelEditorPtr.Pin();
+						if (LevelEditorPin.IsValid())
+						{
+							MenuBuilder.PushCommandList(LevelEditorPin->GetLevelEditorActions().ToSharedRef());
+						}
+
 						MenuBuilder.BeginSection("EditorModes");
 
 						for (TSharedPtr<FUICommandInfo> Command : CommandInfos)
