Removed the "(Experimental)" tag on the LiveCoding UI button.

#rnx

diff --git a/Engine/Source/Editor/LevelEditor/Private/LevelEditorActions.cpp b/Engine/Source/Editor/LevelEditor/Private/LevelEditorActions.cpp
--- a/Engine/Source/Editor/LevelEditor/Private/LevelEditorActions.cpp
+++ b/Engine/Source/Editor/LevelEditor/Private/LevelEditorActions.cpp
@@ -3330,7 +3330,7 @@ void FLevelEditorCommands::RegisterCommands()
 	UI_COMMAND( RecompileGameCode, "Recompile Game Code", "Recompiles and reloads C++ code for game systems on the fly", EUserInterfaceActionType::Button, FInputChord( EKeys::P, EModifierKey::Alt | EModifierKey::Control | EModifierKey::Shift ) );
 
 #if WITH_LIVE_CODING
-	UI_COMMAND( LiveCoding_Enable, "Enable Live Coding (Experimental)", "Hot-patches C++ function changes into the current process. Currently does not allow class layout changes.", EUserInterfaceActionType::ToggleButton, FInputChord() );
+	UI_COMMAND( LiveCoding_Enable, "Enable Live Coding", "Hot-patches C++ function changes into the current process. Currently does not allow class layout changes.", EUserInterfaceActionType::ToggleButton, FInputChord() );
 	UI_COMMAND( LiveCoding_StartSession, "Start Session", "Starts a live coding session.", EUserInterfaceActionType::Button, FInputChord() );
 	UI_COMMAND( LiveCoding_ShowConsole, "Show Console", "Displays the live coding console window.", EUserInterfaceActionType::Button, FInputChord() );
 	UI_COMMAND( LiveCoding_Settings, "Settings...", "Open the live coding settings", EUserInterfaceActionType::Button, FInputChord() );
