Create a custom toolbar for user made tools in the level editor toolbar

diff --git a/Engine/Plugins/Media/MediaFrameworkUtilities/Source/MediaFrameworkUtilitiesEditor/Private/UI/MediaProfileMenuEntry.cpp b/Engine/Plugins/Media/MediaFrameworkUtilities/Source/MediaFrameworkUtilitiesEditor/Private/UI/MediaProfileMenuEntry.cpp
--- a/Engine/Plugins/Media/MediaFrameworkUtilities/Source/MediaFrameworkUtilitiesEditor/Private/UI/MediaProfileMenuEntry.cpp
+++ b/Engine/Plugins/Media/MediaFrameworkUtilities/Source/MediaFrameworkUtilitiesEditor/Private/UI/MediaProfileMenuEntry.cpp
@@ -43,7 +43,7 @@ struct FMediaProfileMenuEntryImpl
 			FCanExecuteAction::CreateLambda([] { return IMediaProfileManager::Get().GetCurrentMediaProfile() != nullptr; })
 		);
 
-		UToolMenu* Menu = UToolMenus::Get()->ExtendMenu("LevelEditor.LevelEditorToolBar.ModesToolBar");
+		UToolMenu* Menu = UToolMenus::Get()->ExtendMenu("LevelEditor.LevelEditorToolBar.User");
 		FToolMenuSection& Section = Menu->FindOrAddSection("MediaProfile");
 
 		auto ButtonTooltipLambda = [this]()
@@ -85,7 +85,7 @@ struct FMediaProfileMenuEntryImpl
 	{
 		if (!IsEngineExitRequested())
 		{
-			UToolMenus::Get()->RemoveSection("LevelEditor.LevelEditorToolBar.ModesToolBar", "MediaProfile");
+			UToolMenus::Get()->RemoveSection("LevelEditor.LevelEditorToolBar.User", "MediaProfile");
 		}
 	}
 
diff --git a/Engine/Plugins/VirtualProduction/CameraCalibration/Source/CameraCalibrationEditor/Private/UI/CameraCalibrationMenuEntry.cpp b/Engine/Plugins/VirtualProduction/CameraCalibration/Source/CameraCalibrationEditor/Private/UI/CameraCalibrationMenuEntry.cpp
--- a/Engine/Plugins/VirtualProduction/CameraCalibration/Source/CameraCalibrationEditor/Private/UI/CameraCalibrationMenuEntry.cpp
+++ b/Engine/Plugins/VirtualProduction/CameraCalibration/Source/CameraCalibrationEditor/Private/UI/CameraCalibrationMenuEntry.cpp
@@ -44,7 +44,7 @@ struct FCameraCalibrationMenuEntryImpl
 			})
 		);
 
-		UToolMenu* Menu = UToolMenus::Get()->ExtendMenu("LevelEditor.LevelEditorToolBar.ModesToolBar");
+		UToolMenu* Menu = UToolMenus::Get()->ExtendMenu("LevelEditor.LevelEditorToolBar.User");
 		FToolMenuSection& Section = Menu->FindOrAddSection("LensFile");
 
 		auto ButtonTooltipLambda = [this]()
@@ -85,7 +85,7 @@ struct FCameraCalibrationMenuEntryImpl
 	{
 		if (!IsEngineExitRequested())
 		{
-			UToolMenus::Get()->RemoveSection("LevelEditor.LevelEditorToolBar.ModesToolBar", "LensFile");
+			UToolMenus::Get()->RemoveSection("LevelEditor.LevelEditorToolBar.User", "LensFile");
 		}
 	}
 
diff --git a/Engine/Source/Editor/LevelEditor/Private/LevelEditorToolBar.cpp b/Engine/Source/Editor/LevelEditor/Private/LevelEditorToolBar.cpp
--- a/Engine/Source/Editor/LevelEditor/Private/LevelEditorToolBar.cpp
+++ b/Engine/Source/Editor/LevelEditor/Private/LevelEditorToolBar.cpp
@@ -1481,6 +1481,9 @@ void FLevelEditorToolBar::RegisterLevelEditorToolBar( const TSharedRef<FUIComman
 
 	}
 
+	UToolMenu* UserToolbar = UToolMenus::Get()->RegisterMenu("LevelEditor.LevelEditorToolBar.User", NAME_None, EMultiBoxType::SlimHorizontalToolBar);
+	UserToolbar->StyleName = "AssetEditorToolbar";
+
 	UToolMenu* SettingsToolbar = UToolMenus::Get()->RegisterMenu("LevelEditor.LevelEditorToolBar.SettingsToolbar", NAME_None, EMultiBoxType::SlimHorizontalToolBar);
 	SettingsToolbar->StyleName = "AssetEditorToolbar";
 	{
@@ -1541,10 +1544,16 @@ TSharedRef< SWidget > FLevelEditorToolBar::MakeLevelEditorToolBar( const TShared
 			]
 			+ SHorizontalBox::Slot()
 			.HAlign(HAlign_Left)
+			.AutoWidth()
 			[
 				UToolMenus::Get()->GenerateWidget("LevelEditor.LevelEditorToolBar.PlayToolBar", MenuContext)
 			]
 			+ SHorizontalBox::Slot()
+			.HAlign(HAlign_Left)
+			[
+				UToolMenus::Get()->GenerateWidget("LevelEditor.LevelEditorToolBar.User", MenuContext)
+			]
+			+ SHorizontalBox::Slot()
 			.HAlign(HAlign_Right)
 			.Padding(0.0f, 0.0f, 7.0f, 0.0f)
 			[
