Fixed an issue causing some options in the Add Content menu in the level editor toolbar to not appear

diff --git a/Engine/Plugins/Bridge/Source/Bridge/Private/UI/BridgeUIManager.cpp b/Engine/Plugins/Bridge/Source/Bridge/Private/UI/BridgeUIManager.cpp
--- a/Engine/Plugins/Bridge/Source/Bridge/Private/UI/BridgeUIManager.cpp
+++ b/Engine/Plugins/Bridge/Source/Bridge/Private/UI/BridgeUIManager.cpp
@@ -69,8 +69,8 @@ void FBridgeUIManagerImpl::SetupMenuItem()
 	// }
 
 	// Adding Bridge entry to Quick Content menu.
-	UToolMenu* ContentMenu = UToolMenus::Get()->ExtendMenu("LevelEditor.LevelEditorToolBar.ContentQuickMenu");
-	FToolMenuSection& Section = ContentMenu->FindOrAddSection("ExternalContent");
+	UToolMenu* AddMenu = UToolMenus::Get()->ExtendMenu("LevelEditor.LevelEditorToolBar.AddQuickMenu");
+	FToolMenuSection& Section = AddMenu->FindOrAddSection("Content");
 	Section.AddMenuEntry("OpenBridgeTab",
 		LOCTEXT("OpenBridgeTab_Label", "Quixel Bridge"),
 		LOCTEXT("OpenBridgeTab_Desc", "Opens the Quixel Bridge."),
diff --git a/Engine/Source/Editor/LevelEditor/Private/LevelEditorToolBar.cpp b/Engine/Source/Editor/LevelEditor/Private/LevelEditorToolBar.cpp
--- a/Engine/Source/Editor/LevelEditor/Private/LevelEditorToolBar.cpp
+++ b/Engine/Source/Editor/LevelEditor/Private/LevelEditorToolBar.cpp
@@ -1202,8 +1202,6 @@ void FLevelEditorToolBar::RegisterLevelEditorToolBar( const TSharedRef<FUIComman
 
 	RegisterQuickSettingsMenu();
 	RegisterOpenBlueprintMenu();
-	RegisterContentMenu();
-	RegisterCreateMenu();
 	RegisterAddMenu();
 
 #define LOCTEXT_NAMESPACE "LevelEditorToolBar"
@@ -1935,33 +1933,6 @@ void FLevelEditorToolBar::RegisterOpenBlueprintMenu()
 #undef LOCTEXT_NAMESPACE
 }
 
-void FLevelEditorToolBar::RegisterContentMenu()
-{
-#define LOCTEXT_NAMESPACE "LevelToolBarViewMenu"
-	UToolMenu* ContentMenu = UToolMenus::Get()->RegisterMenu("LevelEditor.LevelEditorToolBar.ContentQuickMenu");
-	{
-		FToolMenuSection& Section = ContentMenu->FindOrAddSection("Content");
-		Section.InitSection("ExternalContent", LOCTEXT("Content_Label", "Get Content"), FToolMenuInsert());
-		if (FLauncherPlatformModule::Get()->CanOpenLauncher(true))
-		{
-			Section.AddMenuEntry(FLevelEditorCommands::Get().OpenMarketplace);
-		}
-	}
-#undef LOCTEXT_NAMESPACE
-}
-
-void FLevelEditorToolBar::RegisterCreateMenu()
-{
-#define LOCTEXT_NAMESPACE "LevelToolBarViewMenu"
-	UToolMenu* ContentMenu = UToolMenus::Get()->RegisterMenu("LevelEditor.LevelEditorToolBar.CreateQuickMenu");
-	ContentMenu->bSearchable = true;
-	{
-		FToolMenuSection& Section = ContentMenu->FindOrAddSection("Create");
-		Section.InitSection("ProjectContent", LOCTEXT("Create_Label", "Create"), FToolMenuInsert());
-	}
-#undef LOCTEXT_NAMESPACE
-}
-
 void FLevelEditorToolBar::RegisterAddMenu()
 {
 #define LOCTEXT_NAMESPACE "LevelToolBarViewMenu"
@@ -2007,27 +1978,6 @@ TSharedRef< SWidget > FLevelEditorToolBar::GenerateCinematicsMenuContent(TShared
 
 	return UToolMenus::Get()->GenerateWidget("LevelEditor.LevelEditorToolBar.Cinematics", MenuContext);
 }
-
-TSharedRef< SWidget > FLevelEditorToolBar::GenerateContentMenuWidget(TSharedRef<FUICommandList> InCommandList, TWeakPtr<SLevelEditor> InLevelEditor)
-{
-	FToolMenuContext MenuContext(InCommandList);
-	ULevelEditorMenuContext* LevelEditorMenuContext = NewObject<ULevelEditorMenuContext>();
-	LevelEditorMenuContext->LevelEditor = InLevelEditor;
-	MenuContext.AddObject(LevelEditorMenuContext);
-
-	return UToolMenus::Get()->GenerateWidget("LevelEditor.LevelEditorToolBar.ContentQuickMenu", MenuContext);
-}
-
-TSharedRef< SWidget > FLevelEditorToolBar::GenerateCreateMenuWidget(TSharedRef<FUICommandList> InCommandList, TWeakPtr<SLevelEditor> InLevelEditor)
-{
-	FToolMenuContext MenuContext(InCommandList);
-	ULevelEditorMenuContext* LevelEditorMenuContext = NewObject<ULevelEditorMenuContext>();
-	LevelEditorMenuContext->LevelEditor = InLevelEditor;
-	MenuContext.AddObject(LevelEditorMenuContext);
-
-	return UToolMenus::Get()->GenerateWidget("LevelEditor.LevelEditorToolBar.CreateQuickMenu", MenuContext);
-}
-
 TSharedRef< SWidget > FLevelEditorToolBar::GenerateAddMenuWidget(TSharedRef<FUICommandList> InCommandList, TWeakPtr<SLevelEditor> InLevelEditor)
 {
 	FToolMenuContext MenuContext(InCommandList);
diff --git a/Engine/Source/Editor/LevelEditor/Private/LevelEditorToolBar.h b/Engine/Source/Editor/LevelEditor/Private/LevelEditorToolBar.h
--- a/Engine/Source/Editor/LevelEditor/Private/LevelEditorToolBar.h
+++ b/Engine/Source/Editor/LevelEditor/Private/LevelEditorToolBar.h
@@ -64,10 +64,6 @@ protected:
 	 */
 	static TSharedRef< SWidget > GenerateCinematicsMenuContent(TSharedRef<FUICommandList> InCommandList, TWeakPtr<SLevelEditor> InLevelEditor);
 
-	static TSharedRef< SWidget > GenerateContentMenuWidget(TSharedRef<FUICommandList> InCommandList, TWeakPtr<SLevelEditor> InLevelEditor);
-
-	static TSharedRef< SWidget > GenerateCreateMenuWidget(TSharedRef<FUICommandList> InCommandList, TWeakPtr<SLevelEditor> InLevelEditor);
-
 	static TSharedRef< SWidget > GenerateAddMenuWidget(TSharedRef<FUICommandList> InCommandList, TWeakPtr<SLevelEditor> InLevelEditor);
 
 	/**
@@ -89,8 +85,6 @@ private:
 
 	static void RegisterQuickSettingsMenu();
 	static void RegisterOpenBlueprintMenu();
-	static void RegisterContentMenu();
-	static void RegisterCreateMenu();
 	static void RegisterAddMenu();
 
 	static FText GetActiveModeName(TWeakPtr<SLevelEditor> LevelEditorPtr);
