Move selection-related commands from Level Editor's actor context menu to a Select top-level menu

The old context menu previously only showed relevant commands, but the new top-level menu always shows all commands,
disabling irrelevant ones. I've moved that logic into the commands' CanExecute callback in order to have them
properly enable/disable when available/unavailable.

diff --git a/Engine/Plugins/Experimental/GeometryCollectionPlugin/Source/GeometryCollectionEditor/GeometryCollectionEditor.Build.cs b/Engine/Plugins/Experimental/GeometryCollectionPlugin/Source/GeometryCollectionEditor/GeometryCollectionEditor.Build.cs
--- a/Engine/Plugins/Experimental/GeometryCollectionPlugin/Source/GeometryCollectionEditor/GeometryCollectionEditor.Build.cs
+++ b/Engine/Plugins/Experimental/GeometryCollectionPlugin/Source/GeometryCollectionEditor/GeometryCollectionEditor.Build.cs
@@ -31,7 +31,9 @@ namespace UnrealBuildTool.Rules
 					"ToolMenus",
 					"Chaos",
 					"MeshDescription",
-					"StaticMeshDescription"
+					"StaticMeshDescription",
+					"LevelEditor",
+					"InputCore",
 				}
 			);
 
diff --git a/Engine/Plugins/Experimental/GeometryCollectionPlugin/Source/GeometryCollectionEditor/Private/GeometryCollection/GeometryCollectionEditorPlugin.cpp b/Engine/Plugins/Experimental/GeometryCollectionPlugin/Source/GeometryCollectionEditor/Private/GeometryCollection/GeometryCollectionEditorPlugin.cpp
--- a/Engine/Plugins/Experimental/GeometryCollectionPlugin/Source/GeometryCollectionEditor/Private/GeometryCollection/GeometryCollectionEditorPlugin.cpp
+++ b/Engine/Plugins/Experimental/GeometryCollectionPlugin/Source/GeometryCollectionEditor/Private/GeometryCollection/GeometryCollectionEditorPlugin.cpp
@@ -5,6 +5,7 @@
 
 #include "AssetToolsModule.h"
 #include "CoreMinimal.h"
+#include "LevelEditor.h"
 #include "Modules/ModuleManager.h"
 #include "PropertyEditorModule.h"
 #include "GeometryCollection/GeometryCollectionComponent.h"
@@ -16,6 +17,7 @@
 #include "GeometryCollection/GeometryCollectionConversion.h"
 #include "GeometryCollection/GeometryCollectionThumbnailRenderer.h"
 #include "GeometryCollection/GeometryCollectionSelectRigidBodyEdMode.h"
+#include "GeometryCollection/GeometryCollectionSelectionCommands.h"
 #include "HAL/ConsoleManager.h"
 #include "Features/IModularFeatures.h"
 #include "GeometryCollection/DetailCustomizations/GeomComponentCacheCustomization.h"
@@ -25,6 +27,7 @@
 #include "Styling/SlateStyle.h"
 #include "Styling/SlateStyleRegistry.h"
 #include "Styling/CoreStyle.h"
+#include "ToolMenus.h"
 
 IMPLEMENT_MODULE( IGeometryCollectionEditorPlugin, GeometryCollectionEditor )
 
@@ -35,6 +38,8 @@ void IGeometryCollectionEditorPlugin::StartupModule()
 {
 	FGeometryCollectionEditorStyle::Get();
 
+	FGeometryCollectionSelectionCommands::Register();
+
 	FAssetToolsModule& AssetToolsModule = FAssetToolsModule::GetModule();
 	IAssetTools& AssetTools = AssetToolsModule.Get();
 	GeometryCollectionAssetActions = new FAssetTypeActions_GeometryCollection();
@@ -187,6 +192,28 @@ void IGeometryCollectionEditorPlugin::StartupModule()
 		));
 	}
 
+	// Map global level editor commands
+	FLevelEditorModule& LevelEditorModule = FModuleManager::Get().LoadModuleChecked<FLevelEditorModule>("LevelEditor");
+	TSharedRef<FUICommandList> CommandList = LevelEditorModule.GetGlobalLevelEditorActions();
+	CommandList->MapAction(
+		FGeometryCollectionSelectionCommands::Get().SelectAllGeometry,
+		FExecuteAction::CreateLambda([]()
+		{
+			FGeometryCollectionCommands::SelectAllGeometry(TArray<FString>(), nullptr);
+		}));
+	CommandList->MapAction(
+		FGeometryCollectionSelectionCommands::Get().SelectNone,
+		FExecuteAction::CreateLambda([]()
+		{
+			FGeometryCollectionCommands::SelectNone(TArray<FString>(), nullptr);
+		}));
+	CommandList->MapAction(
+		FGeometryCollectionSelectionCommands::Get().SelectInverseGeometry,
+		FExecuteAction::CreateLambda([]()
+		{
+			FGeometryCollectionCommands::SelectInverseGeometry(TArray<FString>(), nullptr);
+		}));
+
 	// Bind our scene outliner provider to the editor
 	IModularFeatures& ModularFeatures = IModularFeatures::Get();
 	ModularFeatures.RegisterModularFeature(ITargetCacheProvider::GetFeatureName(), &TargetCacheProvider);
@@ -229,6 +256,9 @@ void IGeometryCollectionEditorPlugin::StartupModule()
 
 	UThumbnailManager::Get().RegisterCustomRenderer(UGeometryCollection::StaticClass(), UGeometryCollectionThumbnailRenderer::StaticClass());
 
+	// Register tool menus - delayed until ToolMenus module is loaded
+	UToolMenus::RegisterStartupCallback(
+		FSimpleMulticastDelegate::FDelegate::CreateRaw(this, &IGeometryCollectionEditorPlugin::RegisterMenus));
 }
 
 
@@ -263,6 +293,20 @@ void IGeometryCollectionEditorPlugin::ShutdownModule()
 
 		// Unregister rigid body selection editor mode
 		FEditorModeRegistry::Get().UnregisterMode(FGeometryCollectionSelectRigidBodyEdMode::EditorModeID);
+		
+		// Unregister tool menu callback and remove any entries we added
+		UToolMenus::UnRegisterStartupCallback(this);
+		UToolMenus::UnregisterOwner(this);
+
+		// Unmap global level editor commands
+		FLevelEditorModule& LevelEditorModule = FModuleManager::Get().LoadModuleChecked<FLevelEditorModule>("LevelEditor");
+		TSharedRef<FUICommandList> CommandList = LevelEditorModule.GetGlobalLevelEditorActions();
+		CommandList->UnmapAction(FGeometryCollectionSelectionCommands::Get().SelectAllGeometry);
+		CommandList->UnmapAction(FGeometryCollectionSelectionCommands::Get().SelectNone);
+		CommandList->UnmapAction(FGeometryCollectionSelectionCommands::Get().SelectInverseGeometry);
+
+		// Unregister commands
+		FGeometryCollectionSelectionCommands::Unregister();
 	}
 }
 
@@ -275,3 +319,16 @@ const ISlateStyle* IGeometryCollectionEditorPlugin::GetEditorStyle()
 {
 	return FSlateStyleRegistry::FindSlateStyle(GetEditorStyleName());
 }
+
+void IGeometryCollectionEditorPlugin::RegisterMenus()
+{
+	FToolMenuOwnerScoped OwnerScoped(this);
+
+	UToolMenu* Menu = UToolMenus::Get()->ExtendMenu("LevelEditor.MainMenu.Select");
+	FToolMenuSection& Section = Menu->AddSection("SelectBones", NSLOCTEXT("LevelViewportContextMenu", "GeometryCollectionHeading", "Geometry Collection"));
+	{
+		Section.AddMenuEntry(FGeometryCollectionSelectionCommands::Get().SelectAllGeometry);
+		Section.AddMenuEntry(FGeometryCollectionSelectionCommands::Get().SelectNone);
+		Section.AddMenuEntry(FGeometryCollectionSelectionCommands::Get().SelectInverseGeometry);
+	}
+}
\ No newline at end of file
diff --git a/Engine/Plugins/Experimental/GeometryCollectionPlugin/Source/GeometryCollectionEditor/Private/GeometryCollection/GeometryCollectionSelectionCommands.cpp b/Engine/Plugins/Experimental/GeometryCollectionPlugin/Source/GeometryCollectionEditor/Private/GeometryCollection/GeometryCollectionSelectionCommands.cpp
new file mode 100644
--- /dev/null
+++ b/Engine/Plugins/Experimental/GeometryCollectionPlugin/Source/GeometryCollectionEditor/Private/GeometryCollection/GeometryCollectionSelectionCommands.cpp
@@ -0,0 +1,12 @@
+// Copyright Epic Games, Inc. All Rights Reserved.
+
+#include "GeometryCollection/GeometryCollectionSelectionCommands.h"
+
+#define LOCTEXT_NAMESPACE "GeometryCollectionSelectionCommands"
+
+void FGeometryCollectionSelectionCommands::RegisterCommands()
+{
+	UI_COMMAND(SelectAllGeometry, "Select All Geometry In Hierarchy", "Select all geometry in hierarchy", EUserInterfaceActionType::Button, FInputChord());
+	UI_COMMAND(SelectNone, "Deselect All Geometry In Hierarchy", "Deselect all geometry in hierarchy", EUserInterfaceActionType::Button, FInputChord());
+	UI_COMMAND(SelectInverseGeometry, "Select Inverse Geometry In Hierarchy", "Select inverse geometry in hierarchy", EUserInterfaceActionType::Button, FInputChord());
+}
diff --git a/Engine/Plugins/Experimental/GeometryCollectionPlugin/Source/GeometryCollectionEditor/Public/GeometryCollection/GeometryCollectionEditorPlugin.h b/Engine/Plugins/Experimental/GeometryCollectionPlugin/Source/GeometryCollectionEditor/Public/GeometryCollection/GeometryCollectionEditorPlugin.h
--- a/Engine/Plugins/Experimental/GeometryCollectionPlugin/Source/GeometryCollectionEditor/Public/GeometryCollection/GeometryCollectionEditorPlugin.h
+++ b/Engine/Plugins/Experimental/GeometryCollectionPlugin/Source/GeometryCollectionEditor/Public/GeometryCollection/GeometryCollectionEditorPlugin.h
@@ -50,6 +50,9 @@ public:
 	static FName GetEditorStyleName();
 	static const ISlateStyle* GetEditorStyle();
 
+private:
+	void RegisterMenus();
+
 private:
 	TArray<IConsoleObject*> EditorCommands;
 
diff --git a/Engine/Plugins/Experimental/GeometryCollectionPlugin/Source/GeometryCollectionEditor/Public/GeometryCollection/GeometryCollectionSelectionCommands.h b/Engine/Plugins/Experimental/GeometryCollectionPlugin/Source/GeometryCollectionEditor/Public/GeometryCollection/GeometryCollectionSelectionCommands.h
new file mode 100644
--- /dev/null
+++ b/Engine/Plugins/Experimental/GeometryCollectionPlugin/Source/GeometryCollectionEditor/Public/GeometryCollection/GeometryCollectionSelectionCommands.h
@@ -0,0 +1,33 @@
+// Copyright Epic Games, Inc. All Rights Reserved.
+
+#pragma once
+
+#include "CoreMinimal.h"
+#include "EditorStyleSet.h"
+#include "Framework/Commands/Commands.h"
+
+class GEOMETRYCOLLECTIONEDITOR_API FGeometryCollectionSelectionCommands : public TCommands<FGeometryCollectionSelectionCommands>
+{
+public:
+	FGeometryCollectionSelectionCommands()
+		: TCommands<FGeometryCollectionSelectionCommands>
+		(
+			TEXT("GeometryCollectionSelection"),
+			NSLOCTEXT("Contexts", "GeometryCollectionSelection", "Geometry Collection Selection Commands"),
+			NAME_None,
+			FEditorStyle::GetStyleSetName()
+		)
+	{
+	}
+
+	virtual void RegisterCommands() override;
+
+	/** Selects all Geometry Collection geometry */
+	TSharedPtr<FUICommandInfo> SelectAllGeometry;
+
+	/** Deselects all Geometry Collection geometry */
+	TSharedPtr<FUICommandInfo> SelectNone;
+
+	/** Selects inverse of currently seleted Geometry Collection geometry */
+	TSharedPtr<FUICommandInfo> SelectInverseGeometry;
+};
\ No newline at end of file
diff --git a/Engine/Source/Editor/LevelEditor/Private/LevelEditor.cpp b/Engine/Source/Editor/LevelEditor/Private/LevelEditor.cpp
--- a/Engine/Source/Editor/LevelEditor/Private/LevelEditor.cpp
+++ b/Engine/Source/Editor/LevelEditor/Private/LevelEditor.cpp
@@ -1073,22 +1073,26 @@ void FLevelEditorModule::BindGlobalLevelEditorCommands()
 
 	ActionList.MapAction(
 		Commands.SelectImmediateChildren,
-		FExecuteAction::CreateStatic( &FLevelEditorActionCallbacks::ExecuteExecCommand, FString( TEXT("ACTOR SELECT ALL CHILDREN") ) )
+		FExecuteAction::CreateStatic( &FLevelEditorActionCallbacks::ExecuteExecCommand, FString( TEXT("ACTOR SELECT ALL CHILDREN") ) ),
+		FCanExecuteAction::CreateStatic( &FLevelEditorActionCallbacks::ActorSelected_CanExecute )
 	);
 
 	ActionList.MapAction(
 		Commands.SelectAllDescendants,
-		FExecuteAction::CreateStatic( &FLevelEditorActionCallbacks::ExecuteExecCommand, FString( TEXT("ACTOR SELECT ALL DESCENDANTS") ) )
+		FExecuteAction::CreateStatic( &FLevelEditorActionCallbacks::ExecuteExecCommand, FString( TEXT("ACTOR SELECT ALL DESCENDANTS") ) ),
+		FCanExecuteAction::CreateStatic( &FLevelEditorActionCallbacks::ActorSelected_CanExecute )
 	);
 
 	ActionList.MapAction(
 		Commands.SelectAllActorsOfSameClass,
-		FExecuteAction::CreateStatic( &FLevelEditorActionCallbacks::OnSelectAllActorsOfClass, false )
+		FExecuteAction::CreateStatic( &FLevelEditorActionCallbacks::OnSelectAllActorsOfClass, false ),
+		FCanExecuteAction::CreateStatic( &FLevelEditorActionCallbacks::CanSelectAllActorsOfClass )
 		);
 
 	ActionList.MapAction(
 		Commands.SelectAllActorsOfSameClassWithArchetype,
-		FExecuteAction::CreateStatic( &FLevelEditorActionCallbacks::OnSelectAllActorsOfClass, true )
+		FExecuteAction::CreateStatic( &FLevelEditorActionCallbacks::OnSelectAllActorsOfClass, true ),
+		FCanExecuteAction::CreateStatic( &FLevelEditorActionCallbacks::CanSelectAllActorsOfClass )
 		);
 
 	ActionList.MapAction(
@@ -1099,42 +1103,51 @@ void FLevelEditorModule::BindGlobalLevelEditorCommands()
 
 	ActionList.MapAction(
 		Commands.SelectRelevantLights,
-		FExecuteAction::CreateStatic( &FLevelEditorActionCallbacks::ExecuteExecCommand, FString( TEXT("ACTOR SELECT RELEVANTLIGHTS") ) )
+		FExecuteAction::CreateStatic( &FLevelEditorActionCallbacks::ExecuteExecCommand, FString( TEXT("ACTOR SELECT RELEVANTLIGHTS") ) ),
+		FCanExecuteAction::CreateStatic( &FLevelEditorActionCallbacks::ActorSelected_CanExecute )
 		);
 
 	ActionList.MapAction(
 		Commands.SelectStaticMeshesOfSameClass,
-		FExecuteAction::CreateStatic( &FLevelEditorActionCallbacks::ExecuteExecCommand, FString( TEXT("ACTOR SELECT MATCHINGSTATICMESH") ) )
+		FExecuteAction::CreateStatic( &FLevelEditorActionCallbacks::ExecuteExecCommand, FString( TEXT("ACTOR SELECT MATCHINGSTATICMESH") ) ),
+		FCanExecuteAction::CreateStatic( &FLevelEditorActionCallbacks::ActorTypesSelected_CanExecute, FLevelEditorActionCallbacks::IncludeStaticMeshes, /*bSingleOnly*/ false )
 		);
 
 	ActionList.MapAction(
 		Commands.SelectStaticMeshesAllClasses,
-		FExecuteAction::CreateStatic( &FLevelEditorActionCallbacks::ExecuteExecCommand, FString( TEXT("ACTOR SELECT MATCHINGSTATICMESH ALLCLASSES") ) )
+		FExecuteAction::CreateStatic( &FLevelEditorActionCallbacks::ExecuteExecCommand, FString( TEXT("ACTOR SELECT MATCHINGSTATICMESH ALLCLASSES") ) ),
+		FCanExecuteAction::CreateStatic( &FLevelEditorActionCallbacks::ActorTypesSelected_CanExecute, FLevelEditorActionCallbacks::IncludeStaticMeshes, /*bSingleOnly*/ false )
 		);
 
 	ActionList.MapAction(
 		Commands.SelectOwningHierarchicalLODCluster,
-		FExecuteAction::CreateStatic(&FLevelEditorActionCallbacks::OnSelectOwningHLODCluster)
+		FExecuteAction::CreateStatic(&FLevelEditorActionCallbacks::OnSelectOwningHLODCluster),
+		FCanExecuteAction::CreateStatic( &FLevelEditorActionCallbacks::ActorTypesSelected_CanExecute, FLevelEditorActionCallbacks::IncludeStaticMeshes, /*bSingleOnly*/ true )
 		);
 
+	FLevelEditorActionCallbacks::EActorTypeFlags IncludePawnsAndSkeletalMeshes = static_cast<FLevelEditorActionCallbacks::EActorTypeFlags>(FLevelEditorActionCallbacks::IncludePawns | FLevelEditorActionCallbacks::IncludeSkeletalMeshes);
 	ActionList.MapAction(
 		Commands.SelectSkeletalMeshesOfSameClass,
-		FExecuteAction::CreateStatic( &FLevelEditorActionCallbacks::ExecuteExecCommand, FString( TEXT("ACTOR SELECT MATCHINGSKELETALMESH") ) )
+		FExecuteAction::CreateStatic( &FLevelEditorActionCallbacks::ExecuteExecCommand, FString( TEXT("ACTOR SELECT MATCHINGSKELETALMESH") ) ),
+		FCanExecuteAction::CreateStatic( &FLevelEditorActionCallbacks::ActorTypesSelected_CanExecute, IncludePawnsAndSkeletalMeshes, /*bSingleOnly*/ false )
 		);
 
 	ActionList.MapAction(
 		Commands.SelectSkeletalMeshesAllClasses,
-		FExecuteAction::CreateStatic( &FLevelEditorActionCallbacks::ExecuteExecCommand, FString( TEXT("ACTOR SELECT MATCHINGSKELETALMESH ALLCLASSES") ) )
+		FExecuteAction::CreateStatic( &FLevelEditorActionCallbacks::ExecuteExecCommand, FString( TEXT("ACTOR SELECT MATCHINGSKELETALMESH ALLCLASSES") ) ),
+		FCanExecuteAction::CreateStatic( &FLevelEditorActionCallbacks::ActorTypesSelected_CanExecute, IncludePawnsAndSkeletalMeshes, /*bSingleOnly*/ false )
 		);
 
 	ActionList.MapAction(
 		Commands.SelectAllWithSameMaterial,
-		FExecuteAction::CreateStatic( &FLevelEditorActionCallbacks::ExecuteExecCommand, FString( TEXT("ACTOR SELECT MATCHINGMATERIAL") ) )
+		FExecuteAction::CreateStatic( &FLevelEditorActionCallbacks::ExecuteExecCommand, FString( TEXT("ACTOR SELECT MATCHINGMATERIAL") ) ),
+		FCanExecuteAction::CreateStatic( &FLevelEditorActionCallbacks::ActorSelected_CanExecute )
 		);
 
 	ActionList.MapAction(
 		Commands.SelectMatchingEmitter,
-		FExecuteAction::CreateStatic( &FLevelEditorActionCallbacks::ExecuteExecCommand, FString( TEXT("ACTOR SELECT MATCHINGEMITTER") ) )
+		FExecuteAction::CreateStatic( &FLevelEditorActionCallbacks::ExecuteExecCommand, FString( TEXT("ACTOR SELECT MATCHINGEMITTER") ) ),
+		FCanExecuteAction::CreateStatic( &FLevelEditorActionCallbacks::ActorTypesSelected_CanExecute, FLevelEditorActionCallbacks::IncludeEmitters, /*bSingleOnly*/ false )
 		);
 
 	ActionList.MapAction(
@@ -1757,27 +1770,6 @@ void FLevelEditorModule::BindGlobalLevelEditorCommands()
 		Commands.OpenMergeActor,
 		FExecuteAction::CreateStatic(&FLevelEditorActionCallbacks::OpenMergeActor_Clicked)
 		);
-
-	ActionList.MapAction(
-		Commands.GeometryCollectionSelectAllGeometry,
-		FExecuteAction::CreateStatic(FLevelEditorActionCallbacks::GeometryCollection_SelectAllGeometry),
-		FCanExecuteAction(),
-		FIsActionChecked::CreateStatic(FLevelEditorActionCallbacks::GeometryCollection_IsChecked)
-		);
-
-	ActionList.MapAction(
-		Commands.GeometryCollectionSelectNone,
-		FExecuteAction::CreateStatic(FLevelEditorActionCallbacks::GeometryCollection_SelectNone),
-		FCanExecuteAction(),
-		FIsActionChecked::CreateStatic(FLevelEditorActionCallbacks::GeometryCollection_IsChecked)
-		);
-
-	ActionList.MapAction(
-		Commands.GeometryCollectionSelectInverseGeometry,
-		FExecuteAction::CreateStatic(FLevelEditorActionCallbacks::GeometryCollection_SelectInverseGeometry),
-		FCanExecuteAction(),
-		FIsActionChecked::CreateStatic(FLevelEditorActionCallbacks::GeometryCollection_IsChecked)
-		);
 }
 	
 
diff --git a/Engine/Source/Editor/LevelEditor/Private/LevelEditorActions.cpp b/Engine/Source/Editor/LevelEditor/Private/LevelEditorActions.cpp
--- a/Engine/Source/Editor/LevelEditor/Private/LevelEditorActions.cpp
+++ b/Engine/Source/Editor/LevelEditor/Private/LevelEditorActions.cpp
@@ -1824,6 +1824,12 @@ void FLevelEditorActionCallbacks::OnSelectAllActorsOfClass( bool bArchetype )
 	GEditor->SelectAllActorsWithClass( bArchetype );
 }
 
+bool FLevelEditorActionCallbacks::CanSelectAllActorsOfClass()
+{
+	FSelectedActorInfo SelectionInfo = AssetSelectionUtils::GetSelectedActorInfo();
+	return SelectionInfo.NumSelected > 0 && SelectionInfo.bAllSelectedActorsOfSameType;
+}
+
 void FLevelEditorActionCallbacks::OnSelectComponentOwnerActor()
 {
 	auto ComponentOwner = Cast<AActor>(*GEditor->GetSelectedActorIterator());
@@ -1867,6 +1873,9 @@ void FLevelEditorActionCallbacks::OnApplyMaterialToSurface()
 void FLevelEditorActionCallbacks::OnSelectAllLights()
 {
 	GEditor->GetSelectedActors()->BeginBatchSelectOperation();
+
+	GEditor->SelectNone(false, true);
+
 	// Select all light actors.
 	for( ALight* Light : TActorRange<ALight>(GetWorld()) )
 	{
@@ -3190,6 +3199,35 @@ bool FLevelEditorActionCallbacks::ActorsSelected_CanExecute()
 	return GEditor->GetSelectedActorCount() > 1;
 }
 
+bool FLevelEditorActionCallbacks::ActorTypesSelected_CanExecute(EActorTypeFlags TypeFlags, bool bSingleOnly)
+{
+	FSelectedActorInfo SelectionInfo = AssetSelectionUtils::GetSelectedActorInfo();
+	if (SelectionInfo.NumSelected > 0 && (!bSingleOnly || SelectionInfo.NumSelected == 1))
+	{
+		if ((TypeFlags & IncludePawns) && SelectionInfo.bHavePawn)
+		{
+			return true;
+		}
+
+		if ((TypeFlags & IncludeStaticMeshes) && SelectionInfo.bHaveStaticMesh)
+		{
+			return true;
+		}
+
+		if ((TypeFlags & IncludeSkeletalMeshes) && SelectionInfo.bHaveSkeletalMesh)
+		{
+			return true;
+		}
+
+		if ((TypeFlags & IncludeEmitters) && SelectionInfo.bHaveEmitter)
+		{
+			return true;
+		}
+	}
+
+	return false;
+}
+
 bool FLevelEditorActionCallbacks::ElementSelected_CanExecute()
 {
 	// TODO: Ideally this would come from some level editor context
@@ -3569,10 +3607,6 @@ void FLevelEditorCommands::RegisterCommands()
 		);
 	}
 
-	UI_COMMAND(GeometryCollectionSelectAllGeometry, "Select All Geometry In Hierarchy", "Select all geometry in hierarchy", EUserInterfaceActionType::Button, FInputChord());
-	UI_COMMAND(GeometryCollectionSelectNone, "Deselect All Geometry In Hierarchy", "Deselect all geometry in hierarchy", EUserInterfaceActionType::Button, FInputChord());
-	UI_COMMAND(GeometryCollectionSelectInverseGeometry, "Select Inverse Geometry In Hierarchy", "Select inverse geometry in hierarchy", EUserInterfaceActionType::Button, FInputChord());
-
 	UI_COMMAND(OpenMergeActor, "Merge Actors", "Opens the Merge Actor panel", EUserInterfaceActionType::Button, FInputChord());
 }
 
diff --git a/Engine/Source/Editor/LevelEditor/Private/LevelEditorContextMenu.cpp b/Engine/Source/Editor/LevelEditor/Private/LevelEditorContextMenu.cpp
--- a/Engine/Source/Editor/LevelEditor/Private/LevelEditorContextMenu.cpp
+++ b/Engine/Source/Editor/LevelEditor/Private/LevelEditorContextMenu.cpp
@@ -60,13 +60,6 @@ class FLevelEditorContextMenuImpl
 public:
 	static FSelectedActorInfo SelectionInfo;
 public:
-	/**
-	 * Fills in menu options for the select actor menu
-	 *
-	 * @param MenuBuilder	The menu to add items to
-	 */
-	static void FillSelectActorMenu(UToolMenu* Menu);
-
 	/**
 	 * Fills in menu options for the actor visibility menu
 	 *
@@ -394,13 +387,6 @@ void FLevelEditorContextMenu::RegisterActorContextMenu()
 		{
 			FToolMenuSection& Section = InMenu->AddSection("ActorSelectVisibilityLevels");
 
-			// Add a sub-menu for "Select"
-			Section.AddSubMenu(
-				"SelectSubMenu",
-				LOCTEXT("SelectSubMenu", "Select"),
-				LOCTEXT("SelectSubMenu_ToolTip", "Opens the actor selection menu"),
-				FNewToolMenuDelegate::CreateStatic(&FLevelEditorContextMenuImpl::FillSelectActorMenu));
-
 			Section.AddSubMenu(
 				"EditSubMenu",
 				LOCTEXT("EditSubMenu", "Edit"),
@@ -879,118 +865,6 @@ void FLevelEditorContextMenu::BuildGroupMenu(UToolMenu* Menu, const FSelectedAct
 	}
 }
 
-void FLevelEditorContextMenuImpl::FillSelectActorMenu(UToolMenu* Menu)
-{
-	FText SelectAllActorStr = FText::Format( LOCTEXT("SelectActorsOfSameClass", "Select All {0}(s)"), FText::FromString( SelectionInfo.SelectionStr ) );
-	int32 NumSelectedSurfaces = AssetSelectionUtils::GetNumSelectedSurfaces( SelectionInfo.SharedWorld );
-
-	{
-		FToolMenuSection& Section = Menu->AddSection("SelectActorGeneral", LOCTEXT("SelectAnyHeading", "General"));
-		Section.AddMenuEntry( FGenericCommands::Get().SelectAll, TAttribute<FText>(), LOCTEXT("SelectAll_ToolTip", "Selects all actors") );
-		Section.AddMenuEntry( FLevelEditorCommands::Get().SelectNone );
-		Section.AddMenuEntry( FLevelEditorCommands::Get().InvertSelection );
-	}
-
-	if( !SelectionInfo.bHaveBrush && SelectionInfo.bAllSelectedActorsOfSameType && SelectionInfo.SelectionStr.Len() != 0 )
-	{
-		// These menu options appear if only if all the actors are the same type and we aren't selecting brush
-		FToolMenuSection& Section = Menu->AddSection("SelectAllActorsOfSameClass");
-		Section.AddMenuEntry(FLevelEditorCommands::Get().SelectAllActorsOfSameClass, SelectAllActorStr);
-	}
-
-	{
-		FToolMenuSection& Section = Menu->AddSection("SelectActorHierarchy", LOCTEXT("SelectHierarchyHeading", "Hierarchy"));
-		Section.AddMenuEntry( FLevelEditorCommands::Get().SelectImmediateChildren );
-		Section.AddMenuEntry( FLevelEditorCommands::Get().SelectAllDescendants );
-	}
-
-	// Add brush commands when we have a brush or any surfaces selected
-	{
-		FToolMenuSection& Section = Menu->AddSection("SelectBSP", LOCTEXT("SelectBSPHeading", "BSP"));
-		if( SelectionInfo.bHaveBrush || NumSelectedSurfaces > 0 )
-		{
-			if( SelectionInfo.bAllSelectedAreBrushes )
-			{
-				Section.AddMenuEntry( FLevelEditorCommands::Get().SelectAllActorsOfSameClass, SelectAllActorStr);
-			}
-		}
-
-		Section.AddMenuEntry(FLevelEditorCommands::Get().SelectAllAddditiveBrushes);
-		Section.AddMenuEntry(FLevelEditorCommands::Get().SelectAllSubtractiveBrushes);
-		Section.AddMenuEntry(FLevelEditorCommands::Get().SelectAllSurfaces);
-	}
-
-	if( SelectionInfo.NumSelected > 0 || NumSelectedSurfaces > 0 )
-	{
-		// If any actors are selected add lights selection options
-		{
-			FToolMenuSection& Section = Menu->AddSection("SelectLights", LOCTEXT("SelectLightHeading", "Lights"));
-			Section.AddMenuEntry(FLevelEditorCommands::Get().SelectRelevantLights);
-
-			if ( SelectionInfo.bHaveLight )
-			{
-				Section.AddMenuEntry(FLevelEditorCommands::Get().SelectAllLights);
-				Section.AddMenuEntry(FLevelEditorCommands::Get().SelectStationaryLightsExceedingOverlap);
-			}
-		}
-
-		if( SelectionInfo.bHaveStaticMesh )
-		{
-			// if any static meshes are selected allow selecting actors using the same mesh
-			{
-				FToolMenuSection& Section = Menu->AddSection("SelectMeshes", LOCTEXT("SelectStaticMeshHeading", "Static Meshes"));
-				Section.AddMenuEntry(FLevelEditorCommands::Get().SelectStaticMeshesOfSameClass, LOCTEXT("SelectStaticMeshesOfSameClass_Menu", "Select Matching (Selected Classes)"));
-				Section.AddMenuEntry(FLevelEditorCommands::Get().SelectStaticMeshesAllClasses, LOCTEXT("SelectStaticMeshesAllClasses_Menu", "Select Matching (All Classes)"));
-			}
-
-			if (SelectionInfo.NumSelected == 1)
-			{
-				{
-					FToolMenuSection& Section = Menu->AddSection("SelectHLODCluster", LOCTEXT("SelectHLODClusterHeading", "Hierachical LODs"));
-					Section.AddMenuEntry(FLevelEditorCommands::Get().SelectOwningHierarchicalLODCluster, LOCTEXT("SelectOwningHierarchicalLODCluster_Menu", "Select Owning HierarchicalLODCluster"));
-				}
-			}			
-		}
-
-		if( SelectionInfo.bHavePawn || SelectionInfo.bHaveSkeletalMesh )
-		{
-			// if any skeletal meshes are selected allow selecting actors using the same mesh
-			{
-				FToolMenuSection& Section = Menu->AddSection("SelectSkeletalMeshes", LOCTEXT("SelectSkeletalMeshHeading", "Skeletal Meshes"));
-				Section.AddMenuEntry(FLevelEditorCommands::Get().SelectSkeletalMeshesOfSameClass);
-				Section.AddMenuEntry(FLevelEditorCommands::Get().SelectSkeletalMeshesAllClasses);
-			}
-		}
-
-		if( SelectionInfo.bHaveEmitter )
-		{
-			{
-				FToolMenuSection& Section = Menu->AddSection("SelectEmitters", LOCTEXT("SelectEmitterHeading", "Emitters"));
-				Section.AddMenuEntry(FLevelEditorCommands::Get().SelectMatchingEmitter);
-			}
-		}
-	}
-
-	if( SelectionInfo.bHaveBrush || SelectionInfo.NumSelected > 0 )
-	{
-		FToolMenuSection& Section = Menu->AddSection("SelectMaterial", LOCTEXT("SelectMaterialHeading", "Materials"));
-		{
-			Section.AddMenuEntry(FLevelEditorCommands::Get().SelectAllWithSameMaterial);
-		}
-	}
-
-	// Add geometry collection commands
-	if (FModuleManager::Get().IsModuleLoaded("GeometryCollectionEditor"))
-	{
-		FToolMenuSection& Section = Menu->AddSection("SelectBones", LOCTEXT("GeometryCollectionHeading", "Geometry Collection"));
-		{
-			Section.AddMenuEntry(FLevelEditorCommands::Get().GeometryCollectionSelectAllGeometry);
-			Section.AddMenuEntry(FLevelEditorCommands::Get().GeometryCollectionSelectNone);
-			Section.AddMenuEntry(FLevelEditorCommands::Get().GeometryCollectionSelectInverseGeometry);
-		}
-	}
-}
-
 void FLevelEditorContextMenuImpl::FillActorVisibilityMenu(UToolMenu* Menu)
 {
 	{
diff --git a/Engine/Source/Editor/LevelEditor/Private/LevelEditorMenu.cpp b/Engine/Source/Editor/LevelEditor/Private/LevelEditorMenu.cpp
--- a/Engine/Source/Editor/LevelEditor/Private/LevelEditorMenu.cpp
+++ b/Engine/Source/Editor/LevelEditor/Private/LevelEditorMenu.cpp
@@ -22,6 +22,8 @@
 #include "Settings/EditorExperimentalSettings.h"
 #include "TranslationEditor/Public/ITranslationEditor.h"
 #include "LocalizationDashboard/Public/ILocalizationDashboardModule.h"
+#include "AssetSelection.h"
+#include "EditorViewportCommands.h"
 
 #define LOCTEXT_NAMESPACE "LevelEditorMenu"
 
@@ -231,7 +233,7 @@ void FLevelEditorMenu::RegisterLevelEditorMenus()
 			FToolMenuSection& Section = Menu->FindOrAddSection(NAME_None);
 
 
-			FToolMenuEntry& Entry =
+			FToolMenuEntry& BuildEntry =
 				Section.AddSubMenu(
 					"Build",
 					LOCTEXT("BuildMenu", "Build"),
@@ -239,7 +241,17 @@ void FLevelEditorMenu::RegisterLevelEditorMenus()
 					FNewToolMenuChoice()
 				);
 
-			Entry.InsertPosition = FToolMenuInsert("Help", EToolMenuInsertType::Before);
+			BuildEntry.InsertPosition = FToolMenuInsert("Help", EToolMenuInsertType::Before);
+
+			FToolMenuEntry& SelectEntry =
+				Section.AddSubMenu(
+					"Select",
+					LOCTEXT("SelectMenu", "Select"),
+					LOCTEXT("SelectMenu_ToolTip", "Level Selection"),
+					FNewToolMenuChoice()
+				);
+
+			SelectEntry.InsertPosition = FToolMenuInsert("Help", EToolMenuInsertType::Before);
 		}
 
 
@@ -306,6 +318,8 @@ void FLevelEditorMenu::RegisterLevelEditorMenus()
 	// Extend the Build menu
 	RegisterBuildMenu();
 
+	// Extend the Select menu
+	RegisterSelectMenu();
 }
 
 TSharedRef< SWidget > FLevelEditorMenu::MakeLevelEditorMenu( const TSharedPtr<FUICommandList>& CommandList, TSharedPtr<class SLevelEditor> LevelEditor )
@@ -637,4 +651,76 @@ void FLevelEditorMenu::RegisterBuildMenu()
 	}
 
 }
+
+void FLevelEditorMenu::RegisterSelectMenu()
+{
+	static const FName BaseMenuName = "LevelEditor.MainMenu.Select";
+	UToolMenu* Menu = UToolMenus::Get()->RegisterMenu(BaseMenuName);
+
+	{
+		FToolMenuSection& Section = Menu->AddSection("SelectActorGeneral", NSLOCTEXT("LevelViewportContextMenu", "SelectAnyHeading", "General"));
+		Section.AddMenuEntry(FGenericCommands::Get().SelectAll, TAttribute<FText>(), NSLOCTEXT("LevelViewportContextMenu", "SelectAll_ToolTip", "Selects all actors"));
+		Section.AddMenuEntry(FLevelEditorCommands::Get().SelectNone);
+		Section.AddMenuEntry(FLevelEditorCommands::Get().InvertSelection);
+		Section.AddDynamicEntry(NAME_None, FNewToolMenuSectionDelegate::CreateLambda([](FToolMenuSection& InSection)
+			{
+				FSelectedActorInfo SelectionInfo = AssetSelectionUtils::GetSelectedActorInfo();
+				TAttribute<FText> SelectAllActorsText;
+				if (SelectionInfo.bAllSelectedActorsOfSameType && !SelectionInfo.SelectionStr.IsEmpty())
+				{
+					SelectAllActorsText = FText::Format(NSLOCTEXT("LevelViewportContextMenu", "SelectActorsOfSameClass", "Select All {0}(s)"), FText::FromString(SelectionInfo.SelectionStr));
+				}
+				InSection.AddMenuEntry(FLevelEditorCommands::Get().SelectAllActorsOfSameClass, SelectAllActorsText);
+			}));
+		Section.AddMenuEntry(FEditorViewportCommands::Get().FocusViewportToSelection);
+	}
+
+	{
+		FToolMenuSection& Section = Menu->AddSection("SelectActorHierarchy", NSLOCTEXT("LevelViewportContextMenu", "SelectHierarchyHeading", "Hierarchy"));
+		Section.AddMenuEntry(FLevelEditorCommands::Get().SelectImmediateChildren);
+		Section.AddMenuEntry(FLevelEditorCommands::Get().SelectAllDescendants);
+	}
+
+	{
+		FToolMenuSection& Section = Menu->AddSection("SelectBSP", NSLOCTEXT("LevelViewportContextMenu", "SelectBSPHeading", "BSP"));
+		Section.AddMenuEntry(FLevelEditorCommands::Get().SelectAllAddditiveBrushes);
+		Section.AddMenuEntry(FLevelEditorCommands::Get().SelectAllSubtractiveBrushes);
+		Section.AddMenuEntry(FLevelEditorCommands::Get().SelectAllSurfaces);
+	}
+
+	{
+		FToolMenuSection& Section = Menu->AddSection("SelectLights", NSLOCTEXT("LevelViewportContextMenu", "SelectLightHeading", "Lights"));
+		Section.AddMenuEntry(FLevelEditorCommands::Get().SelectRelevantLights);
+		Section.AddMenuEntry(FLevelEditorCommands::Get().SelectAllLights);
+		Section.AddMenuEntry(FLevelEditorCommands::Get().SelectStationaryLightsExceedingOverlap);
+	}
+
+	{
+		FToolMenuSection& Section = Menu->AddSection("SelectMeshes", NSLOCTEXT("LevelViewportContextMenu", "SelectStaticMeshHeading", "Static Meshes"));
+		Section.AddMenuEntry(FLevelEditorCommands::Get().SelectStaticMeshesOfSameClass, NSLOCTEXT("LevelViewportContextMenu", "SelectStaticMeshesOfSameClass_Menu", "Select Matching (Selected Classes)"));
+		Section.AddMenuEntry(FLevelEditorCommands::Get().SelectStaticMeshesAllClasses, NSLOCTEXT("LevelViewportContextMenu", "SelectStaticMeshesAllClasses_Menu", "Select Matching (All Classes)"));
+	}
+
+	{
+		FToolMenuSection& Section = Menu->AddSection("SelectHLODCluster", NSLOCTEXT("LevelViewportContextMenu", "SelectHLODClusterHeading", "Hierachical LODs"));
+		Section.AddMenuEntry(FLevelEditorCommands::Get().SelectOwningHierarchicalLODCluster, NSLOCTEXT("LevelViewportContextMenu", "SelectOwningHierarchicalLODCluster_Menu", "Select Owning HierarchicalLODCluster"));
+	}
+
+	{
+		FToolMenuSection& Section = Menu->AddSection("SelectSkeletalMeshes", NSLOCTEXT("LevelViewportContextMenu", "SelectSkeletalMeshHeading", "Skeletal Meshes"));
+		Section.AddMenuEntry(FLevelEditorCommands::Get().SelectSkeletalMeshesOfSameClass);
+		Section.AddMenuEntry(FLevelEditorCommands::Get().SelectSkeletalMeshesAllClasses);
+	}
+
+	{
+		FToolMenuSection& Section = Menu->AddSection("SelectEmitters", NSLOCTEXT("LevelViewportContextMenu", "SelectEmitterHeading", "Emitters"));
+		Section.AddMenuEntry(FLevelEditorCommands::Get().SelectMatchingEmitter);
+	}
+
+	{
+		FToolMenuSection& Section = Menu->AddSection("SelectMaterial", NSLOCTEXT("LevelViewportContextMenu", "SelectMaterialHeading", "Materials"));
+		Section.AddMenuEntry(FLevelEditorCommands::Get().SelectAllWithSameMaterial);
+	}
+}
+
 #undef LOCTEXT_NAMESPACE
\ No newline at end of file
diff --git a/Engine/Source/Editor/LevelEditor/Private/LevelEditorMenu.h b/Engine/Source/Editor/LevelEditor/Private/LevelEditorMenu.h
--- a/Engine/Source/Editor/LevelEditor/Private/LevelEditorMenu.h
+++ b/Engine/Source/Editor/LevelEditor/Private/LevelEditorMenu.h
@@ -26,4 +26,5 @@ public:
 	static TSharedRef< SWidget > MakeLevelEditorMenu( const TSharedPtr<FUICommandList>& CommandList, TSharedPtr<class SLevelEditor> LevelEditor );
 private:
 	static void RegisterBuildMenu();
+	static void RegisterSelectMenu();
 };
diff --git a/Engine/Source/Editor/LevelEditor/Public/LevelEditorActions.h b/Engine/Source/Editor/LevelEditor/Public/LevelEditorActions.h
--- a/Engine/Source/Editor/LevelEditor/Public/LevelEditorActions.h
+++ b/Engine/Source/Editor/LevelEditor/Public/LevelEditorActions.h
@@ -622,16 +622,6 @@ public:
         // Open merge actor command
 	TSharedPtr< FUICommandInfo > OpenMergeActor;
 
-	/** Selects all Geometry Collection geometry */
-	TSharedPtr< FUICommandInfo > GeometryCollectionSelectAllGeometry;
-
-	/** Deselects all Geometry Collection geometry */
-	TSharedPtr< FUICommandInfo > GeometryCollectionSelectNone;
-
-	/** Selects inverse of currently seleted Geometry Collection geometry */
-	TSharedPtr< FUICommandInfo > GeometryCollectionSelectInverseGeometry;
-
-
 };
 
 /**
@@ -1004,6 +994,9 @@ public:
 	 */
 	static void OnSelectAllActorsOfClass( bool bArchetype );
 
+	/** Called to see if all selected actors are the same class */
+	static bool CanSelectAllActorsOfClass();
+
 	/** Called when selecting the actor that owns the currently selected component(s) */
 	static void OnSelectComponentOwnerActor();
 
@@ -1331,6 +1324,23 @@ public:
 	 */
 	static bool ActorSelected_CanExecute();
 
+	enum EActorTypeFlags : uint8
+	{
+		IncludePawns			= 1 << 0,
+		IncludeStaticMeshes		= 1 << 1,
+		IncludeSkeletalMeshes	= 1 << 2,
+		IncludeEmitters			= 1 << 3,
+	};
+
+	/**
+	 * Checks to see if at least one actor (of the given types) is selected
+	 *
+	 * @param TypeFlags		actor types to look for - one or more of EActorTypeFlags or'ed together
+	 * @param bSingleOnly	if true, then requires selection to be exactly one actor
+	 * @return				true if it can execute.
+	 */
+	static bool ActorTypesSelected_CanExecute(EActorTypeFlags TypeFlags, bool bSingleOnly);
+
 	/**
 	 * Checks to see if multiple actors are selected
 	 *	@return true if it can execute.
diff --git a/Engine/Source/Editor/WorldPartitionEditor/Private/WorldPartitionEditorModule.cpp b/Engine/Source/Editor/WorldPartitionEditor/Private/WorldPartitionEditorModule.cpp
--- a/Engine/Source/Editor/WorldPartitionEditor/Private/WorldPartitionEditorModule.cpp
+++ b/Engine/Source/Editor/WorldPartitionEditor/Private/WorldPartitionEditorModule.cpp
@@ -70,7 +70,7 @@ static TSharedRef<FExtender> OnExtendLevelEditorMenu(const TSharedRef<FUICommand
 	if (Volumes.Num())
 	{
 		Extender->AddMenuExtension(
-			"SelectActorGeneral",
+			"ActorSelectVisibilityLevels",
 			EExtensionHook::After,
 			nullptr,
 			FMenuExtensionDelegate::CreateStatic(&CreateLevelViewportContextMenuEntries, Volumes));
