Editor no longer crashes when resetting layout while a Level Sequence is open.

diff --git a/Engine/Source/Editor/LevelEditor/Private/LevelEditor.cpp b/Engine/Source/Editor/LevelEditor/Private/LevelEditor.cpp
--- a/Engine/Source/Editor/LevelEditor/Private/LevelEditor.cpp
+++ b/Engine/Source/Editor/LevelEditor/Private/LevelEditor.cpp
@@ -371,11 +371,13 @@ void FLevelEditorModule::SummonWorldBrowserComposition()
 }
 
 // @todo remove when world-centric mode is added
-void FLevelEditorModule::AttachSequencer( TSharedPtr<SWidget> SequencerWidget, TSharedPtr<IAssetEditorInstance> SequencerAssetEditor )
+void FLevelEditorModule::AttachSequencer(TSharedPtr<SWidget> SequencerWidget, TSharedPtr<IAssetEditorInstance> SequencerAssetEditor)
 {
 	TSharedPtr<SLevelEditor> LevelEditorInstance = LevelEditorInstancePtr.Pin();
-
-	LevelEditorInstance->AttachSequencer( SequencerWidget, SequencerAssetEditor );
+	if (LevelEditorInstance.IsValid())
+	{
+		LevelEditorInstance->AttachSequencer(SequencerWidget, SequencerAssetEditor);
+	}
 }
 
 TSharedPtr<IAssetViewport> FLevelEditorModule::GetFirstActiveViewport()
