Add new Actions (Actor/Component/Element) menu to the level editor's main menu bar that brings up the context menu corresponding to the selection

- The menu's title changes depending on the selection ("Actor", "Component", or "Element" currently since those are the three currently-defined context menus for the level editor)
- The context menus will get reorganized in subsequent changes so that the pulldown Actions menu displays the full set of entries while the right-click context menu is trimmed down to a shorter subset of entries. This change does not reorganize any items yet.
- Add support to UToolMenus for using a TAttribute<FText> for the pulldown menu title/tooltip and for supplying FOnGetContent for generating the pulldown menu.

diff --git a/Engine/Source/Developer/ToolMenus/Private/ToolMenus.cpp b/Engine/Source/Developer/ToolMenus/Private/ToolMenus.cpp
--- a/Engine/Source/Developer/ToolMenus/Private/ToolMenus.cpp
+++ b/Engine/Source/Developer/ToolMenus/Private/ToolMenus.cpp
@@ -1389,23 +1389,33 @@ void UToolMenus::PopulateMenuBarBuilder(FMenuBarBuilder& MenuBarBuilder, UToolMe
 		const FToolMenuSection& Section = MenuData->Sections[i];
 		for (const FToolMenuEntry& Block : Section.Blocks)
 		{
-			FName SubMenuFullName = JoinMenuPaths(MenuData->MenuName, Block.Name);
-			FNewMenuDelegate NewMenuDelegate;
-			if (Block.SubMenuData.ConstructMenu.NewMenuLegacy.IsBound())
+			if (Block.SubMenuData.ConstructMenu.OnGetContent.IsBound())
 			{
-				NewMenuDelegate = Block.SubMenuData.ConstructMenu.NewMenuLegacy;
+				MenuBarBuilder.AddPullDownMenu(
+					Block.Label,
+					Block.ToolTip,
+					Block.SubMenuData.ConstructMenu.OnGetContent,
+					Block.Name
+				);
+			}
+			else if (Block.SubMenuData.ConstructMenu.NewMenuLegacy.IsBound())
+			{
+				MenuBarBuilder.AddPullDownMenu(
+					Block.Label,
+					Block.ToolTip,
+					Block.SubMenuData.ConstructMenu.NewMenuLegacy,
+					Block.Name
+				);
 			}
 			else
 			{
-				NewMenuDelegate = FNewMenuDelegate::CreateUObject(this, &UToolMenus::PopulateSubMenu, TWeakObjectPtr<UToolMenu>(MenuData), Block.Name);
+				MenuBarBuilder.AddPullDownMenu(
+					Block.Label,
+					Block.ToolTip,
+					FNewMenuDelegate::CreateUObject(this, &UToolMenus::PopulateSubMenu, TWeakObjectPtr<UToolMenu>(MenuData), Block.Name),
+					Block.Name
+				);
 			}
-
-			MenuBarBuilder.AddPullDownMenu(
-				Block.Label.Get(),
-				Block.ToolTip.Get(),
-				NewMenuDelegate,
-				Block.Name
-			);
 		}
 	}
 
diff --git a/Engine/Source/Editor/LevelEditor/Private/LevelEditorContextMenu.cpp b/Engine/Source/Editor/LevelEditor/Private/LevelEditorContextMenu.cpp
--- a/Engine/Source/Editor/LevelEditor/Private/LevelEditorContextMenu.cpp
+++ b/Engine/Source/Editor/LevelEditor/Private/LevelEditorContextMenu.cpp
@@ -244,7 +244,7 @@ void FLevelEditorContextMenu::RegisterActorContextMenu()
 			return;
 		}
 
-		TWeakPtr<SLevelEditor> LevelEditor = LevelEditorContext->LevelEditor;
+		TWeakPtr<ILevelEditor> LevelEditor = LevelEditorContext->LevelEditor;
 
 		// Generate information about our selection
 		TArray<AActor*> SelectedActors;
@@ -315,7 +315,7 @@ void FLevelEditorContextMenu::RegisterActorContextMenu()
 			{
 				const FLevelViewportCommands& Actions = FLevelViewportCommands::Get();
 
-				auto Viewport = LevelEditor.Pin()->GetActiveViewport();
+				auto Viewport = LevelEditor.Pin()->GetActiveViewportInterface();
 				if (Viewport.IsValid())
 				{
 					auto& ViewportClient = Viewport->GetLevelViewportClient();
@@ -521,6 +521,41 @@ void FLevelEditorContextMenu::RegisterSceneOutlinerContextMenu()
 	}));
 }
 
+void FLevelEditorContextMenu::RegisterMenuBarEmptyContextMenu()
+{
+	UToolMenus* ToolMenus = UToolMenus::Get();
+	if (ToolMenus->IsMenuRegistered("LevelEditor.MenuBarEmptyContextMenu"))
+	{
+		return;
+	}
+
+	UToolMenu* Menu = ToolMenus->RegisterMenu("LevelEditor.MenuBarEmptyContextMenu");
+	FToolMenuSection& Section = Menu->AddSection("MenuBarEmpty");
+
+	const FText EmptySelectionInformationalMessage = LOCTEXT("EmptySelectionInformationalMessage", "Select an object to view actions.");
+
+#if PLATFORM_MAC
+	// Can't include arbitrary widgets in a main menu on Mac, so display the informational message using a disabled entry.
+	Section.AddMenuEntry(
+		NAME_None,
+		EmptySelectionInformationalMessage,
+		TAttribute<FText>(),
+		TAttribute<FSlateIcon>(),
+		FUIAction(FExecuteAction(), FCanExecuteAction::CreateLambda([]() { return false; })));
+#else
+	Section.AddEntry(FToolMenuEntry::InitWidget(
+		NAME_None,
+		SNew(SBox)
+		.Padding(FMargin(80.f, 8.f))
+		[
+			SNew(STextBlock)
+			.Text(EmptySelectionInformationalMessage)
+			.TextStyle(FAppStyle::Get(), "HintText")
+		],
+		FText::GetEmpty(), /*bNoIndent*/ true, /*bSearchable*/ false));
+#endif
+}
+
 void FLevelEditorContextMenu::RegisterEmptySelectionContextMenu()
 {
 	UToolMenus* ToolMenus = UToolMenus::Get();
@@ -572,18 +607,86 @@ FName FLevelEditorContextMenu::GetContextMenuName(ELevelEditorMenuContext Contex
 		return "LevelEditor.SceneOutlinerContextMenu";
 	}
 
+	if (ContextType == ELevelEditorMenuContext::MainMenu)
+	{
+		return "LevelEditor.MenuBarEmptyContextMenu";
+	}
+
 	return "LevelEditor.EmptySelectionContextMenu";
 }
 
-FName FLevelEditorContextMenu::InitMenuContext(FToolMenuContext& Context, TWeakPtr<SLevelEditor> LevelEditor, ELevelEditorMenuContext ContextType, const FTypedElementHandle& HitProxyElement)
+FText FLevelEditorContextMenu::GetContextMenuTitle(ELevelEditorMenuContext ContextType, const UTypedElementSelectionSet* InSelectionSet)
+{
+	if (InSelectionSet)
+	{
+		if (InSelectionSet->HasSelectedObjects<UActorComponent>())
+		{
+			return LOCTEXT("ComponentContextMenuTitle", "Component");
+		}
+
+		if (InSelectionSet->HasSelectedObjects<AActor>())
+		{
+			return LOCTEXT("ActorContextMenuTitle", "Actor");
+		}
+
+		if (InSelectionSet->GetNumSelectedElements())
+		{
+			return LOCTEXT("ElementContextMenuTitle", "Element");
+		}
+	}
+
+	// Show "Actor" label by default as title when nothing selected since most selections (currently) will be actors anyways
+	return LOCTEXT("ActorContextMenuTitle", "Actor");
+}
+
+FText FLevelEditorContextMenu::GetContextMenuToolTip(ELevelEditorMenuContext ContextType, const UTypedElementSelectionSet* InSelectionSet)
+{
+	if (InSelectionSet)
+	{
+		const int32 ComponentCount = InSelectionSet->CountSelectedObjects<UActorComponent>();
+		if (ComponentCount == 1)
+		{
+			UActorComponent* Component = InSelectionSet->GetTopSelectedObject<UActorComponent>();
+			check(Component);
+			return FText::Format(LOCTEXT("ComponentContextMenuToolTipSingle", "Show actions for component \"{0}\""), FText::FromString(Component->GetName()));
+		}
+		else if (ComponentCount > 1)
+		{
+			return FText::Format(LOCTEXT("ComponentContextMenuToolTipOther", "Show actions for {0} components"), FText::AsNumber(ComponentCount));
+		}
+
+		const int32 ActorCount = InSelectionSet->CountSelectedObjects<AActor>();
+		if (ActorCount == 1)
+		{
+			AActor* Actor = InSelectionSet->GetTopSelectedObject<AActor>();
+			check(Actor);
+			return FText::Format(LOCTEXT("ActorContextMenuToolTipSingle", "Show actions for actor \"{0}\""), FText::FromString(Actor->GetActorLabel()));
+		}
+		else if (ActorCount > 1)
+		{
+			return FText::Format(LOCTEXT("ActorContextMenuToolTipOther", "Show actions for {0} actors"), FText::AsNumber(ActorCount));
+		}
+
+		const int32 ElementCount = InSelectionSet->GetNumSelectedElements();
+		if (ElementCount)
+		{
+			return FText::Format(LOCTEXT("ElementContextMenuToolTip", "Show actions for {0} {0}|plural(one=element,other=elements)"), FText::AsNumber(ElementCount));
+		}
+	}
+
+	return LOCTEXT("NothingSelectedToolTip", "Select an object to show actions");
+}
+
+FName FLevelEditorContextMenu::InitMenuContext(FToolMenuContext& Context, TWeakPtr<ILevelEditor> LevelEditor, ELevelEditorMenuContext ContextType, const FTypedElementHandle& HitProxyElement)
 {
 	RegisterComponentContextMenu();
 	RegisterActorContextMenu();
 	RegisterElementContextMenu();
 	RegisterSceneOutlinerContextMenu();
+	RegisterMenuBarEmptyContextMenu();
 	RegisterEmptySelectionContextMenu();
 
-	TSharedPtr<SLevelEditor> LevelEditorPtr = LevelEditor.Pin();
+	TSharedPtr<ILevelEditor> LevelEditorPtr = LevelEditor.Pin();
 	check(LevelEditorPtr);
 
 	TSharedPtr<FUICommandList> LevelEditorActionsList = LevelEditorPtr->GetLevelEditorActions();
@@ -654,7 +757,7 @@ FName FLevelEditorContextMenu::InitMenuContext(FToolMenuContext& Context, TWeakP
 	return GetContextMenuName(ContextType, ContextObject->CurrentSelection);
 }
 
-UToolMenu* FLevelEditorContextMenu::GenerateMenu(TWeakPtr<SLevelEditor> LevelEditor, ELevelEditorMenuContext ContextType, TSharedPtr<FExtender> Extender, const FTypedElementHandle& HitProxyElement)
+UToolMenu* FLevelEditorContextMenu::GenerateMenu(TWeakPtr<ILevelEditor> LevelEditor, ELevelEditorMenuContext ContextType, TSharedPtr<FExtender> Extender, const FTypedElementHandle& HitProxyElement)
 {
 	FToolMenuContext Context;
 	if (Extender.IsValid())
@@ -668,7 +771,7 @@ UToolMenu* FLevelEditorContextMenu::GenerateMenu(TWeakPtr<SLevelEditor> LevelEdi
 
 // NOTE: We intentionally receive a WEAK pointer here because we want to be callable by a delegate whose
 //       payload contains a weak reference to a level editor instance
-TSharedPtr< SWidget > FLevelEditorContextMenu::BuildMenuWidget(TWeakPtr< SLevelEditor > LevelEditor, ELevelEditorMenuContext ContextType, TSharedPtr<FExtender> Extender, const FTypedElementHandle& HitProxyElement)
+TSharedRef< SWidget > FLevelEditorContextMenu::BuildMenuWidget(TWeakPtr< ILevelEditor > LevelEditor, ELevelEditorMenuContext ContextType, TSharedPtr<FExtender> Extender, const FTypedElementHandle& HitProxyElement)
 {
 	UToolMenu* Menu = GenerateMenu(LevelEditor, ContextType, Extender, HitProxyElement);
 	return UToolMenus::Get()->GenerateWidget(Menu);
diff --git a/Engine/Source/Editor/LevelEditor/Private/LevelEditorContextMenu.h b/Engine/Source/Editor/LevelEditor/Private/LevelEditorContextMenu.h
--- a/Engine/Source/Editor/LevelEditor/Private/LevelEditorContextMenu.h
+++ b/Engine/Source/Editor/LevelEditor/Private/LevelEditorContextMenu.h
@@ -45,7 +45,7 @@ public:
 	 * @param	HitProxyActor	The hitproxy actor in the case the ContextType is Viewport
 	 * @return	Widget for this context menu
 	 */
-	static TSharedPtr< SWidget > BuildMenuWidget(TWeakPtr< SLevelEditor > LevelEditor, ELevelEditorMenuContext ContextType, TSharedPtr<FExtender> Extender = TSharedPtr<FExtender>(), const FTypedElementHandle& HitProxyElement = FTypedElementHandle());
+	static TSharedRef< SWidget > BuildMenuWidget(TWeakPtr< ILevelEditor > LevelEditor, ELevelEditorMenuContext ContextType, TSharedPtr<FExtender> Extender = TSharedPtr<FExtender>(), const FTypedElementHandle& HitProxyElement = FTypedElementHandle());
 
 	/**
 	 * Populates the specified menu builder for the context menu that can be inserted into a pop-up window
@@ -56,20 +56,27 @@ public:
 	 * @param	Extender		Allows extension of this menu based on context.
 	 * @param	HitProxyActor	The hitproxy actor in the case the ContextType is Viewport
 	 */
-	static UToolMenu* GenerateMenu(TWeakPtr< SLevelEditor > LevelEditor, ELevelEditorMenuContext ContextType, TSharedPtr<FExtender> Extender = TSharedPtr<FExtender>(), const FTypedElementHandle& HitProxyElement = FTypedElementHandle());
+	static UToolMenu* GenerateMenu(TWeakPtr< ILevelEditor > LevelEditor, ELevelEditorMenuContext ContextType, TSharedPtr<FExtender> Extender = TSharedPtr<FExtender>(), const FTypedElementHandle& HitProxyElement = FTypedElementHandle());
 
 	/* Adds required information to Context for build menu based on current selection */
-	static FName InitMenuContext(FToolMenuContext& Context, TWeakPtr<SLevelEditor> LevelEditor, ELevelEditorMenuContext ContextType, const FTypedElementHandle& HitProxyElement = FTypedElementHandle());
+	static FName InitMenuContext(FToolMenuContext& Context, TWeakPtr<ILevelEditor> LevelEditor, ELevelEditorMenuContext ContextType, const FTypedElementHandle& HitProxyElement = FTypedElementHandle());
 
 	/* Returns name of menu to display based on current selection */
 	static FName GetContextMenuName(ELevelEditorMenuContext ContextType, const UTypedElementSelectionSet* InSelectionSet);
 
+	/* Returns a user-readable title for the menu to display; the title can be displayed in UI like the menu bar */
+	static FText GetContextMenuTitle(ELevelEditorMenuContext ContextType, const UTypedElementSelectionSet* InSelectionSet);
+
+	/* Returns a user-readable tooltip describing the menu to display */
+	static FText GetContextMenuToolTip(ELevelEditorMenuContext ContextType, const UTypedElementSelectionSet* InSelectionSet);
+
 private:
 
 	static void RegisterComponentContextMenu();
 	static void RegisterActorContextMenu();
 	static void RegisterElementContextMenu();
 	static void RegisterSceneOutlinerContextMenu();
+	static void RegisterMenuBarEmptyContextMenu();
 	static void RegisterEmptySelectionContextMenu();
 
 	static void AddPlayFromHereSubMenu(FToolMenuSection& Section);
diff --git a/Engine/Source/Editor/LevelEditor/Private/LevelEditorMenu.cpp b/Engine/Source/Editor/LevelEditor/Private/LevelEditorMenu.cpp
--- a/Engine/Source/Editor/LevelEditor/Private/LevelEditorMenu.cpp
+++ b/Engine/Source/Editor/LevelEditor/Private/LevelEditorMenu.cpp
@@ -11,6 +11,7 @@
 #include "EditorStyleSet.h"
 #include "LevelEditor.h"
 #include "LevelEditorActions.h"
+#include "LevelEditorContextMenu.h"
 #include "Interfaces/IMainFrameModule.h"
 #include "MRUFavoritesList.h"
 #include "Framework/Commands/GenericCommands.h"
@@ -252,6 +253,34 @@ void FLevelEditorMenu::RegisterLevelEditorMenus()
 				);
 
 			SelectEntry.InsertPosition = FToolMenuInsert("Help", EToolMenuInsertType::Before);
+
+			const FName LevelEditorName("LevelEditor");
+			FToolMenuEntry& ActionsEntry =
+				Section.AddSubMenu(
+					"Actions",
+					TAttribute<FText>::CreateLambda([LevelEditorName]()
+					{
+						FLevelEditorModule& LevelEditorModule = FModuleManager::LoadModuleChecked<FLevelEditorModule>(LevelEditorName);
+						return LevelEditorModule.GetLevelEditorInstance().Pin()->GetLevelViewportContextMenuTitle();
+					}),
+					TAttribute<FText>::CreateLambda([LevelEditorName]()
+					{
+						FLevelEditorModule& LevelEditorModule = FModuleManager::LoadModuleChecked<FLevelEditorModule>(LevelEditorName);
+						return FLevelEditorContextMenu::GetContextMenuToolTip(
+							ELevelEditorMenuContext::MainMenu, LevelEditorModule.GetLevelEditorInstance().Pin()->GetElementSelectionSet());
+					}),
+					FOnGetContent::CreateLambda([LevelEditorName]()
+					{
+						// Generate the context menu completely separate from the main menu hierarchy for consistency with the right-click context menu.
+						// This means that extenders/UToolMenu extensions registered for the viewport context menu apply here (since they'll take effect when generating the menu widget below),
+						// and NOT any extenders registered for the main menu bar.
+						// I have verified that this works properly with the global Mac menu bar.
+						FLevelEditorModule& LevelEditorModule = FModuleManager::LoadModuleChecked<FLevelEditorModule>(LevelEditorName);
+						return FLevelEditorContextMenu::BuildMenuWidget(
+							LevelEditorModule.GetLevelEditorInstance(), ELevelEditorMenuContext::MainMenu, nullptr, FTypedElementHandle());
+					}));
+			
+			ActionsEntry.InsertPosition = FToolMenuInsert("Help", EToolMenuInsertType::Before);
 		}
 
 
diff --git a/Engine/Source/Editor/LevelEditor/Private/SLevelEditor.cpp b/Engine/Source/Editor/LevelEditor/Private/SLevelEditor.cpp
--- a/Engine/Source/Editor/LevelEditor/Private/SLevelEditor.cpp
+++ b/Engine/Source/Editor/LevelEditor/Private/SLevelEditor.cpp
@@ -1608,6 +1608,11 @@ void SLevelEditor::SummonLevelViewportContextMenu(const FTypedElementHandle& Hit
 	FLevelEditorContextMenu::SummonMenu( SharedThis( this ), ELevelEditorMenuContext::Viewport, HitProxyElement);
 }
 
+FText SLevelEditor::GetLevelViewportContextMenuTitle() const
+{
+	return CachedViewportContextMenuTitle;
+}
+
 void SLevelEditor::SummonLevelViewportViewOptionMenu(const ELevelViewportType ViewOption)
 {
 	FLevelEditorContextMenu::SummonViewOptionMenu( SharedThis( this ), ViewOption);
@@ -1721,6 +1726,16 @@ void SLevelEditor::OnElementSelectionChanged(const UTypedElementSelectionSet* Se
 		}
 	}
 
+	CachedViewportContextMenuTitle = FLevelEditorContextMenu::GetContextMenuTitle(ELevelEditorMenuContext::MainMenu, SelectionSet);
+
+#if PLATFORM_MAC
+	// The titles of the level editor's main menus can change when the selection changes (e.g. Action menu).
+	// Since Mac uses a global menu bar, we need to force it to sync with the main menu.
+	FLevelEditorModule& LevelEditorModule = FModuleManager::GetModuleChecked<FLevelEditorModule>(LevelEditorModuleName);
+	TSharedPtr<FTabManager> LevelEditorTabManager = LevelEditorModule.GetLevelEditorTabManager();
+	LevelEditorTabManager->UpdateMainMenu(nullptr, true);
+#endif
+
 	bNeedsRefresh = false;
 }
 
diff --git a/Engine/Source/Editor/LevelEditor/Private/SLevelEditor.h b/Engine/Source/Editor/LevelEditor/Private/SLevelEditor.h
--- a/Engine/Source/Editor/LevelEditor/Private/SLevelEditor.h
+++ b/Engine/Source/Editor/LevelEditor/Private/SLevelEditor.h
@@ -69,6 +69,7 @@ public:
 	virtual const UTypedElementSelectionSet* GetElementSelectionSet() const override;
 	virtual UTypedElementSelectionSet* GetMutableElementSelectionSet() override;
 	virtual void SummonLevelViewportContextMenu(const FTypedElementHandle& HitProxyElement = FTypedElementHandle()) override;
+	virtual FText GetLevelViewportContextMenuTitle() const override;
 	virtual void SummonLevelViewportViewOptionMenu(const ELevelViewportType ViewOption) override;
 	virtual const TArray< TSharedPtr< class IToolkit > >& GetHostedToolkits() const override;
 	virtual TArray< TSharedPtr< SLevelViewport > > GetViewports() const override;
@@ -331,4 +332,6 @@ private:
 
 	TSharedPtr<FLevelEditorModeUILayer> ModeUILayer;
 
+	/** Viewport context menu title, cached each time the element selection set changes */
+	FText CachedViewportContextMenuTitle;
 };
diff --git a/Engine/Source/Editor/LevelEditor/Public/ILevelEditor.h b/Engine/Source/Editor/LevelEditor/Public/ILevelEditor.h
--- a/Engine/Source/Editor/LevelEditor/Public/ILevelEditor.h
+++ b/Engine/Source/Editor/LevelEditor/Public/ILevelEditor.h
@@ -47,6 +47,9 @@ public:
 		SummonLevelViewportContextMenu(HitProxyElement);
 	}
 
+	/** Gets the title for the context menu for this level editor */
+	virtual FText GetLevelViewportContextMenuTitle() const = 0;
+
 	/** Summons a context menu for view options */
 	virtual void SummonLevelViewportViewOptionMenu(ELevelViewportType ViewOption) = 0;
 
diff --git a/Engine/Source/Editor/LevelEditor/Public/LevelEditorMenuContext.h b/Engine/Source/Editor/LevelEditor/Public/LevelEditorMenuContext.h
--- a/Engine/Source/Editor/LevelEditor/Public/LevelEditorMenuContext.h
+++ b/Engine/Source/Editor/LevelEditor/Public/LevelEditorMenuContext.h
@@ -9,6 +9,7 @@
 #include "LevelEditorMenuContext.generated.h"
 
 class SLevelEditor;
+class ILevelEditor;
 class UActorComponent;
 class SLevelViewport;
 class SLevelViewportToolBar;
@@ -28,10 +29,12 @@ public:
 UENUM()
 enum class ELevelEditorMenuContext : uint8
 {
-	/** This context menu is applicable to a viewport */
+	/** This context menu is applicable to a viewport (limited subset of entries) */
 	Viewport,
 	/** This context menu is applicable to the Scene Outliner (disables click-position-based menu items) */
 	SceneOutliner,
+	/** This is the replica of the context menu that appears in the main menu bar (lists all entries) */
+	MainMenu,
 };
 
 UCLASS()
@@ -40,7 +43,7 @@ class LEVELEDITOR_API ULevelEditorContextMenuContext : public UObject
 	GENERATED_BODY()
 public:
 
-	TWeakPtr<SLevelEditor> LevelEditor;
+	TWeakPtr<ILevelEditor> LevelEditor;
 
 	UPROPERTY(VisibleAnywhere, BlueprintReadOnly, Category="LevelEditor|Menu")
 	ELevelEditorMenuContext ContextType;
diff --git a/Engine/Source/Runtime/Slate/Private/Framework/MultiBox/MultiBoxBuilder.cpp b/Engine/Source/Runtime/Slate/Private/Framework/MultiBox/MultiBoxBuilder.cpp
--- a/Engine/Source/Runtime/Slate/Private/Framework/MultiBox/MultiBoxBuilder.cpp
+++ b/Engine/Source/Runtime/Slate/Private/Framework/MultiBox/MultiBoxBuilder.cpp
@@ -357,7 +357,7 @@ void FMenuBuilder::ApplySectionBeginning()
 	}
 }
 
-void FMenuBarBuilder::AddPullDownMenu(const FText& InMenuLabel, const FText& InToolTip, const FNewMenuDelegate& InPullDownMenu, FName InExtensionHook, FName InTutorialHighlightName)
+void FMenuBarBuilder::AddPullDownMenu(const TAttribute<FText>& InMenuLabel, const TAttribute<FText>& InToolTip, const FNewMenuDelegate& InPullDownMenu, FName InExtensionHook, FName InTutorialHighlightName)
 {
 	ApplySectionBeginning();
 
@@ -375,6 +375,24 @@ void FMenuBarBuilder::AddPullDownMenu(const FText& InMenuLabel, const FText& InT
 	ApplyHook(InExtensionHook, EExtensionHook::After);
 }
 
+void FMenuBarBuilder::AddPullDownMenu(const TAttribute<FText>& InMenuLabel, const TAttribute<FText>& InToolTip, const FOnGetContent& InMenuContentGenerator, FName InExtensionHook, FName InTutorialHighlightName)
+{
+	ApplySectionBeginning();
+
+	ApplyHook(InExtensionHook, EExtensionHook::Before);
+
+	const bool bIsSubMenu = false;
+	const bool bOpenSubMenuOnClick = false;
+	// Pulldown menus always close all menus not just themselves
+	const bool bShouldCloseSelfOnly = false;
+	TSharedRef< FMenuEntryBlock > NewMenuEntryBlock(new FMenuEntryBlock(InExtensionHook, InMenuLabel, InToolTip, InMenuContentGenerator, ExtenderStack.Top(), bIsSubMenu, bOpenSubMenuOnClick, CommandListStack.Last(), bShouldCloseSelfOnly));
+	NewMenuEntryBlock->SetTutorialHighlightName(GenerateTutorialIdentfierName(TutorialHighlightName, InTutorialHighlightName, nullptr, MultiBox->GetBlocks().Num()));
+
+	MultiBox->AddMultiBlock(NewMenuEntryBlock);
+
+	ApplyHook(InExtensionHook, EExtensionHook::After);
+}
+
 void FMenuBarBuilder::ApplyHook(FName InExtensionHook, EExtensionHook::Position HookPosition)
 {
 	if (ExtendersEnabled())
diff --git a/Engine/Source/Runtime/Slate/Public/Framework/MultiBox/MultiBoxBuilder.h b/Engine/Source/Runtime/Slate/Public/Framework/MultiBox/MultiBoxBuilder.h
--- a/Engine/Source/Runtime/Slate/Public/Framework/MultiBox/MultiBoxBuilder.h
+++ b/Engine/Source/Runtime/Slate/Public/Framework/MultiBox/MultiBoxBuilder.h
@@ -384,8 +384,17 @@ public:
 	 * @param	InToolTip			The tooltip that should be shown when the menu is hovered over
 	 * @param	InPullDownMenu		Pull down menu object for the menu to add.
 	 */
-	void AddPullDownMenu( const FText& InMenuLabel, const FText& InToolTip, const FNewMenuDelegate& InPullDownMenu, FName InExtensionHook = NAME_None, FName InTutorialHighlightName = NAME_None);
-	
+	void AddPullDownMenu( const TAttribute<FText>& InMenuLabel, const TAttribute<FText>& InToolTip, const FNewMenuDelegate& InPullDownMenu, FName InExtensionHook = NAME_None, FName InTutorialHighlightName = NAME_None);
+
+	/**
+	 * Adds a pull down menu
+	 *
+	 * @param	InMenuLabel				The text that should be shown for the menu
+	 * @param	InToolTip				The tooltip that should be shown when the menu is hovered over
+	 * @param	InMenuContentGenerator	Delegate that generates a widget for this pulldown menu's content.  Called when the menu is summoned.
+	 */
+	void AddPullDownMenu(const TAttribute<FText>& InMenuLabel, const TAttribute<FText>& InToolTip, const FOnGetContent& InMenuContentGenerator, FName InExtensionHook = NAME_None, FName InTutorialHighlightName = NAME_None);
+
 protected:
 	/** FMultiBoxBuilder interface */
 	virtual void ApplyHook(FName InExtensionHook, EExtensionHook::Position HookPosition) override;
