Fix high resolution screenshot tool performance issue.
 * execute r.ResetRenderTargetsExtent when the tool window is closed.

diff --git a/Engine/Source/Editor/LevelEditor/Private/HighResScreenshotUI.cpp b/Engine/Source/Editor/LevelEditor/Private/HighResScreenshotUI.cpp
--- a/Engine/Source/Editor/LevelEditor/Private/HighResScreenshotUI.cpp
+++ b/Engine/Source/Editor/LevelEditor/Private/HighResScreenshotUI.cpp
@@ -12,6 +12,10 @@
 #include "EditorStyleSet.h"
 #include "Widgets/Input/SNumericEntryBox.h"
 
+#if !(UE_BUILD_SHIPPING || UE_BUILD_TEST)
+#include "Kismet/KismetSystemLibrary.h"
+#endif
+
 TWeakPtr<class SWindow> SHighResScreenshotDialog::CurrentWindow = NULL;
 TWeakPtr<class SHighResScreenshotDialog> SHighResScreenshotDialog::CurrentDialog = NULL;
 bool SHighResScreenshotDialog::bMaskVisualizationWasEnabled = false;
@@ -259,6 +263,7 @@ void SHighResScreenshotDialog::WindowClosedHandler(const TSharedRef<SWindow>& In
 	FHighResScreenshotConfig& Config = GetHighResScreenshotConfig();
 
 	ResetViewport();
+	ResetFrameBuffer();
 
 	// Cleanup the config after each usage as it is a static and we don't want it to keep pointers or settings around between runs.
 	Config.bDisplayCaptureRegion = false;
@@ -294,6 +299,21 @@ void SHighResScreenshotDialog::ResetViewport()
 	}
 }
 
+void SHighResScreenshotDialog::ResetFrameBuffer()
+{
+#if !(UE_BUILD_SHIPPING || UE_BUILD_TEST)
+	FHighResScreenshotConfig& Config = GetHighResScreenshotConfig();
+	if (Config.TargetViewport.IsValid())
+	{
+		FViewportClient* Client = Config.TargetViewport.Pin()->GetClient();
+		if (Client)
+		{
+			UKismetSystemLibrary::ExecuteConsoleCommand(Client->GetWorld(), TEXT("r.ResetRenderTargetsExtent"), nullptr);
+		}
+	}
+#endif
+}
+
 TWeakPtr<class SWindow> SHighResScreenshotDialog::OpenDialog(const TSharedPtr<FSceneViewport>& InViewport, TSharedPtr<SCaptureRegionWidget> InCaptureRegionWidget)
 {
 	FHighResScreenshotConfig& Config = GetHighResScreenshotConfig();
diff --git a/Engine/Source/Editor/LevelEditor/Private/HighresScreenshotUI.h b/Engine/Source/Editor/LevelEditor/Private/HighresScreenshotUI.h
--- a/Engine/Source/Editor/LevelEditor/Private/HighresScreenshotUI.h
+++ b/Engine/Source/Editor/LevelEditor/Private/HighresScreenshotUI.h
@@ -203,6 +203,7 @@ private:
 	static void WindowClosedHandler(const TSharedRef<SWindow>& InWindow);
 
 	static void ResetViewport();
+	static void ResetFrameBuffer();
 
 	TSharedPtr<SWindow> Window;
 	TSharedPtr<class SCaptureRegionWidget> CaptureRegionWidget;
