High Resolution Screencapture Tool Reskin

diff --git a/Engine/Source/Editor/EditorStyle/Private/StarshipStyle.cpp b/Engine/Source/Editor/EditorStyle/Private/StarshipStyle.cpp
--- a/Engine/Source/Editor/EditorStyle/Private/StarshipStyle.cpp
+++ b/Engine/Source/Editor/EditorStyle/Private/StarshipStyle.cpp
@@ -415,14 +415,17 @@ void FStarshipEditorStyle::FStyle::SetupGeneralStyles()
 
 	// Common icons
 	{
-		Set( "Icons.Contact", new IMAGE_BRUSH( "Icons/icon_mail_16x", Icon16x16 ) );
+		Set("Icons.Contact", new IMAGE_BRUSH( "Icons/icon_mail_16x", Icon16x16 ) );
+		Set("Icons.Crop", new IMAGE_BRUSH_SVG("Starship/Common/Crop", Icon16x16));
+		Set("Icons.Fullscreen", new IMAGE_BRUSH_SVG("Starship/Common/EnableFullscreen", Icon16x16));
+		Set("Icons.Save", new IMAGE_BRUSH_SVG( "Starship/Common/SaveCurrent", Icon16x16 ) );
 
-		Set( "Icons.DirtyBadge", new IMAGE_BRUSH_SVG("Starship/Common/DirtyBadge", Icon12x12));
-		Set( "Icons.MakeStaticMesh", new IMAGE_BRUSH_SVG("Starship/Common/MakeStaticMesh", Icon16x16));
-		Set( "Icons.Documentation", new IMAGE_BRUSH_SVG("Starship/Common/Documentation", Icon16x16));
-		Set( "Icons.Support", new IMAGE_BRUSH_SVG("Starship/Common/Support", Icon16x16));
-		Set( "Icons.Package", new IMAGE_BRUSH_SVG("Starship/Common/ProjectPackage", Icon16x16));
-		Set( "Icons.Comment", new IMAGE_BRUSH_SVG("Starship/Common/Comment", Icon16x16));
+		Set("Icons.DirtyBadge", new IMAGE_BRUSH_SVG("Starship/Common/DirtyBadge", Icon12x12));
+		Set("Icons.MakeStaticMesh", new IMAGE_BRUSH_SVG("Starship/Common/MakeStaticMesh", Icon16x16));
+		Set("Icons.Documentation", new IMAGE_BRUSH_SVG("Starship/Common/Documentation", Icon16x16));
+		Set("Icons.Support", new IMAGE_BRUSH_SVG("Starship/Common/Support", Icon16x16));
+		Set("Icons.Package", new IMAGE_BRUSH_SVG("Starship/Common/ProjectPackage", Icon16x16));
+		Set("Icons.Comment", new IMAGE_BRUSH_SVG("Starship/Common/Comment", Icon16x16));
 
 	}
 
@@ -1288,8 +1291,7 @@ void FStarshipEditorStyle::FStyle::SetupGeneralStyles()
 
 	// Main frame
 	{
-		Set( "MainFrame.AutoSaveImage", new IMAGE_BRUSH( "Icons/icon_Autosave", Icon24x24 ) );
-
+		Set( "MainFrame.AutoSaveImage", 	   new IMAGE_BRUSH_SVG( "Starship/Common/SaveCurrent", Icon16x16 ) );
 		Set( "MainFrame.SaveAll",              new IMAGE_BRUSH_SVG( "Starship/Common/SaveAll", Icon16x16 ) );
 		Set( "MainFrame.ChoosePackagesToSave", new IMAGE_BRUSH_SVG( "Starship/Common/icon_file_choosepackages_16px", Icon16x16 ) );
 		Set( "MainFrame.NewProject",           new IMAGE_BRUSH_SVG( "Starship/Common/ProjectNew", Icon16x16 ) );
@@ -1731,17 +1733,6 @@ void FStarshipEditorStyle::FStyle::SetupGeneralStyles()
 	}
 
 	
-	// Highres Screenshot
-	{
-		Set("HighresScreenshot.WarningStrip", new IMAGE_BRUSH( "Common/WarningStripe", FVector2D(20,6), FLinearColor::White, ESlateBrushTileType::Horizontal ) );
-		Set("HighresScreenshot.SpecifyCaptureRectangle", new IMAGE_BRUSH( "Icons/icon_CaptureRegion_24x", Icon24x24 ) );
-		Set("HighresScreenshot.FullViewportCaptureRegion", new IMAGE_BRUSH( "Icons/icon_CaptureRegion_FullViewport_24x", Icon24x24 ) );
-		Set("HighresScreenshot.CameraSafeAreaCaptureRegion", new IMAGE_BRUSH("Icons/icon_CaptureRegion_Camera_Safe_24x", Icon24x24));
-		Set("HighresScreenshot.Capture", new IMAGE_BRUSH( "Icons/icon_HighResScreenshotCapture_24px", Icon24x24 ) );
-		Set("HighresScreenshot.AcceptCaptureRegion", new IMAGE_BRUSH( "Icons/icon_CaptureRegionAccept_24x", Icon24x24 ) );
-		Set("HighresScreenshot.DiscardCaptureRegion", new IMAGE_BRUSH( "Icons/icon_CaptureRegionDiscard_24x", Icon24x24 ) );
-	}
-
 	// Scalability 
 	{
 		const float Tint = 0.65f;
diff --git a/Engine/Source/Editor/LevelEditor/Private/HighResScreenshotUI.cpp b/Engine/Source/Editor/LevelEditor/Private/HighResScreenshotUI.cpp
--- a/Engine/Source/Editor/LevelEditor/Private/HighResScreenshotUI.cpp
+++ b/Engine/Source/Editor/LevelEditor/Private/HighResScreenshotUI.cpp
@@ -7,10 +7,11 @@
 #include "Widgets/Images/SImage.h"
 #include "Widgets/Layout/SGridPanel.h"
 #include "Widgets/Input/SButton.h"
-#include "Widgets/Layout/SSplitter.h"
-#include "Widgets/Input/SSlider.h"
+#include "Widgets/Layout/SSeparator.h"
+#include "Widgets/Input/SSpinBox.h"
+#include "SWarningOrErrorBox.h"
+#include "Styling/StyleColors.h"
 #include "EditorStyleSet.h"
-#include "Widgets/Input/SNumericEntryBox.h"
 
 #if !(UE_BUILD_SHIPPING || UE_BUILD_TEST)
 #include "Kismet/KismetSystemLibrary.h"
@@ -27,231 +28,217 @@ SHighResScreenshotDialog::SHighResScreenshotDialog()
 
 void SHighResScreenshotDialog::Construct( const FArguments& InArgs )
 {
+
+	FMargin GridPadding(6.f, 3.f);
 	this->ChildSlot
+	.Padding(0.f)
+	[
+		SNew(SBorder)
+		.Padding(0.f)
+		.BorderImage(FAppStyle::Get().GetBrush("ToolPanel.GroupBorder"))
 		[
 			SNew(SVerticalBox)
 			+SVerticalBox::Slot()
+			.AutoHeight()
 			[
-				SNew(SBorder)
-				.BorderImage(FEditorStyle::GetBrush("ToolPanel.GroupBorder"))
+				SNew(SHorizontalBox)
+				+SHorizontalBox::Slot()
+				.AutoWidth()
+				.Padding(6.f)
+				[
+					SAssignNew(CaptureRegionButton, SButton)
+					.ButtonStyle(&FAppStyle::Get().GetWidgetStyle<FButtonStyle>("SimpleButton"))
+					.IsEnabled(this, &SHighResScreenshotDialog::IsCaptureRegionEditingAvailable)
+					.Visibility(this, &SHighResScreenshotDialog::GetSpecifyCaptureRegionVisibility)
+					.ToolTipText(NSLOCTEXT("HighResScreenshot", "ScreenshotSpecifyCaptureRectangleTooltip", "Specify the region which will be captured by the screenshot"))
+					.OnClicked(this, &SHighResScreenshotDialog::OnSelectCaptureRegionClicked)
+					.ContentPadding(4.f)
+					[
+						SNew(SImage)
+						.ColorAndOpacity(FSlateColor::UseForeground())
+						.Image(FAppStyle::Get().GetBrush("Icons.Crop"))
+					]
+				]
+				+SHorizontalBox::Slot()
+				.AutoWidth()
+				.Padding(6.f)
+				[
+					SNew( SButton )
+					.ButtonStyle(&FAppStyle::Get().GetWidgetStyle<FButtonStyle>("SimpleButton"))
+					.Visibility(this, &SHighResScreenshotDialog::GetCaptureRegionControlsVisibility)
+					.ToolTipText( NSLOCTEXT("HighResScreenshot", "ScreenshotAcceptCaptureRegionTooltip", "Accept any changes made to the capture region") )
+					.OnClicked( this, &SHighResScreenshotDialog::OnSelectCaptureAcceptRegionClicked )
+					.ContentPadding(4.f)
+					[
+						SNew(SImage)
+						.ColorAndOpacity(FSlateColor::UseForeground())
+						.Image(FAppStyle::Get().GetBrush("Icons.Check"))
+					]
+				]
+				+SHorizontalBox::Slot()
+				.AutoWidth()
+				.Padding(6.f)
+				[
+					SNew( SButton )
+					.ButtonStyle(&FAppStyle::Get().GetWidgetStyle<FButtonStyle>("SimpleButton"))
+					.ToolTipText( NSLOCTEXT("HighResScreenshot", "ScreenshotDiscardCaptureRegionTooltip", "Discard any changes made to the capture region") )
+					.Visibility(this, &SHighResScreenshotDialog::GetCaptureRegionControlsVisibility)
+					.OnClicked( this, &SHighResScreenshotDialog::OnSelectCaptureCancelRegionClicked )
+					.ContentPadding(4.f)
+					[
+						SNew(SImage)
+						.ColorAndOpacity(FSlateColor::UseForeground())
+						.Image(FAppStyle::Get().GetBrush("Icons.X"))
+					]
+				]
+				+SHorizontalBox::Slot()
+				.AutoWidth()
+				.Padding(6.f)
 				[
-					SNew(SVerticalBox)
-					+SVerticalBox::Slot()
-					.Padding(5)
+					SNew( SButton )
+					.ButtonStyle(&FAppStyle::Get().GetWidgetStyle<FButtonStyle>("SimpleButton"))
+					.ToolTipText( NSLOCTEXT("HighResScreenshot", "ScreenshotFullViewportCaptureRegionTooltip", "Set the capture rectangle to the whole viewport") )
+					.Visibility(this, &SHighResScreenshotDialog::GetCaptureRegionControlsVisibility)
+					.OnClicked( this, &SHighResScreenshotDialog::OnSetFullViewportCaptureRegionClicked )
+					.ContentPadding(4.f)
 					[
-						SNew(SSplitter)
-						.Orientation(Orient_Horizontal)
-						+SSplitter::Slot()
-						.Value(1)
-						[
-							SNew(SVerticalBox)
-							+SVerticalBox::Slot()
-							.VAlign(VAlign_Center)
-							[
-								SNew( STextBlock )
-								.Text( NSLOCTEXT("HighResScreenshot", "ScreenshotSizeMultiplier", "Screenshot Size Multiplier") )
-							]
-							+SVerticalBox::Slot()
-							.VAlign(VAlign_Center)
-							[
-								SNew( STextBlock )
-								.Text( NSLOCTEXT("HighResScreenshot", "UseDateTimeAsImageName", "Use Date & Timestamp as Image name") )
-							]
-							+ SVerticalBox::Slot()
-							.VAlign(VAlign_Center)
-							[
-								SNew(STextBlock)
-								.Text(NSLOCTEXT("HighResScreenshot", "IncludeBufferVisTargets", "Include Buffer Visualization Targets"))
-							]
-							+ SVerticalBox::Slot()
-							.VAlign(VAlign_Center)
-							[
-								SAssignNew(HDRLabel, STextBlock)
-								.Text(NSLOCTEXT("HighResScreenshot", "CaptureHDR", "Write HDR format visualization targets"))
-							]
-							+SVerticalBox::Slot()
-							.VAlign(VAlign_Center)
-							[
-								SAssignNew(Force128BitRenderingLabel, STextBlock)
-								.Text(NSLOCTEXT("HighResScreenshot", "Force128BitPipeline", "Force 128-bit buffers for rendering pipeline"))
-							]
-							+SVerticalBox::Slot()
-							.VAlign(VAlign_Center)
-							[
-								SNew( STextBlock )
-								.Text( NSLOCTEXT("HighResScreenshot", "UseCustomDepth", "Use custom depth as mask") )
-							]
-						]
-						+SSplitter::Slot()
-						.Value(1)
-						[
-							SNew(SVerticalBox)
-							+SVerticalBox::Slot()
-							.VAlign(VAlign_Center)
-							[
-								SNew (SHorizontalBox)
-								+SHorizontalBox::Slot()
-								.FillWidth(1)
-								[
-									SNew( SNumericEntryBox<float> )
-									.Value(this, &SHighResScreenshotDialog::GetResolutionMultiplier)
-									.OnValueCommitted(this, &SHighResScreenshotDialog::OnResolutionMultiplierChanged)
-								]
-								+SHorizontalBox::Slot()
-								.HAlign(HAlign_Fill)
-								.Padding(5,0,0,0)
-								.FillWidth(3)
-								[
-									SNew( SSlider )
-									.Value(this, &SHighResScreenshotDialog::GetResolutionMultiplierSlider)
-									.OnValueChanged(this, &SHighResScreenshotDialog::OnResolutionMultiplierSliderChanged)
-								]
-							]
-							+ SVerticalBox::Slot()
-							.VAlign(VAlign_Center)
-							[
-								SNew(SCheckBox)
-								.OnCheckStateChanged(this, &SHighResScreenshotDialog::OnDateTimeBasedNamingEnabledChanged)
-								.IsChecked(this, &SHighResScreenshotDialog::GetDateTimeBasedNamingEnabled)
-							]
-							+SVerticalBox::Slot()
-							.VAlign(VAlign_Center)
-							[
-								SNew( SCheckBox )
-								.OnCheckStateChanged(this, &SHighResScreenshotDialog::OnBufferVisualizationDumpEnabledChanged)
-								.IsChecked(this, &SHighResScreenshotDialog::GetBufferVisualizationDumpEnabled)
-							]
-							+ SVerticalBox::Slot()
-							.VAlign(VAlign_Center)
-							[
-								SAssignNew(HDRCheckBox, SCheckBox)
-								.OnCheckStateChanged(this, &SHighResScreenshotDialog::OnHDREnabledChanged)
-								.IsChecked(this, &SHighResScreenshotDialog::GetHDRCheckboxUIState)
-							]
-							+ SVerticalBox::Slot()
-							.VAlign(VAlign_Center)
-							[
-								SAssignNew(Force128BitRenderingCheckBox, SCheckBox)
-								.OnCheckStateChanged(this, &SHighResScreenshotDialog::OnForce128BitRenderingChanged)
-								.IsChecked(this, &SHighResScreenshotDialog::GetForce128BitRenderingCheckboxUIState)
-							]
-							+SVerticalBox::Slot()
-							.VAlign(VAlign_Center)
-							[
-								SNew( SCheckBox )
-								.OnCheckStateChanged(this, &SHighResScreenshotDialog::OnMaskEnabledChanged)
-								.IsChecked(this, &SHighResScreenshotDialog::GetMaskEnabled)
-							]
-						]
+						SNew(SImage)
+						.ColorAndOpacity(FSlateColor::UseForeground())
+						.Image(FAppStyle::Get().GetBrush("Icons.Fullscreen"))
 					]
 				]
 			]
+
 			+SVerticalBox::Slot()
 			.AutoHeight()
 			[
-				SNew(SVerticalBox)
-				+SVerticalBox::Slot()
-				.AutoHeight()
+				SNew(SSeparator)
+				.Thickness(1.f)
+				.SeparatorImage(FAppStyle::Get().GetBrush("Brushes.Background"))
+			]
+
+			+SVerticalBox::Slot()
+			.Padding(6.0)
+			.HAlign(HAlign_Fill)
+			[
+				// Row/Column
+				SNew(SGridPanel) 
+				.FillColumn(1, 1.0f)
+
+				+SGridPanel::Slot(0, 0)
+				.Padding(GridPadding)
+				.VAlign(VAlign_Center)
 				[
-					SNew(SImage)
-					.Image(FEditorStyle::GetBrush("HighresScreenshot.WarningStrip"))
+					SNew( STextBlock )
+					.Text( NSLOCTEXT("HighResScreenshot", "ScreenshotSizeMultiplier", "Screenshot Size Multiplier") )
 				]
-				+SVerticalBox::Slot()
-				.AutoHeight()
+				+SGridPanel::Slot(0, 1)
+				.Padding(GridPadding)
 				[
-					SNew( SBorder )
-					.BorderImage(FEditorStyle::GetBrush("ToolPanel.GroupBorder"))
-					[
-						SNew( STextBlock )
-						.Text( NSLOCTEXT("HighResScreenshot", "CaptureWarningText", "Due to the high system requirements of a high resolution screenshot, very large multipliers might cause the graphics driver to become unresponsive and possibly crash. In these circumstances, please try using a lower multiplier") )
-						.AutoWrapText(true)
-					]
+					SNew( STextBlock )
+					.Text( NSLOCTEXT("HighResScreenshot", "UseDateTimeAsImageName", "Use Date & Timestamp as Image name") )
+				]
+				+ SGridPanel::Slot(0, 2)
+				.Padding(GridPadding)
+				[
+					SNew(STextBlock)
+					.Text(NSLOCTEXT("HighResScreenshot", "IncludeBufferVisTargets", "Include Buffer Visualization Targets"))
+				]
+				+ SGridPanel::Slot(0, 3)
+				.Padding(GridPadding)
+				[
+					SAssignNew(HDRLabel, STextBlock)
+					.Text(NSLOCTEXT("HighResScreenshot", "CaptureHDR", "Write HDR format visualization targets"))
+				]
+				+SGridPanel::Slot(0, 4)
+				.Padding(GridPadding)
+				[
+					SAssignNew(Force128BitRenderingLabel, STextBlock)
+					.Text(NSLOCTEXT("HighResScreenshot", "Force128BitPipeline", "Force 128-bit buffers for rendering pipeline"))
+				]
+				+SGridPanel::Slot(0, 5)
+				.Padding(GridPadding)
+				[
+					SNew( STextBlock )
+					.Text( NSLOCTEXT("HighResScreenshot", "UseCustomDepth", "Use custom depth as mask") )
+				]
+
+				+SGridPanel::Slot(1, 0)
+				.Padding(GridPadding)
+				.HAlign(HAlign_Fill)
+				[
+						SNew( SSpinBox<float> )
+						.MinValue(FHighResScreenshotConfig::MinResolutionMultipler)
+						.MaxValue(FHighResScreenshotConfig::MaxResolutionMultipler)
+						.Delta(1.0f)
+						.Value(this, &SHighResScreenshotDialog::GetResolutionMultiplierSlider)
+						.OnValueChanged(this, &SHighResScreenshotDialog::OnResolutionMultiplierSliderChanged)
+				]
+				+ SGridPanel::Slot(1, 1)
+				.Padding(GridPadding)
+				[
+					SNew(SCheckBox)
+					.OnCheckStateChanged(this, &SHighResScreenshotDialog::OnDateTimeBasedNamingEnabledChanged)
+					.IsChecked(this, &SHighResScreenshotDialog::GetDateTimeBasedNamingEnabled)
+				]
+				+SGridPanel::Slot(1, 2)
+				.Padding(GridPadding)
+				[
+					SNew( SCheckBox )
+					.OnCheckStateChanged(this, &SHighResScreenshotDialog::OnBufferVisualizationDumpEnabledChanged)
+					.IsChecked(this, &SHighResScreenshotDialog::GetBufferVisualizationDumpEnabled)
+				]
+				+ SGridPanel::Slot(1, 3)
+				.Padding(GridPadding)
+				[
+					SAssignNew(HDRCheckBox, SCheckBox)
+					.OnCheckStateChanged(this, &SHighResScreenshotDialog::OnHDREnabledChanged)
+					.IsChecked(this, &SHighResScreenshotDialog::GetHDRCheckboxUIState)
+				]
+				+ SGridPanel::Slot(1, 4)
+				.Padding(GridPadding)
+				[
+					SAssignNew(Force128BitRenderingCheckBox, SCheckBox)
+					.OnCheckStateChanged(this, &SHighResScreenshotDialog::OnForce128BitRenderingChanged)
+					.IsChecked(this, &SHighResScreenshotDialog::GetForce128BitRenderingCheckboxUIState)
 				]
-				+SVerticalBox::Slot()
-				.AutoHeight()
+				+SGridPanel::Slot(1, 5)
+				.Padding(GridPadding)
 				[
-					SNew(SImage)
-					.Image(FEditorStyle::GetBrush("HighresScreenshot.WarningStrip"))
+					SNew( SCheckBox )
+					.OnCheckStateChanged(this, &SHighResScreenshotDialog::OnMaskEnabledChanged)
+					.IsChecked(this, &SHighResScreenshotDialog::GetMaskEnabled)
 				]
 			]
+
 			+SVerticalBox::Slot()
 			.AutoHeight()
+			.Padding(16.f)
 			[
-				SNew(SBorder)
-				.BorderImage(FEditorStyle::GetBrush("ToolPanel.GroupBorder"))
+
+				SNew(SHorizontalBox)
+				+SHorizontalBox::Slot()
 				[
-					SNew(SHorizontalBox)
-					+SHorizontalBox::Slot()
-					.AutoWidth()
-					[
-						SNew(SGridPanel)
-						+SGridPanel::Slot(0, 0)
-						[
-							SAssignNew(CaptureRegionButton, SButton)
-							.IsEnabled(this, &SHighResScreenshotDialog::IsCaptureRegionEditingAvailable)
-							.Visibility(this, &SHighResScreenshotDialog::GetSpecifyCaptureRegionVisibility)
-							.ToolTipText(NSLOCTEXT("HighResScreenshot", "ScreenshotSpecifyCaptureRectangleTooltip", "Specify the region which will be captured by the screenshot"))
-							.OnClicked(this, &SHighResScreenshotDialog::OnSelectCaptureRegionClicked)
-							[
-								SNew(SImage)
-								.Image(FEditorStyle::GetBrush("HighresScreenshot.SpecifyCaptureRectangle"))
-							]
-						]
-						+SGridPanel::Slot(0, 0)
-						[
-							SNew( SButton )
-							.Visibility(this, &SHighResScreenshotDialog::GetCaptureRegionControlsVisibility)
-							.ToolTipText( NSLOCTEXT("HighResScreenshot", "ScreenshotAcceptCaptureRegionTooltip", "Accept any changes made to the capture region") )
-							.OnClicked( this, &SHighResScreenshotDialog::OnSelectCaptureAcceptRegionClicked )
-							[
-								SNew(SImage)
-								.Image(FEditorStyle::GetBrush("HighresScreenshot.AcceptCaptureRegion"))
-							]
-						]
-					]
-					+SHorizontalBox::Slot()
-					.HAlign(HAlign_Right)
-					.AutoWidth()
-					[
-						SNew( SButton )
-						.ToolTipText( NSLOCTEXT("HighResScreenshot", "ScreenshotDiscardCaptureRegionTooltip", "Discard any changes made to the capture region") )
-						.Visibility(this, &SHighResScreenshotDialog::GetCaptureRegionControlsVisibility)
-						.OnClicked( this, &SHighResScreenshotDialog::OnSelectCaptureCancelRegionClicked )
-						[
-							SNew(SImage)
-							.Image(FEditorStyle::GetBrush("HighresScreenshot.DiscardCaptureRegion"))
-						]
-					]
-					+SHorizontalBox::Slot()
-					.HAlign(HAlign_Right)
-					.AutoWidth()
-					[
-						SNew( SButton )
-						.ToolTipText( NSLOCTEXT("HighResScreenshot", "ScreenshotFullViewportCaptureRegionTooltip", "Set the capture rectangle to the whole viewport") )
-						.Visibility(this, &SHighResScreenshotDialog::GetCaptureRegionControlsVisibility)
-						.OnClicked( this, &SHighResScreenshotDialog::OnSetFullViewportCaptureRegionClicked )
-						[
-							SNew(SImage)
-							.Image(FEditorStyle::GetBrush("HighresScreenshot.FullViewportCaptureRegion"))
-						]
-					]
-					+SHorizontalBox::Slot()
-						// for padding
-					+SHorizontalBox::Slot()
-					.HAlign(HAlign_Right)
-					.AutoWidth()
-					[
-						SNew( SButton )
-						.ToolTipText( NSLOCTEXT("HighResScreenshot", "ScreenshotCaptureTooltop", "Take a screenshot") )
-						.OnClicked( this, &SHighResScreenshotDialog::OnCaptureClicked )
-						[
-							SNew(SImage)
-							.Image(FEditorStyle::GetBrush("HighresScreenshot.Capture"))
-						]
-					]
+					SNew(SWarningOrErrorBox)
+					.Padding(FMargin(16.f, 13.f, 16.f, 13.f))
+					.Message( NSLOCTEXT("HighResScreenshot", "CaptureWarningText", "Large multipliers may cause the graphics driver to crash.  Please try using a lower multiplier.") )
+				]
+
+				+SHorizontalBox::Slot()
+				.VAlign(VAlign_Bottom)
+				.HAlign(HAlign_Right)
+				.AutoWidth()
+				.Padding(FMargin(24.f, 0.f, 0.f, 0.f))
+				[
+					SNew( SButton )
+					.ButtonStyle(&FAppStyle::Get().GetWidgetStyle<FButtonStyle>("PrimaryButton"))
+					.ToolTipText(NSLOCTEXT("HighResScreenshot", "ScreenshotCaptureTooltop", "Take a screenshot") )
+					.OnClicked(this, &SHighResScreenshotDialog::OnCaptureClicked )
+					.Text(NSLOCTEXT("HighResScreenshot", "CaptureCommit", "Capture"))
 				]
 			]
-		];
+		]
+	];
 
 	SetHDRUIEnableState(Config.bDumpBufferVisualizationTargets);
 	SetForce128BitRenderingState(Config.bDumpBufferVisualizationTargets);
@@ -337,7 +324,7 @@ TWeakPtr<class SWindow> SHighResScreenshotDialog::OpenDialog(const TSharedPtr<FS
 		TSharedRef<SHighResScreenshotDialog> Dialog = SNew(SHighResScreenshotDialog);
 		TSharedRef<SWindow> Window = SNew(SWindow)
 			.Title( NSLOCTEXT("HighResScreenshot", "HighResolutionScreenshot", "High Resolution Screenshot") )
-			.ClientSize(FVector2D(484,231))
+			.ClientSize(FVector2D(480, 286))
 			.SupportsMinimize(false)
 			.SupportsMaximize(false)
 			.FocusWhenFirstShown(true)
diff --git a/Engine/Source/Editor/LevelEditor/Private/HighresScreenshotUI.h b/Engine/Source/Editor/LevelEditor/Private/HighresScreenshotUI.h
--- a/Engine/Source/Editor/LevelEditor/Private/HighresScreenshotUI.h
+++ b/Engine/Source/Editor/LevelEditor/Private/HighresScreenshotUI.h
@@ -75,13 +75,6 @@ private:
 
 	bool IsSetCameraSafeAreaCaptureRegionEnabled() const;
 
-	void OnResolutionMultiplierChanged( float NewValue, ETextCommit::Type CommitInfo )
-	{
-		NewValue = FMath::Clamp(NewValue, FHighResScreenshotConfig::MinResolutionMultipler, FHighResScreenshotConfig::MaxResolutionMultipler);
-		Config.ResolutionMultiplier = NewValue;
-		Config.ResolutionMultiplierScale = (NewValue - FHighResScreenshotConfig::MinResolutionMultipler) / (FHighResScreenshotConfig::MaxResolutionMultipler - FHighResScreenshotConfig::MinResolutionMultipler);
-	}
-
 	void OnResolutionMultiplierSliderChanged( float NewValue )
 	{
 		Config.ResolutionMultiplierScale = NewValue;
@@ -135,12 +128,12 @@ private:
 
 	EVisibility GetSpecifyCaptureRegionVisibility() const
 	{
-		return bCaptureRegionControlsVisible ? EVisibility::Hidden : EVisibility::Visible;
+		return bCaptureRegionControlsVisible ? EVisibility::Collapsed: EVisibility::Visible;
 	}
 
 	EVisibility GetCaptureRegionControlsVisibility() const
 	{
-		return bCaptureRegionControlsVisible ? EVisibility::Visible : EVisibility::Hidden;
+		return bCaptureRegionControlsVisible ? EVisibility::Visible : EVisibility::Collapsed;
 	}
 
 	void SetCaptureRegionControlsVisibility(bool bVisible)
@@ -148,11 +141,6 @@ private:
 		bCaptureRegionControlsVisible = bVisible;
 	}
 
-	TOptional<float> GetResolutionMultiplier() const
-	{
-		return TOptional<float>(Config.ResolutionMultiplier);
-	}
-
 	float GetResolutionMultiplierSlider() const
 	{
 		return Config.ResolutionMultiplierScale;
