misc non-unity compilation fixes

diff --git a/Engine/Source/Editor/LevelEditor/Private/HighResScreenshotUI.cpp b/Engine/Source/Editor/LevelEditor/Private/HighResScreenshotUI.cpp
--- a/Engine/Source/Editor/LevelEditor/Private/HighResScreenshotUI.cpp
+++ b/Engine/Source/Editor/LevelEditor/Private/HighResScreenshotUI.cpp
@@ -15,6 +15,7 @@
 
 #if !(UE_BUILD_SHIPPING || UE_BUILD_TEST)
 #include "Kismet/KismetSystemLibrary.h"
+#include "Engine/World.h"
 #endif
 
 TWeakPtr<class SWindow> SHighResScreenshotDialog::CurrentWindow = NULL;
