PR for adding support for multiple outliners

diff --git a/Engine/Source/Editor/LevelEditor/Private/LevelEditor.cpp b/Engine/Source/Editor/LevelEditor/Private/LevelEditor.cpp
--- a/Engine/Source/Editor/LevelEditor/Private/LevelEditor.cpp
+++ b/Engine/Source/Editor/LevelEditor/Private/LevelEditor.cpp
@@ -63,6 +63,9 @@ const FName LevelEditorTabIds::LevelEditorSelectionDetails4(TEXT("LevelEditorSel
 const FName LevelEditorTabIds::PlacementBrowser(TEXT("PlacementBrowser"));
 const FName LevelEditorTabIds::LevelEditorBuildAndSubmit(TEXT("LevelEditorBuildAndSubmit"));
 const FName LevelEditorTabIds::LevelEditorSceneOutliner(TEXT("LevelEditorSceneOutliner"));
+const FName LevelEditorTabIds::LevelEditorSceneOutliner2(TEXT("LevelEditorSceneOutliner2"));
+const FName LevelEditorTabIds::LevelEditorSceneOutliner3(TEXT("LevelEditorSceneOutliner3"));
+const FName LevelEditorTabIds::LevelEditorSceneOutliner4(TEXT("LevelEditorSceneOutliner4"));
 const FName LevelEditorTabIds::LevelEditorStatsViewer(TEXT("LevelEditorStatsViewer"));
 const FName LevelEditorTabIds::LevelEditorLayerBrowser(TEXT("LevelEditorLayerBrowser"));
 const FName LevelEditorTabIds::LevelEditorDataLayerBrowser(TEXT("LevelEditorDataLayerBrowser"));
diff --git a/Engine/Source/Editor/LevelEditor/Private/SLevelEditor.cpp b/Engine/Source/Editor/LevelEditor/Private/SLevelEditor.cpp
--- a/Engine/Source/Editor/LevelEditor/Private/SLevelEditor.cpp
+++ b/Engine/Source/Editor/LevelEditor/Private/SLevelEditor.cpp
@@ -730,7 +730,7 @@ TSharedRef<SDockTab> SLevelEditor::SpawnLevelEditorTab( const FSpawnTabArgs& Arg
 
 		return NewTab;
 	}
-	else if( TabIdentifier == LevelEditorTabIds::LevelEditorSceneOutliner)
+	else if (TabIdentifier == LevelEditorTabIds::LevelEditorSceneOutliner || TabIdentifier == LevelEditorTabIds::LevelEditorSceneOutliner2 || TabIdentifier == LevelEditorTabIds::LevelEditorSceneOutliner3 || TabIdentifier == LevelEditorTabIds::LevelEditorSceneOutliner4)
 	{
 		FSceneOutlinerInitializationOptions InitOptions;
 		InitOptions.bShowTransient = true;
@@ -1200,12 +1200,32 @@ TSharedRef<SWidget> SLevelEditor::RestoreContentArea( const TSharedRef<SDockTab>
 		}
 
 		{
+			const FText OutlinerTooltip = NSLOCTEXT("LevelEditorTabs", "LevelEditorSceneOutlinerTooltipText", "Open the Outliner tab, which provides a searchable and filterable list of all actors in the world.");
 			const FSlateIcon OutlinerIcon(FEditorStyle::GetStyleSetName(), "LevelEditor.Tabs.Outliner");
-		    LevelEditorTabManager->RegisterTabSpawner(LevelEditorTabIds::LevelEditorSceneOutliner, FOnSpawnTab::CreateSP<SLevelEditor, FName, FString>(this, &SLevelEditor::SpawnLevelEditorTab, LevelEditorTabIds::LevelEditorSceneOutliner, FString()) )
-				.SetDisplayName(NSLOCTEXT("LevelEditorTabs", "LevelEditorSceneOutliner", "Outliner"))
-				.SetTooltipText(NSLOCTEXT("LevelEditorTabs", "LevelEditorSceneOutlinerTooltipText", "Open the Outliner tab, which provides a searchable and filterable list of all actors in the world."))
-				.SetGroup( MenuStructure.GetLevelEditorCategory() )	
-				.SetIcon( OutlinerIcon );	
+
+			LevelEditorTabManager->RegisterTabSpawner(LevelEditorTabIds::LevelEditorSceneOutliner, FOnSpawnTab::CreateSP<SLevelEditor, FName, FString>(this, &SLevelEditor::SpawnLevelEditorTab, LevelEditorTabIds::LevelEditorSceneOutliner, FString()))
+				.SetDisplayName(NSLOCTEXT("LevelEditorTabs", "LevelEditorSceneOutliner", "Outliner 1"))
+				.SetTooltipText(OutlinerTooltip)
+				.SetGroup(MenuStructure.GetLevelEditorOutlinerCategory())
+				.SetIcon(OutlinerIcon);
+
+			LevelEditorTabManager->RegisterTabSpawner(LevelEditorTabIds::LevelEditorSceneOutliner2, FOnSpawnTab::CreateSP<SLevelEditor, FName, FString>(this, &SLevelEditor::SpawnLevelEditorTab, LevelEditorTabIds::LevelEditorSceneOutliner2, FString()))
+				.SetDisplayName(NSLOCTEXT("LevelEditorTabs", "LevelEditorSceneOutliner2", "Outliner 2"))
+				.SetTooltipText(OutlinerTooltip)
+				.SetGroup(MenuStructure.GetLevelEditorOutlinerCategory())
+				.SetIcon(OutlinerIcon);
+
+			LevelEditorTabManager->RegisterTabSpawner(LevelEditorTabIds::LevelEditorSceneOutliner3, FOnSpawnTab::CreateSP<SLevelEditor, FName, FString>(this, &SLevelEditor::SpawnLevelEditorTab, LevelEditorTabIds::LevelEditorSceneOutliner3, FString()))
+				.SetDisplayName(NSLOCTEXT("LevelEditorTabs", "LevelEditorSceneOutliner3", "Outliner 3"))
+				.SetTooltipText(OutlinerTooltip)
+				.SetGroup(MenuStructure.GetLevelEditorOutlinerCategory())
+				.SetIcon(OutlinerIcon);
+
+			LevelEditorTabManager->RegisterTabSpawner(LevelEditorTabIds::LevelEditorSceneOutliner4, FOnSpawnTab::CreateSP<SLevelEditor, FName, FString>(this, &SLevelEditor::SpawnLevelEditorTab, LevelEditorTabIds::LevelEditorSceneOutliner4, FString()))
+				.SetDisplayName(NSLOCTEXT("LevelEditorTabs", "LevelEditorSceneOutliner4", "Outliner 4"))
+				.SetTooltipText(OutlinerTooltip)
+				.SetGroup(MenuStructure.GetLevelEditorOutlinerCategory())
+				.SetIcon(OutlinerIcon);
 		}
 
 		{
@@ -1958,4 +1978,4 @@ void SLevelEditor::RemoveViewportOverlayWidget(TSharedRef<SWidget> InWidget, TSh
 }
 
 
-#undef LOCTEXT_NAMESPACE
+#undef LOCTEXT_NAMESPACE
\ No newline at end of file
diff --git a/Engine/Source/Editor/LevelEditor/Public/LevelEditor.h b/Engine/Source/Editor/LevelEditor/Public/LevelEditor.h
--- a/Engine/Source/Editor/LevelEditor/Public/LevelEditor.h
+++ b/Engine/Source/Editor/LevelEditor/Public/LevelEditor.h
@@ -42,6 +42,9 @@ public:
 	static const FName PlacementBrowser;
 	static const FName LevelEditorBuildAndSubmit;
 	static const FName LevelEditorSceneOutliner;
+	static const FName LevelEditorSceneOutliner2;
+	static const FName LevelEditorSceneOutliner3;
+	static const FName LevelEditorSceneOutliner4;
 	static const FName LevelEditorStatsViewer;
 	static const FName LevelEditorLayerBrowser;
 	static const FName LevelEditorDataLayerBrowser;
diff --git a/Engine/Source/Editor/WorkspaceMenuStructure/Private/WorkspaceMenuStructureModule.cpp b/Engine/Source/Editor/WorkspaceMenuStructure/Private/WorkspaceMenuStructureModule.cpp
--- a/Engine/Source/Editor/WorkspaceMenuStructure/Private/WorkspaceMenuStructureModule.cpp
+++ b/Engine/Source/Editor/WorkspaceMenuStructure/Private/WorkspaceMenuStructureModule.cpp
@@ -59,6 +59,11 @@ public:
 	{
 		return LevelEditorWorldPartitionCategory.ToSharedRef();
 	}
+		
+	virtual TSharedRef<FWorkspaceItem> GetLevelEditorOutlinerCategory() const override
+	{
+		return LevelEditorOutlinerCategory.ToSharedRef();
+	}
 
 	virtual TSharedRef<FWorkspaceItem> GetToolsCategory() const override
 	{
@@ -114,6 +119,7 @@ public:
 		LevelEditorVirtualProductionCategory = LevelEditorCategory->AddGroup(LOCTEXT("WorkspaceMenu_LevelEditorVirtualProductionCategory", "Virtual Production"), LOCTEXT("WorkspaceMenu_LevelEditorVirtualProductionCategoryTooltip", "Open a Virtual Production tab."), FSlateIcon(FEditorStyle::GetStyleSetName(), "LevelEditor.Tabs.VirtualProduction"), true);
 		LevelEditorModesCategory = LevelEditorCategory->AddGroup(LOCTEXT("WorkspaceMenu_LevelEditorToolsCategory", "Editor Modes" ), FSlateIcon(FEditorStyle::GetStyleSetName(), "LevelEditor.Tabs.EditorModes"), true );
 		LevelEditorWorldPartitionCategory = LevelEditorCategory->AddGroup(LOCTEXT("WorkspaceMenu_LevelEditorWorldPartitionCategory", "World Partition"), LOCTEXT("WorkspaceMenu_LevelEditorWorldPartitionCategoryTooltip", "Open a World Partition tab."), FSlateIcon(FEditorStyle::GetStyleSetName(), "LevelEditor.Tabs.WorldPartition"), true);
+		LevelEditorOutlinerCategory = LevelEditorCategory->AddGroup(LOCTEXT("WorkspaceMenu_LevelEditorOutlinerCategory", "Outliner"), LOCTEXT("WorkspaceMenu_LevelEditorOutlinerCategoryTooltip", "Open an Outliner tab."), FSlateIcon(FEditorStyle::GetStyleSetName(), "LevelEditor.Tabs.Outliner"), true);
 	}
 
 	void ResetToolsCategory()
@@ -160,6 +166,7 @@ private:
 	TSharedPtr<FWorkspaceItem> LevelEditorVirtualProductionCategory;
 	TSharedPtr<FWorkspaceItem> LevelEditorModesCategory;
 	TSharedPtr<FWorkspaceItem> LevelEditorWorldPartitionCategory;
+	TSharedPtr<FWorkspaceItem> LevelEditorOutlinerCategory;
 
 	TSharedPtr<FWorkspaceItem> ToolsCategory;
 	TSharedPtr<FWorkspaceItem> DeveloperToolsCategory;
diff --git a/Engine/Source/Editor/WorkspaceMenuStructure/Public/WorkspaceMenuStructure.h b/Engine/Source/Editor/WorkspaceMenuStructure/Public/WorkspaceMenuStructure.h
--- a/Engine/Source/Editor/WorkspaceMenuStructure/Public/WorkspaceMenuStructure.h
+++ b/Engine/Source/Editor/WorkspaceMenuStructure/Public/WorkspaceMenuStructure.h
@@ -51,6 +51,9 @@ public:
 
 	/** See diagram above */
 	virtual TSharedRef<FWorkspaceItem> GetLevelEditorWorldPartitionCategory() const = 0;
+	
+	/** See diagram above */
+	virtual TSharedRef<FWorkspaceItem> GetLevelEditorOutlinerCategory() const = 0;
 
 	/** See diagram above */
 	virtual TSharedRef<FWorkspaceItem> GetLevelEditorModesCategory() const = 0;
