Check that the LevelEditorInstance is valid

diff --git a/Engine/Source/Editor/LevelEditor/Private/LevelEditor.cpp b/Engine/Source/Editor/LevelEditor/Private/LevelEditor.cpp
--- a/Engine/Source/Editor/LevelEditor/Private/LevelEditor.cpp
+++ b/Engine/Source/Editor/LevelEditor/Private/LevelEditor.cpp
@@ -383,7 +383,7 @@ void FLevelEditorModule::SummonWorldBrowserComposition()
 void FLevelEditorModule::AttachSequencer(TSharedPtr<SWidget> SequencerWidget, TSharedPtr<IAssetEditorInstance> SequencerAssetEditor)
 {
 	TSharedPtr<SLevelEditor> LevelEditorInstance = LevelEditorInstancePtr.Pin();
-	if (LevelEditorInstance.IsValid())
+	if (LevelEditorInstance)
 	{
 		LevelEditorInstance->AttachSequencer(SequencerWidget, SequencerAssetEditor);
 	}
