Add ability to get status bar name from IToolkitHost interface. The mode UI layer classes now get their status bar name from their toolkit host.

#rnx

diff --git a/Engine/Source/Editor/EditorFramework/Public/Toolkits/AssetEditorModeUILayer.h b/Engine/Source/Editor/EditorFramework/Public/Toolkits/AssetEditorModeUILayer.h
--- a/Engine/Source/Editor/EditorFramework/Public/Toolkits/AssetEditorModeUILayer.h
+++ b/Engine/Source/Editor/EditorFramework/Public/Toolkits/AssetEditorModeUILayer.h
@@ -59,7 +59,7 @@ public:
 	};
 	virtual const FName GetStatusBarName() const
 	{
-		return NAME_None;
+		return ToolkitHost->GetStatusBarName();
 	}
 
 protected:
diff --git a/Engine/Source/Editor/EditorFramework/Public/Toolkits/IToolkitHost.h b/Engine/Source/Editor/EditorFramework/Public/Toolkits/IToolkitHost.h
--- a/Engine/Source/Editor/EditorFramework/Public/Toolkits/IToolkitHost.h
+++ b/Engine/Source/Editor/EditorFramework/Public/Toolkits/IToolkitHost.h
@@ -45,6 +45,10 @@ public:
 	/** Returns the common actions implementation for this toolkit host */
 	virtual UTypedElementCommonActions* GetCommonActions() const = 0;
 
+	/** Returns the status bar name that can be used with UStatusBarSubsystem to display messages 
+	    in the status bar of this toolkit host. Should return NAME_None if not applicable. */
+	virtual FName GetStatusBarName() const = 0;
+
 	/** Allows Toolkits to push widgets to the viewport.  Passing a nullptr for the Viewport will add  or 
 	    remove the OverlayWidget to or from the ActiveViewport */
 	virtual void AddViewportOverlayWidget(TSharedRef<SWidget> InOverlaidWidget, TSharedPtr<IAssetViewport> InViewport = nullptr)  {}
diff --git a/Engine/Source/Editor/LevelEditor/Private/SLevelEditor.cpp b/Engine/Source/Editor/LevelEditor/Private/SLevelEditor.cpp
--- a/Engine/Source/Editor/LevelEditor/Private/SLevelEditor.cpp
+++ b/Engine/Source/Editor/LevelEditor/Private/SLevelEditor.cpp
@@ -307,7 +307,7 @@ void SLevelEditor::Initialize( const TSharedRef<SDockTab>& OwnerTab, const TShar
 		.Padding(0.0f, 2.0f, 0.0f, 0.0f)
 		.AutoHeight()
 		[
-			GEditor->GetEditorSubsystem<UStatusBarSubsystem>()->MakeStatusBarWidget(TEXT("LevelEditor.StatusBar"), OwnerTab)
+			GEditor->GetEditorSubsystem<UStatusBarSubsystem>()->MakeStatusBarWidget(GetStatusBarName(), OwnerTab)
 		]
 	];
 	
@@ -1936,6 +1936,12 @@ UTypedElementCommonActions* SLevelEditor::GetCommonActions() const
 	return CommonActions;
 }
 
+FName SLevelEditor::GetStatusBarName() const
+{
+	static const FName LevelEditorStatusBarName = "LevelEditor.StatusBar";
+	return LevelEditorStatusBarName;
+}
+
 void SLevelEditor::AddViewportOverlayWidget(TSharedRef<SWidget> InWidget, TSharedPtr<IAssetViewport> InViewport)
 {
 	if (InViewport != nullptr)
diff --git a/Engine/Source/Editor/LevelEditor/Private/SLevelEditor.h b/Engine/Source/Editor/LevelEditor/Private/SLevelEditor.h
--- a/Engine/Source/Editor/LevelEditor/Private/SLevelEditor.h
+++ b/Engine/Source/Editor/LevelEditor/Private/SLevelEditor.h
@@ -131,6 +131,7 @@ public:
 	virtual void SetActorDetailsSCSEditorUICustomization(TSharedPtr<ISCSEditorUICustomization> InActorDetailsSCSEditorUICustomization) override;
 	virtual FEditorModeTools& GetEditorModeManager() const override;
 	virtual UTypedElementCommonActions* GetCommonActions() const override;
+	virtual FName GetStatusBarName() const override;
 	virtual FOnActiveViewportChanged& OnActiveViewportChanged() { return OnActiveViewportChangedDelegate; }
 	virtual void AddViewportOverlayWidget(TSharedRef<SWidget>, TSharedPtr<IAssetViewport> InViewport = nullptr) override;
 	virtual void RemoveViewportOverlayWidget(TSharedRef<SWidget>, TSharedPtr<IAssetViewport> InViewport = nullptr) override; 
diff --git a/Engine/Source/Editor/LevelEditor/Private/SLevelEditorToolBox.h b/Engine/Source/Editor/LevelEditor/Private/SLevelEditorToolBox.h
--- a/Engine/Source/Editor/LevelEditor/Private/SLevelEditorToolBox.h
+++ b/Engine/Source/Editor/LevelEditor/Private/SLevelEditorToolBox.h
@@ -44,9 +44,4 @@ public:
 	virtual void OnToolkitHostingFinished(const TSharedRef<IToolkit>& Toolkit) override;
 
 	virtual TSharedPtr<FWorkspaceItem> GetModeMenuCategory() const override;
-	virtual const FName GetStatusBarName() const override
-	{
-		static const FName LevelEditorStatusBarName = "LevelEditor.StatusBar";
-		return LevelEditorStatusBarName;
-	}
 };
diff --git a/Engine/Source/Editor/UnrealEd/Private/Toolkits/SStandaloneAssetEditorToolkitHost.h b/Engine/Source/Editor/UnrealEd/Private/Toolkits/SStandaloneAssetEditorToolkitHost.h
--- a/Engine/Source/Editor/UnrealEd/Private/Toolkits/SStandaloneAssetEditorToolkitHost.h
+++ b/Engine/Source/Editor/UnrealEd/Private/Toolkits/SStandaloneAssetEditorToolkitHost.h
@@ -59,6 +59,7 @@ public:
 	virtual void OnToolkitHostingFinished( const TSharedRef< class IToolkit >& Toolkit ) override;
 	virtual UWorld* GetWorld() const override;
 	virtual UTypedElementCommonActions* GetCommonActions() const override;
+	FName GetStatusBarName() const override { return StatusBarName; } 
 	virtual void AddViewportOverlayWidget(TSharedRef<SWidget> InOverlaidWidget, TSharedPtr<IAssetViewport> InViewport = nullptr) override;
 	virtual void RemoveViewportOverlayWidget(TSharedRef<SWidget> InOverlaidWidget, TSharedPtr<IAssetViewport> InViewport = nullptr) override;
 	virtual FOnActiveViewportChanged& OnActiveViewportChanged() override { return OnActiveViewportChangedDelegate; }
