Fixed the component tree potentially showing stale items

SSubobjectInstanceEditor::OnObjectReplaced was never bound, so we now just handle this directly in SActorDetails instead

#rnx

diff --git a/Engine/Source/Editor/LevelEditor/Private/SActorDetails.cpp b/Engine/Source/Editor/LevelEditor/Private/SActorDetails.cpp
--- a/Engine/Source/Editor/LevelEditor/Private/SActorDetails.cpp
+++ b/Engine/Source/Editor/LevelEditor/Private/SActorDetails.cpp
@@ -148,8 +148,7 @@ void SActorDetails::Construct(const FArguments& InArgs, UTypedElementSelectionSe
 		.ObjectContext(this, &SActorDetails::GetActorContextAsObject)
 		.AllowEditing(this, &SActorDetails::GetAllowComponentTreeEditing)
 		.OnSelectionUpdated(this, &SActorDetails::OnSubobjectEditorTreeViewSelectionChanged)
-		.OnItemDoubleClicked(this, &SActorDetails::OnSubobjectEditorTreeViewItemDoubleClicked)
-		.OnObjectReplaced(this, &SActorDetails::OnSubobjectEditorTreeViewObjectReplaced);
+		.OnItemDoubleClicked(this, &SActorDetails::OnSubobjectEditorTreeViewItemDoubleClicked);
 
 	ComponentsBox->SetContent(SubobjectEditor.ToSharedRef());
 
@@ -580,13 +579,6 @@ void SActorDetails::OnSubobjectEditorTreeViewSelectionChanged(const TArray<FSubo
 	RefreshSubobjectTreeElements(SelectedNodes, /*bForceRefresh*/false, DetailsView->IsLocked());
 }
 
-void SActorDetails::OnSubobjectEditorTreeViewObjectReplaced()
-{
-	// Enable the selection guard to prevent OnTreeSelectionChanged() from altering the editor's component selection
-	TGuardValue<bool> SelectionGuard(bSelectionGuard, true);
-	SubobjectEditor->UpdateTree();
-}
-
 void SActorDetails::OnSubobjectEditorTreeViewItemDoubleClicked(const FSubobjectEditorTreeNodePtrType ClickedNode)
 {
 	if (ClickedNode && ClickedNode->IsComponentNode())
@@ -881,4 +873,10 @@ void SActorDetails::OnObjectsReplaced(const TMap<UObject*, UObject*>& InReplacem
 			OverrideSelection(NewSelection);
 		}
 	}
+	else
+	{
+		// Enable the selection guard to prevent OnTreeSelectionChanged() from altering the editor's component selection
+		TGuardValue<bool> SelectionGuard(bSelectionGuard, true);
+		SubobjectEditor->UpdateTree();
+	}
 }
diff --git a/Engine/Source/Editor/LevelEditor/Private/SActorDetails.h b/Engine/Source/Editor/LevelEditor/Private/SActorDetails.h
--- a/Engine/Source/Editor/LevelEditor/Private/SActorDetails.h
+++ b/Engine/Source/Editor/LevelEditor/Private/SActorDetails.h
@@ -85,7 +85,6 @@ private:
 	void OnComponentsEditedInWorld();
 	void OnSubobjectEditorTreeViewSelectionChanged(const TArray<TSharedPtr<FSubobjectEditorTreeNode> >& SelectedNodes);
 	void OnSubobjectEditorTreeViewItemDoubleClicked(const TSharedPtr<FSubobjectEditorTreeNode> ClickedNode);
-	void OnSubobjectEditorTreeViewObjectReplaced();
 	void UpdateComponentTreeFromEditorSelection();
 
 	bool IsPropertyReadOnly(const struct FPropertyAndParent& PropertyAndParent) const;
diff --git a/Engine/Source/Editor/SubobjectEditor/Private/SSubobjectInstanceEditor.cpp b/Engine/Source/Editor/SubobjectEditor/Private/SSubobjectInstanceEditor.cpp
--- a/Engine/Source/Editor/SubobjectEditor/Private/SSubobjectInstanceEditor.cpp
+++ b/Engine/Source/Editor/SubobjectEditor/Private/SSubobjectInstanceEditor.cpp
@@ -40,7 +40,6 @@ void SSubobjectInstanceEditor::Construct(const FArguments& InArgs)
 	OnSelectionUpdated = InArgs._OnSelectionUpdated;
 	OnItemDoubleClicked = InArgs._OnItemDoubleClicked;
 	AllowEditing = InArgs._AllowEditing;
-	OnObjectReplaced = InArgs._OnObjectReplaced;
 	bAllowTreeUpdates = true;
 	
 	CreateCommandList();
diff --git a/Engine/Source/Editor/SubobjectEditor/Public/SSubobjectInstanceEditor.h b/Engine/Source/Editor/SubobjectEditor/Public/SSubobjectInstanceEditor.h
--- a/Engine/Source/Editor/SubobjectEditor/Public/SSubobjectInstanceEditor.h
+++ b/Engine/Source/Editor/SubobjectEditor/Public/SSubobjectInstanceEditor.h
@@ -21,7 +21,6 @@ private:
 	    SLATE_ATTRIBUTE(bool, AllowEditing)
 	    SLATE_EVENT(FOnSelectionUpdated, OnSelectionUpdated)
 	    SLATE_EVENT(FOnItemDoubleClicked, OnItemDoubleClicked)
-		SLATE_EVENT(FSimpleDelegate, OnObjectReplaced)
 	
 	SLATE_END_ARGS()
 	
@@ -48,9 +47,6 @@ protected:
 	// End of SSubobjectEditor
 
 public:
-	/** Delegate to invoke when objects within the Subobject tree are replaced (eg, via re-instancing from a BP compile) */
-	FSimpleDelegate OnObjectReplaced;
-
 private:
 
 	/** @return the tooltip describing how many properties will be applied to the blueprint */
