Level Editor Context Menu: Allow extending "LevelSubMenu" outside of MainMenu context

diff --git a/Engine/Source/Editor/LevelEditor/Private/LevelEditorContextMenu.cpp b/Engine/Source/Editor/LevelEditor/Private/LevelEditorContextMenu.cpp
--- a/Engine/Source/Editor/LevelEditor/Private/LevelEditorContextMenu.cpp
+++ b/Engine/Source/Editor/LevelEditor/Private/LevelEditorContextMenu.cpp
@@ -477,12 +477,12 @@ void FLevelEditorContextMenu::RegisterActorContextMenu()
 				FLevelScriptEventMenuHelper::FillLevelBlueprintEventsMenu(Section, SelectedActors);
 			}
 		}
-
+			
+		// General-purpose extension point for tools that apply to many types of actors
+		// These should appear in the main menu context only, by design
+		// For type-specific actions, consider adding them to "ActorTypeTools" below
 		if (LevelEditorContext->ContextType == ELevelEditorMenuContext::MainMenu)
 		{
-			// General-purpose extension point for tools that apply to many types of actors
-			// These should appear in the main menu context only, by design
-			// For type-specific actions, consider adding them to "ActorTypeTools" below
 			FToolMenuSection& Section = InMenu->AddSection("ActorUETools", LOCTEXT("UEToolsHeading", "UE Tools"));
 
 			Section.AddSubMenu(
@@ -492,21 +492,21 @@ void FLevelEditorContextMenu::RegisterActorContextMenu()
 				FNewToolMenuDelegate::CreateStatic(&FLevelEditorContextMenuImpl::FillMergeActorsMenu),
 				false, // default value
 				FSlateIcon(FAppStyle::GetAppStyleSetName(), "Icons.Merge"));
-
-			Section.AddSubMenu(
-				"LevelSubMenu",
-				LOCTEXT("LevelSubMenu", "Level"),
-				LOCTEXT("LevelSubMenu_ToolTip", "Options for interacting with this actor's level"),
-				FNewToolMenuDelegate::CreateStatic(&FLevelEditorContextMenuImpl::FillActorLevelMenu),
-				false, // default value
-				FSlateIcon(FAppStyle::GetAppStyleSetName(), "Icons.Level"));
 		}
 
 		// General-purpose extension point for tools that only appear for certain types of actors
 		// Generally, you should only use this for type-specific actions since this section appears in all contexts
 		// For tools that apply to many types of actors, add them to "ActorUETools" above
-		InMenu->AddSection("ActorTypeTools");
-
+		FToolMenuSection& ActorTypeToolsSection = InMenu->AddSection("ActorTypeTools");
+
+		ActorTypeToolsSection.AddSubMenu(
+			"LevelSubMenu",
+			LOCTEXT("LevelSubMenu", "Level"),
+			LOCTEXT("LevelSubMenu_ToolTip", "Options for interacting with this actor's level"),
+			FNewToolMenuDelegate::CreateStatic(&FLevelEditorContextMenuImpl::FillActorLevelMenu),
+			false, // default value
+			FSlateIcon(FAppStyle::GetAppStyleSetName(), "Icons.Level"));
+		
 		// DEPRECATED SECTION NAMES -- DO NOT ADD NEW EXTENSIONS TO THESE POINTS -- THEY MAY BE REMOVED IN THE FUTURE
 		// These sections are included here because they used to exist and may have been used as extension points
 		// For new context menu entries, use "ActorUETools" or "ActorTypeTools" above
@@ -1044,36 +1044,45 @@ void FLevelEditorContextMenuImpl::FillActorVisibilityMenu(UToolMenu* Menu)
 
 void FLevelEditorContextMenuImpl::FillActorLevelMenu(UToolMenu* Menu)
 {
+	ULevelEditorContextMenuContext* LevelEditorContext = Menu->FindContext<ULevelEditorContextMenuContext>();
+	if (!LevelEditorContext || !LevelEditorContext->LevelEditor.IsValid())
+	{
+		return;
+	}
+
+	if (LevelEditorContext->ContextType == ELevelEditorMenuContext::MainMenu)
 	{
-		FToolMenuSection& Section = Menu->AddSection("ActorLevel", LOCTEXT("ActorLevel", "Actor Level"));
-		if( SelectionInfo.SharedLevel && !SelectionInfo.SharedLevel->IsInstancedLevel() && SelectionInfo.SharedWorld && SelectionInfo.SharedWorld->GetCurrentLevel() != SelectionInfo.SharedLevel )
 		{
-			// All actors are in the same level and that level is not the current level 
-			// so add a menu entry to make the shared level current
+			FToolMenuSection& Section = Menu->AddSection("ActorLevel", LOCTEXT("ActorLevel", "Actor Level"));
+			if (SelectionInfo.SharedLevel && !SelectionInfo.SharedLevel->IsInstancedLevel() && SelectionInfo.SharedWorld && SelectionInfo.SharedWorld->GetCurrentLevel() != SelectionInfo.SharedLevel)
+			{
+				// All actors are in the same level and that level is not the current level 
+				// so add a menu entry to make the shared level current
+
+				FText MakeCurrentLevelText = FText::Format(LOCTEXT("MakeCurrentLevelMenu", "Make Current Level: {0}"), FText::FromString(SelectionInfo.SharedLevel->GetOutermost()->GetName()));
+				Section.AddMenuEntry(FLevelEditorCommands::Get().MakeActorLevelCurrent, MakeCurrentLevelText);
+			}
+
+			if (!SelectionInfo.bAllSelectedActorsBelongToCurrentLevel)
+			{
+				// Only show this menu entry if any actors are not in the current level
+				Section.AddMenuEntry(FLevelEditorCommands::Get().MoveSelectedToCurrentLevel);
+			}
 
-			FText MakeCurrentLevelText = FText::Format( LOCTEXT("MakeCurrentLevelMenu", "Make Current Level: {0}"), FText::FromString( SelectionInfo.SharedLevel->GetOutermost()->GetName() ) );
-			Section.AddMenuEntry(FLevelEditorCommands::Get().MakeActorLevelCurrent, MakeCurrentLevelText);
+			Section.AddMenuEntry(FLevelEditorCommands::Get().FindActorLevelInContentBrowser);
 		}
 
-		if( !SelectionInfo.bAllSelectedActorsBelongToCurrentLevel )
 		{
-			// Only show this menu entry if any actors are not in the current level
-			Section.AddMenuEntry(FLevelEditorCommands::Get().MoveSelectedToCurrentLevel);
+			FToolMenuSection& Section = Menu->AddSection("LevelBlueprint", LOCTEXT("LevelBlueprint", "Level Blueprint"));
+			Section.AddMenuEntry(FLevelEditorCommands::Get().FindActorInLevelScript);
 		}
 
-		Section.AddMenuEntry(FLevelEditorCommands::Get().FindActorLevelInContentBrowser);
-	}
-
-	{
-		FToolMenuSection& Section = Menu->AddSection("LevelBlueprint", LOCTEXT("LevelBlueprint", "Level Blueprint"));
-		Section.AddMenuEntry(FLevelEditorCommands::Get().FindActorInLevelScript);
-	}
-
-	{
-		FToolMenuSection& Section = Menu->AddSection("LevelBrowser", LOCTEXT("LevelBrowser", "Level Browser"));
-		Section.AddMenuEntry(FLevelEditorCommands::Get().FindLevelsInLevelBrowser);
-		Section.AddMenuEntry(FLevelEditorCommands::Get().AddLevelsToSelection);
-		Section.AddMenuEntry(FLevelEditorCommands::Get().RemoveLevelsFromSelection);
+		{
+			FToolMenuSection& Section = Menu->AddSection("LevelBrowser", LOCTEXT("LevelBrowser", "Level Browser"));
+			Section.AddMenuEntry(FLevelEditorCommands::Get().FindLevelsInLevelBrowser);
+			Section.AddMenuEntry(FLevelEditorCommands::Get().AddLevelsToSelection);
+			Section.AddMenuEntry(FLevelEditorCommands::Get().RemoveLevelsFromSelection);
+		}
 	}
 }
 
