Adds height to screenshot tool ui in case of text overflow

diff --git a/Engine/Source/Editor/LevelEditor/Private/HighResScreenshotUI.cpp b/Engine/Source/Editor/LevelEditor/Private/HighResScreenshotUI.cpp
--- a/Engine/Source/Editor/LevelEditor/Private/HighResScreenshotUI.cpp
+++ b/Engine/Source/Editor/LevelEditor/Private/HighResScreenshotUI.cpp
@@ -327,7 +327,7 @@ TWeakPtr<class SWindow> SHighResScreenshotDialog::OpenDialog(const TSharedPtr<FS
 		TSharedRef<SHighResScreenshotDialog> Dialog = SNew(SHighResScreenshotDialog);
 		TSharedRef<SWindow> Window = SNew(SWindow)
 			.Title( NSLOCTEXT("HighResScreenshot", "HighResolutionScreenshot", "High Resolution Screenshot") )
-			.ClientSize(FVector2D(480, 286))
+			.ClientSize(FVector2D(480, 302))
 			.SupportsMinimize(false)
 			.SupportsMaximize(false)
 			.FocusWhenFirstShown(true)
