Fixed code that depends on the selection order without a proper normalization of the selection set.

diff --git a/Engine/Source/Editor/EditorFramework/Private/Subsystems/EditorElementSubsystem.cpp b/Engine/Source/Editor/EditorFramework/Private/Subsystems/EditorElementSubsystem.cpp
--- a/Engine/Source/Editor/EditorFramework/Private/Subsystems/EditorElementSubsystem.cpp
+++ b/Engine/Source/Editor/EditorFramework/Private/Subsystems/EditorElementSubsystem.cpp
@@ -3,7 +3,9 @@
 #include "Subsystems/EditorElementSubsystem.h"
 
 #include "Elements/Framework/TypedElementRegistry.h"
+#include "Elements/Framework/TypedElementSelectionSet.h"
 #include "Elements/Interfaces/TypedElementWorldInterface.h"
+
 #include "Engine/World.h"
 
 bool UEditorElementSubsystem::SetElementTransform(FTypedElementHandle InElementHandle, const FTransform& InWorldTransform)
@@ -25,4 +27,46 @@ bool UEditorElementSubsystem::SetElementTransform(FTypedElementHandle InElementH
 	}
 
 	return false;
+}
+
+FTypedElementListRef UEditorElementSubsystem::GetEditorNormalizedSelectionSet(const UTypedElementSelectionSet& SelectionSet)
+{
+	const FTypedElementSelectionNormalizationOptions NormalizationOptions = FTypedElementSelectionNormalizationOptions()
+		.SetExpandGroups(true)
+		.SetFollowAttachment(true);
+
+	return SelectionSet.GetNormalizedSelection(NormalizationOptions);
+}
+
+FTypedElementListRef UEditorElementSubsystem::GetEditorManipulableElements(const FTypedElementListRef& NormalizedSelection)
+{
+	 NormalizedSelection->RemoveAll<ITypedElementWorldInterface>([](const TTypedElement<ITypedElementWorldInterface>& InWorldElement)
+		{
+			return !UEditorElementSubsystem::IsElementEditorManipulable(InWorldElement);
+		});
+
+	 return NormalizedSelection;
+}
+
+TTypedElement<ITypedElementWorldInterface> UEditorElementSubsystem::GetLastSelectedEditorManipulableElement(const FTypedElementListRef& NormalizedSelection)
+{
+	return NormalizedSelection->GetBottomElement<ITypedElementWorldInterface>(&UEditorElementSubsystem::IsElementEditorManipulable);
+}
+
+bool UEditorElementSubsystem::IsElementEditorManipulable(const TTypedElement<ITypedElementWorldInterface>& WorldElement)
+{
+	UWorld* OwnerWorld = WorldElement.GetOwnerWorld();
+	if (!OwnerWorld)
+	{
+		return false;
+	}
+
+	bool bIsPlayingWorld = OwnerWorld->IsPlayInEditor();
+
+	if (!WorldElement.CanMoveElement(bIsPlayingWorld ? ETypedElementWorldType::Game : ETypedElementWorldType::Editor))
+	{
+		return false;
+	}
+
+	return true;
 }
\ No newline at end of file
diff --git a/Engine/Source/Editor/EditorFramework/Public/Subsystems/EditorElementSubsystem.h b/Engine/Source/Editor/EditorFramework/Public/Subsystems/EditorElementSubsystem.h
--- a/Engine/Source/Editor/EditorFramework/Public/Subsystems/EditorElementSubsystem.h
+++ b/Engine/Source/Editor/EditorFramework/Public/Subsystems/EditorElementSubsystem.h
@@ -4,6 +4,8 @@
 
 #include "EditorSubsystem.h"
 #include "Elements/Framework/TypedElementHandle.h"
+#include "Elements/Framework/TypedElementListFwd.h"
+#include "Elements/Interfaces/TypedElementWorldInterface.h"
 
 #include "EditorElementSubsystem.generated.h"
 
@@ -17,5 +19,26 @@ public:
 	 * Sets the world transform of the given element handle, if possible.
 	 * @returns false if the world transform could not be set.
 	 */
-	bool SetElementTransform(FTypedElementHandle InElementHandle, const FTransform& InWorldTransform);
+	static bool SetElementTransform(FTypedElementHandle InElementHandle, const FTransform& InWorldTransform);
+
+	/**
+	 * Return a normalized selection set like the editor would use for the gizmo manipulations
+	 */
+	static FTypedElementListRef GetEditorNormalizedSelectionSet(const UTypedElementSelectionSet& SelectionSet);
+
+	/**
+	 * Return only the manipulatable elements from a selection list.
+	 * Require a editor normalized selection list to behave like the editor selection
+	 * Note: A manipulable element is a element that would be move when manipulating the gizmo in the editor
+	 */
+	static FTypedElementListRef GetEditorManipulableElements(const FTypedElementListRef& NormalizedSelection);
+
+	/**
+	 * Return the most recently selected element that is manipulable.
+	 * Require a editor normalized selection list to behave like the editor selection
+	 * Note: A manipulable element is a element that would be move when manipulating the gizmo in the editor.
+	 */
+	static TTypedElement<ITypedElementWorldInterface> GetLastSelectedEditorManipulableElement(const FTypedElementListRef& NormalizedSelection);
+
+	static bool IsElementEditorManipulable(const TTypedElement<ITypedElementWorldInterface>& WorldElement);
 };
\ No newline at end of file
diff --git a/Engine/Source/Editor/LevelEditor/Private/LevelEditorActions.cpp b/Engine/Source/Editor/LevelEditor/Private/LevelEditorActions.cpp
--- a/Engine/Source/Editor/LevelEditor/Private/LevelEditorActions.cpp
+++ b/Engine/Source/Editor/LevelEditor/Private/LevelEditorActions.cpp
@@ -47,6 +47,7 @@
 
 #include "Elements/Framework/TypedElementSelectionSet.h"
 #include "Elements/Interfaces/TypedElementObjectInterface.h"
+#include "Subsystems/EditorElementSubsystem.h"
 
 #include "LevelEditor.h"
 #include "Engine/LevelScriptBlueprint.h"
@@ -2921,7 +2922,7 @@ void FLevelEditorActionCallbacks::MoveElementsToElement_Clicked(bool InAlign)
 	// TODO: Ideally this would come from some level editor context
 	if (const UTypedElementSelectionSet* SelectionSet = GEditor->GetSelectedActors()->GetElementSelectionSet())
 	{
-		if (const TTypedElement<ITypedElementWorldInterface> DestElement = SelectionSet->GetBottomSelectedElement<ITypedElementWorldInterface>())
+		if (const TTypedElement<ITypedElementWorldInterface> DestElement = UEditorElementSubsystem::GetLastSelectedEditorManipulableElement(UEditorElementSubsystem::GetEditorNormalizedSelectionSet(*SelectionSet)))
 		{
 			const FScopedTransaction Transaction(NSLOCTEXT("UnrealEd", "MoveElementsToElement", "Snap Origin to Element"));
 			MoveTo_Clicked(SelectionSet, InAlign, /*bPerElement*/false, DestElement);
@@ -3138,7 +3139,7 @@ void FLevelEditorActionCallbacks::SnapElementsToElement_Clicked( bool InAlign, b
 	// TODO: Ideally this would come from some level editor context
 	if (const UTypedElementSelectionSet* SelectionSet = GEditor->GetSelectedActors()->GetElementSelectionSet())
 	{
-		if (const TTypedElement<ITypedElementWorldInterface> DestElement = SelectionSet->GetBottomSelectedElement<ITypedElementWorldInterface>())
+		if (const TTypedElement<ITypedElementWorldInterface> DestElement = UEditorElementSubsystem::GetLastSelectedEditorManipulableElement(UEditorElementSubsystem::GetEditorNormalizedSelectionSet(*SelectionSet)))
 		{
 			const FScopedTransaction Transaction(NSLOCTEXT("UnrealEd", "SnapElementsToElement", "Snap Elements to Element"));
 			SnapTo_Clicked(SelectionSet, InAlign, InUseLineTrace, InUseBounds, InUsePivot, DestElement);
@@ -3184,7 +3185,7 @@ void FLevelEditorActionCallbacks::SnapTo_Clicked( const UTypedElementSelectionSe
 	// Update the pivot location
 	if (bSnappedElements)
 	{
-		if (TTypedElement<ITypedElementWorldInterface> LastElement = InSelectionSet->GetBottomSelectedElement<ITypedElementWorldInterface>())
+		if (TTypedElement<ITypedElementWorldInterface> LastElement = UEditorElementSubsystem::GetLastSelectedEditorManipulableElement(UEditorElementSubsystem::GetEditorNormalizedSelectionSet(*InSelectionSet)))
 		{
 			FTransform LastElementTransform;
 			if (LastElement.GetWorldTransform(LastElementTransform))
diff --git a/Engine/Source/Editor/UnrealEd/Private/EditorModeManager.cpp b/Engine/Source/Editor/UnrealEd/Private/EditorModeManager.cpp
--- a/Engine/Source/Editor/UnrealEd/Private/EditorModeManager.cpp
+++ b/Engine/Source/Editor/UnrealEd/Private/EditorModeManager.cpp
@@ -41,6 +41,7 @@
 #include "Tools/AssetEditorContextObject.h"
 #include "ContextObjectStore.h"
 #include "UObject/GCObjectScopeGuard.h"
+#include "Subsystems/EditorElementSubsystem.h"
 
 #include "Elements/Interfaces/TypedElementWorldInterface.h"
 
@@ -962,12 +963,24 @@ FMatrix FEditorModeTools::GetLocalCoordinateSystem()
 		});
 
 	// If there isn't an active mode overriding the local coordinate system, create it by looking at the current selection.
-	if (!CustomCoordinateSystemProvided && GetEditorSelectionSet())
+	if (!CustomCoordinateSystemProvided)
 	{
-		if (TTypedElement<ITypedElementWorldInterface> BottomSelected = GetEditorSelectionSet()->GetBottomSelectedElement<ITypedElementWorldInterface>())
+		TTypedElement<ITypedElementWorldInterface> LastSelected;
+		if (GCurrentLevelEditingViewportClient)
+		{
+			// Use the cache from the viewport when available
+			 LastSelected = GCurrentLevelEditingViewportClient->GetElementsToManipulate()->GetBottomElement<ITypedElementWorldInterface>();
+		}
+		else
+		{
+			LastSelected = UEditorElementSubsystem::GetLastSelectedEditorManipulableElement(UEditorElementSubsystem::GetEditorNormalizedSelectionSet(*GetEditorSelectionSet()));
+		}
+
+		
+		if (LastSelected)
 		{
 			FTransform LocalToWorldTransform;
-			BottomSelected.GetWorldTransform(LocalToWorldTransform);
+			LastSelected.GetWorldTransform(LocalToWorldTransform);
 			Matrix = FQuatRotationMatrix(LocalToWorldTransform.GetRotation());
 		}
 	}
diff --git a/Engine/Source/Runtime/TypedElementFramework/Public/Elements/Framework/TypedElementList.h b/Engine/Source/Runtime/TypedElementFramework/Public/Elements/Framework/TypedElementList.h
--- a/Engine/Source/Runtime/TypedElementFramework/Public/Elements/Framework/TypedElementList.h
+++ b/Engine/Source/Runtime/TypedElementFramework/Public/Elements/Framework/TypedElementList.h
@@ -218,6 +218,28 @@ public:
 		return TempElement;
 	}
 
+	/**
+	 * Get the first element that implement the given interface and pass the predicate
+	 * @Predicate A function that should return true when the element is desirable
+	 */
+	template <typename BaseInterfaceType>
+	TTypedElement<BaseInterfaceType> GetTopElement(TFunctionRef<bool (const TTypedElement<BaseInterfaceType>&)> Predicate) const
+	{
+		TTypedElement<BaseInterfaceType> TempElement;
+		TTypedElement<BaseInterfaceType> ElementToReturn;
+		for (int32 ElementIndex = 0; ElementIndex < Num(); ++ElementIndex)
+		{
+			GetElementAt(ElementIndex, TempElement);
+			if (TempElement && Predicate(TempElement))
+			{
+				ElementToReturn = MoveTemp(TempElement);
+				break;
+			}
+		}
+		return ElementToReturn;
+	}
+
+
 	/**
 	 * Get the last element implementing the given interface.
 	 */
@@ -236,6 +258,28 @@ public:
 		return TempElement;
 	}
 
+	/**
+	 * Get the last element that implement the given interface and pass the predicate.
+	 * @Predicate A function that return should true when the element is desirable
+	 */
+	template <typename BaseInterfaceType>
+	TTypedElement<BaseInterfaceType> GetBottomElement(TFunctionRef<bool (const TTypedElement<BaseInterfaceType>&)> Predicate) const
+	{
+		TTypedElement<BaseInterfaceType> TempElement;
+		TTypedElement<BaseInterfaceType> ElementToReturn;
+		for (int32 ElementIndex = Num() - 1; ElementIndex >= 0; --ElementIndex)
+		{
+			GetElementAt(ElementIndex, TempElement);
+			if (TempElement && Predicate(ElementToReturn))
+			{
+				ElementToReturn = MoveTemp(TempElement);
+				break;
+			}
+		}
+		return ElementToReturn;
+	}
+
+
 
 	/**
 	 * Get the element interface from the given handle.
diff --git a/Engine/Source/Runtime/TypedElementRuntime/Public/Elements/Framework/TypedElementSelectionSet.h b/Engine/Source/Runtime/TypedElementRuntime/Public/Elements/Framework/TypedElementSelectionSet.h
--- a/Engine/Source/Runtime/TypedElementRuntime/Public/Elements/Framework/TypedElementSelectionSet.h
+++ b/Engine/Source/Runtime/TypedElementRuntime/Public/Elements/Framework/TypedElementSelectionSet.h
@@ -328,6 +328,16 @@ public:
 		return ElementList->GetTopElement<BaseInterfaceType>();
 	}
 
+	/**
+	 * Get the first element that implement the given interface and pass the predicate
+	 * @Predicate A function that should return true when the element is desirable
+	 */
+	template <typename BaseInterfaceType>
+	TTypedElement<BaseInterfaceType> GetTopSelectedElement(TFunctionRef<bool (const TTypedElement<BaseInterfaceType>&)> Predicate) const
+	{
+		return ElementList->GetTopElement<BaseInterfaceType>(Predicate);
+	}
+
 	/**
 	 * Get the last selected element implementing the given interface.
 	 */
@@ -337,6 +347,16 @@ public:
 		return ElementList->GetBottomElement<BaseInterfaceType>();
 	}
 
+	/**
+	 * Get the last element that implement the given interface and pass the predicate.
+	 * @Predicate A function that return should true when the element is desirable
+	 */
+	template <typename BaseInterfaceType>
+	TTypedElement<BaseInterfaceType> GetBottomSelectedElement(TFunctionRef<bool (const TTypedElement<BaseInterfaceType>&)> Predicate) const
+	{
+		return ElementList->GetBottomElement<BaseInterfaceType>(Predicate);
+	}
+
 	/**
 	 * Test whether there are any selected objects.
 	 */
