Ensure there is some visual separation between non-maximized viewports

diff --git a/Engine/Source/Editor/EditorStyle/Private/StarshipStyle.cpp b/Engine/Source/Editor/EditorStyle/Private/StarshipStyle.cpp
--- a/Engine/Source/Editor/EditorStyle/Private/StarshipStyle.cpp
+++ b/Engine/Source/Editor/EditorStyle/Private/StarshipStyle.cpp
@@ -4801,6 +4801,7 @@ void FStarshipEditorStyle::FStyle::SetupLevelEditorStyle()
 		Set( "LevelViewport.BlackBackground", new FSlateColorBrush( FLinearColor::Black ) ); 
 		Set( "LevelViewport.StartingPlayInEditorBorder", new BOX_BRUSH( "Old/Window/ViewportDebugBorder", 0.8f, FLinearColor(0.1f,1.0f,0.1f,1.0f) ) );
 		Set( "LevelViewport.StartingSimulateBorder", new BOX_BRUSH( "Old/Window/ViewportDebugBorder", 0.8f, FLinearColor(1.0f,1.0f,0.1f,1.0f) ) );
+		Set( "LevelViewport.NonMaximizedBorder", new BORDER_BRUSH("Common/PlainBorder", 2.f / 8.f, FStyleColors::Black));
 		Set( "LevelViewport.ReturningToEditorBorder", new BOX_BRUSH( "Old/Window/ViewportDebugBorder", 0.8f, FLinearColor(0.1f,0.1f,1.0f,1.0f) ) );
 		Set( "LevelViewport.ActorLockIcon", new IMAGE_BRUSH( "Icons/ActorLockedViewport", Icon32x32 ) );
 		Set( "LevelViewport.Icon", new IMAGE_BRUSH( "Icons/icon_tab_viewport_16px", Icon16x16 ) );
diff --git a/Engine/Source/Editor/LevelEditor/Private/SLevelViewport.cpp b/Engine/Source/Editor/LevelEditor/Private/SLevelViewport.cpp
--- a/Engine/Source/Editor/LevelEditor/Private/SLevelViewport.cpp
+++ b/Engine/Source/Editor/LevelEditor/Private/SLevelViewport.cpp
@@ -190,11 +190,12 @@ void SLevelViewport::Construct(const FArguments& InArgs, const FAssetEditorViewp
 	ConfigKey = InConstructionArguments.ConfigKey;
 	LevelViewportClient = InArgs._LevelEditorViewportClient;
 
-	DebuggingBorder = FEditorStyle::GetBrush( "LevelViewport.DebugBorder" );
-	BlackBackground = FEditorStyle::GetBrush( "LevelViewport.BlackBackground" );
-	StartingPlayInEditorBorder = FEditorStyle::GetBrush( "LevelViewport.StartingPlayInEditorBorder" );
-	StartingSimulateBorder = FEditorStyle::GetBrush( "LevelViewport.StartingSimulateBorder" );
-	ReturningToEditorBorder = FEditorStyle::GetBrush( "LevelViewport.ReturningToEditorBorder" );
+	DebuggingBorder = FAppStyle::Get().GetBrush( "LevelViewport.DebugBorder" );
+	BlackBackground = FAppStyle::Get().GetBrush( "LevelViewport.BlackBackground" );
+	StartingPlayInEditorBorder = FAppStyle::Get().GetBrush( "LevelViewport.StartingPlayInEditorBorder" );
+	StartingSimulateBorder = FAppStyle::Get().GetBrush( "LevelViewport.StartingSimulateBorder" );
+	ReturningToEditorBorder = FAppStyle::Get().GetBrush( "LevelViewport.ReturningToEditorBorder" );
+	NonMaximizedBorder = FAppStyle::Get().GetBrush("LevelViewport.NonMaximizedBorder");
 
 	// Default level viewport client values for settings that could appear in layout config ini
 	FLevelEditorViewportInstanceSettings ViewportInstanceSettings;
@@ -1712,6 +1713,10 @@ const FSlateBrush* SLevelViewport::OnGetViewportBorderBrush() const
 					break;
 			}
 		}
+		else if(!IsMaximized())
+		{
+			BorderBrush = NonMaximizedBorder;
+		}
 	}
 	else
 	{
diff --git a/Engine/Source/Editor/LevelEditor/Public/SLevelViewport.h b/Engine/Source/Editor/LevelEditor/Public/SLevelViewport.h
--- a/Engine/Source/Editor/LevelEditor/Public/SLevelViewport.h
+++ b/Engine/Source/Editor/LevelEditor/Public/SLevelViewport.h
@@ -796,7 +796,8 @@ private:
 	const FSlateBrush* StartingSimulateBorder;
 	/** The brush to use when returning back to the editor from PIE or SIE mode */
 	const FSlateBrush* ReturningToEditorBorder;
-
+	/** The brush to use when the viewport is not maximized */
+	const FSlateBrush* NonMaximizedBorder;
 	/** Array of objects dropped during the OnDrop event */
 	TArray<UObject*> DroppedObjects;
 
