- Prevent 'Save Current Level As' through editor modes
- Avoid adding newly saved level if RemoveLevelFromWorld fails

diff --git a/Engine/Source/Editor/LevelEditor/Private/LevelEditor.cpp b/Engine/Source/Editor/LevelEditor/Private/LevelEditor.cpp
--- a/Engine/Source/Editor/LevelEditor/Private/LevelEditor.cpp
+++ b/Engine/Source/Editor/LevelEditor/Private/LevelEditor.cpp
@@ -655,7 +655,7 @@ void FLevelEditorModule::BindGlobalLevelEditorCommands()
 	ActionList.MapAction( Commands.NewLevel, FExecuteAction::CreateStatic( &FLevelEditorActionCallbacks::NewLevel ), FCanExecuteAction::CreateStatic( &FLevelEditorActionCallbacks::NewLevel_CanExecute ) );
 	ActionList.MapAction(Commands.OpenLevel, FExecuteAction::CreateStatic(&FLevelEditorActionCallbacks::OpenLevel), FCanExecuteAction::CreateStatic(&FLevelEditorActionCallbacks::OpenLevel_CanExecute));
 	ActionList.MapAction( Commands.Save, FExecuteAction::CreateStatic( &FLevelEditorActionCallbacks::Save ), FCanExecuteAction::CreateStatic( &FLevelEditorActionCallbacks::CanSaveWorld ) );
-	ActionList.MapAction( Commands.SaveAs, FExecuteAction::CreateStatic( &FLevelEditorActionCallbacks::SaveCurrentAs ), FCanExecuteAction::CreateStatic( &FLevelEditorActionCallbacks::CanSaveWorld), FGetActionCheckState(), FIsActionButtonVisible::CreateStatic( &FLevelEditorActionCallbacks::CanSaveWorld) );
+	ActionList.MapAction( Commands.SaveAs, FExecuteAction::CreateStatic( &FLevelEditorActionCallbacks::SaveCurrentAs ), FCanExecuteAction::CreateStatic( &FLevelEditorActionCallbacks::CanSaveCurrentAs), FGetActionCheckState(), FIsActionButtonVisible::CreateStatic( &FLevelEditorActionCallbacks::CanSaveCurrentAs) );
 	ActionList.MapAction( Commands.SaveAllLevels, FExecuteAction::CreateStatic( &FLevelEditorActionCallbacks::SaveAllLevels ), FCanExecuteAction::CreateStatic( &FLevelEditorActionCallbacks::CanSaveUnpartitionedWorld), FGetActionCheckState(), FIsActionButtonVisible::CreateStatic(&FLevelEditorActionCallbacks::CanSaveUnpartitionedWorld) );
 	ActionList.MapAction( Commands.ToggleFavorite, FExecuteAction::CreateStatic( &FLevelEditorActionCallbacks::ToggleFavorite ), FCanExecuteAction::CreateStatic( &FLevelEditorActionCallbacks::ToggleFavorite_CanExecute ), FIsActionChecked::CreateStatic( &FLevelEditorActionCallbacks::ToggleFavorite_IsChecked ) );
 
diff --git a/Engine/Source/Editor/LevelEditor/Private/LevelEditorActions.cpp b/Engine/Source/Editor/LevelEditor/Private/LevelEditorActions.cpp
--- a/Engine/Source/Editor/LevelEditor/Private/LevelEditorActions.cpp
+++ b/Engine/Source/Editor/LevelEditor/Private/LevelEditorActions.cpp
@@ -427,20 +427,21 @@ void FLevelEditorActionCallbacks::Save()
 	FEditorFileUtils::SaveCurrentLevel();
 }
 
+bool FLevelEditorActionCallbacks::CanSaveCurrentAs()
+{
+	return CanSaveWorld() && GLevelEditorModeTools().IsOperationSupportedForCurrentAsset(EAssetOperation::Duplicate);
+}
+
 void FLevelEditorActionCallbacks::SaveCurrentAs()
 {
+	check(CanSaveCurrentAs());
 	UWorld* World = GetWorld();
 	ULevel* CurrentLevel = World->GetCurrentLevel();
 
 	UClass* CurrentStreamingLevelClass = ULevelStreamingDynamic::StaticClass();
-
-	for (ULevelStreaming* StreamingLevel : World->GetStreamingLevels())
+	if (ULevelStreaming* StreamingLevel = FLevelUtils::FindStreamingLevel(CurrentLevel))
 	{
-		if (StreamingLevel->GetLoadedLevel() == CurrentLevel)
-		{
-			CurrentStreamingLevelClass = StreamingLevel->GetClass();
-			break;
-		}
+		CurrentStreamingLevelClass = StreamingLevel->GetClass();
 	}
 
 	
@@ -453,11 +454,8 @@ void FLevelEditorActionCallbacks::SaveCurrentAs()
 		{
 			FEditorFileUtils::LoadMap(SavedFilename);
 		}
-		else
+		else if(EditorLevelUtils::RemoveLevelFromWorld(CurrentLevel))
 		{
-			// Remove the level we just saved over
-			EditorLevelUtils::RemoveLevelFromWorld(CurrentLevel);
-
 			// Add the new level we just saved as to the plevel
 			FString PackageName;
 			if (FPackageName::TryConvertFilenameToLongPackageName(SavedFilename, PackageName))
diff --git a/Engine/Source/Editor/LevelEditor/Public/LevelEditorActions.h b/Engine/Source/Editor/LevelEditor/Public/LevelEditorActions.h
--- a/Engine/Source/Editor/LevelEditor/Public/LevelEditorActions.h
+++ b/Engine/Source/Editor/LevelEditor/Public/LevelEditorActions.h
@@ -691,6 +691,7 @@ public:
 	static bool CanSaveUnpartitionedWorld();
 
 	/** Save the current level as... */
+	static bool CanSaveCurrentAs();
 	static void SaveCurrentAs();
 
 	/** Saves the current map */
diff --git a/Engine/Source/Editor/LevelInstanceEditor/Private/LevelInstanceEditorMode.h b/Engine/Source/Editor/LevelInstanceEditor/Private/LevelInstanceEditorMode.h
--- a/Engine/Source/Editor/LevelInstanceEditor/Private/LevelInstanceEditorMode.h
+++ b/Engine/Source/Editor/LevelInstanceEditor/Private/LevelInstanceEditorMode.h
@@ -28,7 +28,8 @@ public:
 
 	/** Check to see if an actor can be selected in this mode - no side effects */
 	virtual bool IsSelectionDisallowed(AActor* InActor, bool bInSelection) const override;
-	
+	/** Only accept saving on current asset */
+	virtual bool IsOperationSupportedForCurrentAsset(EAssetOperation InOperation) const { return InOperation == EAssetOperation::Save; }
 	// End UEdMode
 
 private:
diff --git a/Engine/Source/Editor/UnrealEd/Private/EditorModeManager.cpp b/Engine/Source/Editor/UnrealEd/Private/EditorModeManager.cpp
--- a/Engine/Source/Editor/UnrealEd/Private/EditorModeManager.cpp
+++ b/Engine/Source/Editor/UnrealEd/Private/EditorModeManager.cpp
@@ -613,6 +613,22 @@ void FEditorModeTools::ForEachEdMode(TFunctionRef<bool(UEdMode*)> InCalllback) c
 	}
 }
 
+bool FEditorModeTools::TestAllModes(TFunctionRef<bool(UEdMode*)> InCalllback, bool bExpected) const
+{
+	for (UEdMode* Mode : ActiveScriptableModes)
+	{
+		if (Mode)
+		{
+			if (InCalllback(Mode) != bExpected)
+			{
+				return false;
+			}
+		}
+	}
+
+	return true;
+}
+
 void FEditorModeTools::ExitAllModesPendingDeactivate()
 {
 	// Make a copy so we can modify the pending deactivate modes map
@@ -1823,38 +1839,18 @@ bool FEditorModeTools::CanCycleWidgetMode() const
 
 bool FEditorModeTools::CanAutoSave() const
 {
-	bool bCanAutoSave = true;
-	ForEachEdMode([&bCanAutoSave](const UEdMode* Mode)
-		{
-			if (!Mode->CanAutoSave())
-			{
-				bCanAutoSave = false;
-				return false;
-			}
-
-			return true;
-		});
-
-	return bCanAutoSave;
+	return FEditorModeTools::TestAllModes([](UEdMode* Mode) { return Mode->CanAutoSave(); }, true);
 }
 
 
 bool FEditorModeTools::OnRequestClose()
 {
-	bool bCanClose = true;
-
-	ForEachEdMode([&bCanClose](UEdMode* Mode)
-	{
-		if (!Mode->OnRequestClose())
-		{
-			bCanClose = false;
-		}
-
-		return true;
-	});
-
-	return bCanClose;
+	return FEditorModeTools::TestAllModes([](UEdMode* Mode) { return Mode->OnRequestClose(); }, true);
+}
 
+bool FEditorModeTools::IsOperationSupportedForCurrentAsset(EAssetOperation InOperation) const
+{
+	return FEditorModeTools::TestAllModes([InOperation](UEdMode* Mode) { return Mode->IsOperationSupportedForCurrentAsset(InOperation); }, true);
 }
 
 UModeManagerInteractiveToolsContext* FEditorModeTools::GetInteractiveToolsContext() const
diff --git a/Engine/Source/Editor/UnrealEd/Public/EditorModeManager.h b/Engine/Source/Editor/UnrealEd/Public/EditorModeManager.h
--- a/Engine/Source/Editor/UnrealEd/Public/EditorModeManager.h
+++ b/Engine/Source/Editor/UnrealEd/Public/EditorModeManager.h
@@ -519,6 +519,9 @@ public:
 
 	/** returns true if all active EdModes are OK with an AutoSave happening now  */
 	bool CanAutoSave() const;
+	
+	/** returns true if all active EdModes are OK support operation on current asset */
+	bool IsOperationSupportedForCurrentAsset(EAssetOperation InOperation) const;
 
 	void RemoveAllDelegateHandlers();
 
@@ -540,7 +543,8 @@ protected:
 	virtual void DrawBrackets(FEditorViewportClient* ViewportClient, FViewport* Viewport, const FSceneView* View, FCanvas* Canvas);
 
 	void ForEachEdMode(TFunctionRef<bool(UEdMode*)> InCalllback) const;
-
+	bool TestAllModes(TFunctionRef<bool(UEdMode*)> InCalllback, bool bExpected) const;
+	
 	template <class InterfaceToCastTo>
 	void ForEachEdMode(TFunctionRef<bool(InterfaceToCastTo*)> InCallback) const
 	{
diff --git a/Engine/Source/Editor/UnrealEd/Public/Tools/UEdMode.h b/Engine/Source/Editor/UnrealEd/Public/Tools/UEdMode.h
--- a/Engine/Source/Editor/UnrealEd/Public/Tools/UEdMode.h
+++ b/Engine/Source/Editor/UnrealEd/Public/Tools/UEdMode.h
@@ -36,6 +36,19 @@ namespace EEditAction
 	};
 };
 
+/** Generic Asset operations that can be disallowed by edit modes */
+enum class EAssetOperation
+{
+	/** Can Asset be deleted */
+	Delete,
+	/** Can Asset be duplicated / saved as */
+	Duplicate,
+	/** Can Asset be saved */
+	Save,
+	/** Can Asset be renamed */
+	Rename
+};
+
 /** 
  * EToolsContextScope is used to determine the visibility/lifetime of Tools for a ToolsContext.
  * For example Tools at the EdMode scope level will only be available when that Mode is active,
@@ -137,6 +150,12 @@ public:
 	 */
 	virtual bool CanAutoSave() const { return true; }
 
+	/**
+	 * Check to see if this UEdMode wants to disallow operation on current asset
+	 * @return true if operation is supported right now
+	 */
+	virtual bool IsOperationSupportedForCurrentAsset(EAssetOperation InOperation) const { return true; }
+
 	virtual void SelectNone();
 	virtual void SelectionChanged() {}
 
