Fix the target for the Level Editor documentation in the Help menu.

diff --git a/Engine/Source/Editor/LevelEditor/Private/LevelEditorActions.cpp b/Engine/Source/Editor/LevelEditor/Private/LevelEditorActions.cpp
--- a/Engine/Source/Editor/LevelEditor/Private/LevelEditorActions.cpp
+++ b/Engine/Source/Editor/LevelEditor/Private/LevelEditorActions.cpp
@@ -179,7 +179,7 @@ bool FLevelEditorActionCallbacks::DefaultCanExecuteAction()
 
 void FLevelEditorActionCallbacks::BrowseDocumentation()
 {
-	IDocumentation::Get()->OpenHome(FDocumentationSourceInfo(TEXT("help_menu")));
+	IDocumentation::Get()->Open(TEXT("BuildingWorlds/LevelEditor"), FDocumentationSourceInfo(TEXT("help_menu")));
 }
 
 void FLevelEditorActionCallbacks::BrowseViewportControls()
