[LevelEditorSubsystem] Add blueprint function to get the currently piloted actor/camera

diff --git a/Engine/Source/Editor/LevelEditor/Private/LevelEditorSubsystem.cpp b/Engine/Source/Editor/LevelEditor/Private/LevelEditorSubsystem.cpp
--- a/Engine/Source/Editor/LevelEditor/Private/LevelEditorSubsystem.cpp
+++ b/Engine/Source/Editor/LevelEditor/Private/LevelEditorSubsystem.cpp
@@ -135,6 +135,22 @@ void ULevelEditorSubsystem::EjectPilotLevelActor(FName ViewportConfigKey)
 	}
 }
 
+AActor* ULevelEditorSubsystem::GetPilotLevelActor(FName ViewportConfigKey)
+{
+	TSharedPtr<SLevelViewport> LevelViewport = InternalEditorLevelLibrary::GetLevelViewport(ViewportConfigKey);
+	if (!LevelViewport.IsValid())
+	{
+		return nullptr;
+	}
+	
+	FLevelEditorViewportClient& LevelViewportClient = LevelViewport->GetLevelViewportClient();
+	if (AActor* CinematicActorLock = LevelViewportClient.GetCinematicActorLock().GetLockedActor())
+	{
+		return CinematicActorLock;
+	}
+	return LevelViewportClient.GetActiveActorLock().Get();
+}
+
 void ULevelEditorSubsystem::EditorPlaySimulate()
 {
 	FLevelEditorModule& LevelEditorModule = FModuleManager::GetModuleChecked<FLevelEditorModule>("LevelEditor");
diff --git a/Engine/Source/Editor/LevelEditor/Public/LevelEditorSubsystem.h b/Engine/Source/Editor/LevelEditor/Public/LevelEditorSubsystem.h
--- a/Engine/Source/Editor/LevelEditor/Public/LevelEditorSubsystem.h
+++ b/Engine/Source/Editor/LevelEditor/Public/LevelEditorSubsystem.h
@@ -32,6 +32,9 @@ public:
 	UFUNCTION(BlueprintCallable, Category = "Editor Scripting | Level Utility", meta = (DevelopmentOnly))
 	void EjectPilotLevelActor(FName ViewportConfigKey = NAME_None);
 
+	UFUNCTION(BlueprintCallable, Category = "Editor Scripting | Level Utility", meta=(DevelopmentOnly))
+	AActor* GetPilotLevelActor(FName ViewportConfigKey = NAME_None);
+
 	UFUNCTION(BlueprintCallable, Category = "Editor Scripting | Level Utility", meta = (DevelopmentOnly))
 	void EditorPlaySimulate();
 
