Remove Alt+P shortcut for opening the Place Actors Panel, as it is already used for PIE

diff --git a/Engine/Source/Editor/LevelEditor/Private/LevelEditorActions.cpp b/Engine/Source/Editor/LevelEditor/Private/LevelEditorActions.cpp
--- a/Engine/Source/Editor/LevelEditor/Private/LevelEditorActions.cpp
+++ b/Engine/Source/Editor/LevelEditor/Private/LevelEditorActions.cpp
@@ -3573,7 +3573,7 @@ void FLevelEditorCommands::RegisterCommands()
 	UI_COMMAND( FindActorInLevelScript, "Find in Level Blueprint", "Finds any references to the selected actor in its level's blueprint", EUserInterfaceActionType::Button, FInputChord() );
 
 	UI_COMMAND( WorldProperties, "World Settings", "Displays the world settings", EUserInterfaceActionType::Button, FInputChord() );
-	UI_COMMAND( OpenPlaceActors, "Place Actors Panel", "Opens the Place Actors Panel", EUserInterfaceActionType::Button, FInputChord(EModifierKey::Alt, EKeys::P) );
+	UI_COMMAND( OpenPlaceActors, "Place Actors Panel", "Opens the Place Actors Panel", EUserInterfaceActionType::Button, FInputChord() );
 	UI_COMMAND( OpenContentBrowser, "Open Content Browser", "Opens the Content Browser", EUserInterfaceActionType::Button, FInputChord(EModifierKey::Control|EModifierKey::Shift, EKeys::F) );
 	UI_COMMAND( OpenMarketplace, "Unreal Marketplace", "Opens the Marketplace", EUserInterfaceActionType::Button, FInputChord() );
 	UI_COMMAND( ImportContent, "Import Content...", "Import Content into a specified location", EUserInterfaceActionType::Button, FInputChord());
