Fix Quick Add menu items disappearing when the Quixel Bridge plugin is not loaded.

diff --git a/Engine/Plugins/Bridge/Source/Bridge/Private/UI/BridgeUIManager.cpp b/Engine/Plugins/Bridge/Source/Bridge/Private/UI/BridgeUIManager.cpp
--- a/Engine/Plugins/Bridge/Source/Bridge/Private/UI/BridgeUIManager.cpp
+++ b/Engine/Plugins/Bridge/Source/Bridge/Private/UI/BridgeUIManager.cpp
@@ -80,7 +80,8 @@ void FBridgeUIManagerImpl::SetupMenuItem()
 		LOCTEXT("OpenBridgeTab_Desc", "Opens the Quixel Bridge."),
 		FSlateIcon(FBridgeStyle::GetStyleSetName(), "Bridge.MenuLogo"),
 		FUIAction(FExecuteAction::CreateRaw(this, &FBridgeUIManagerImpl::CreateWindow), FCanExecuteAction())
-	);
+	).InsertPosition = FToolMenuInsert("ImportContent", EToolMenuInsertType::After);
+
 
 	UToolMenu* WindowMenu = UToolMenus::Get()->ExtendMenu("MainFrame.MainMenu.Window");
 	FToolMenuSection* ContentSectionPtr = WindowMenu->FindSection("GetContent");
diff --git a/Engine/Source/Editor/LevelEditor/Private/LevelEditorToolBar.cpp b/Engine/Source/Editor/LevelEditor/Private/LevelEditorToolBar.cpp
--- a/Engine/Source/Editor/LevelEditor/Private/LevelEditorToolBar.cpp
+++ b/Engine/Source/Editor/LevelEditor/Private/LevelEditorToolBar.cpp
@@ -1968,11 +1968,11 @@ void FLevelEditorToolBar::RegisterAddMenu()
 
 		Section.InitSection("Content", LOCTEXT("Content_Label", "Get Content"), FToolMenuInsert(NAME_None, EToolMenuInsertType::First));
 
-		Section.AddMenuEntry(FLevelEditorCommands::Get().ImportContent);
+		Section.AddMenuEntry(FLevelEditorCommands::Get().ImportContent).InsertPosition.Position = EToolMenuInsertType::First;
 
 		if (FLauncherPlatformModule::Get()->CanOpenLauncher(true))
 		{
-			Section.AddMenuEntry(FLevelEditorCommands::Get().OpenMarketplace).InsertPosition = FToolMenuInsert("OpenBridgeTab", EToolMenuInsertType::After);
+			Section.AddMenuEntry(FLevelEditorCommands::Get().OpenMarketplace);
 		}
 	}
 #undef LOCTEXT_NAMESPACE
