Re-create the Outliner on Map Change,
Disable the Pinned Column and Context Menu option on non-wp maps,
Rename the Level Column to "Package Short Name" on wp enabled maps,
Add the ability to access all outliners from SLevelEditor

diff --git a/Engine/Source/Editor/FoliageEdit/Private/FoliageEdMode.cpp b/Engine/Source/Editor/FoliageEdit/Private/FoliageEdMode.cpp
--- a/Engine/Source/Editor/FoliageEdit/Private/FoliageEdMode.cpp
+++ b/Engine/Source/Editor/FoliageEdit/Private/FoliageEdMode.cpp
@@ -1688,10 +1688,14 @@ void RefreshSceneOutliner()
 	TWeakPtr<class ILevelEditor> LevelEditor = FModuleManager::GetModuleChecked<FLevelEditorModule>(TEXT("LevelEditor")).GetLevelEditorInstance();
 	if (LevelEditor.IsValid())
 	{
-		TSharedPtr<class ISceneOutliner> SceneOutlinerPtr = LevelEditor.Pin()->GetSceneOutliner();
-		if (SceneOutlinerPtr)
+		TArray<TWeakPtr<class ISceneOutliner>> SceneOutlinerPtrs = LevelEditor.Pin()->GetAllSceneOutliners();
+
+		for (TWeakPtr<class ISceneOutliner> SceneOutlinerPtr : SceneOutlinerPtrs)
 		{
-			SceneOutlinerPtr->FullRefresh();
+			if (TSharedPtr<class ISceneOutliner> SceneOutlinerPin = SceneOutlinerPtr.Pin())
+			{
+				SceneOutlinerPin->FullRefresh();
+			}
 		}
 	}
 }
diff --git a/Engine/Source/Editor/LevelEditor/Private/SLevelEditor.cpp b/Engine/Source/Editor/LevelEditor/Private/SLevelEditor.cpp
--- a/Engine/Source/Editor/LevelEditor/Private/SLevelEditor.cpp
+++ b/Engine/Source/Editor/LevelEditor/Private/SLevelEditor.cpp
@@ -645,6 +645,13 @@ void SLevelEditor::AttachSequencer( TSharedPtr<SWidget> SequencerWidget, TShared
 	}
 }
 
+TArray<TWeakPtr<ISceneOutliner>> SLevelEditor::GetAllSceneOutliners() const
+{
+	TArray<TWeakPtr<ISceneOutliner>> OutValueArray;
+	SceneOutliners.GenerateValueArray(OutValueArray);
+	return OutValueArray;
+}
+
 TSharedRef<SDockTab> SLevelEditor::SummonDetailsPanel( FName TabIdentifier )
 {
 	TSharedRef<SActorDetails> ActorDetails = StaticCastSharedRef<SActorDetails>( CreateActorDetails( TabIdentifier ) );
@@ -678,6 +685,57 @@ TSharedRef<SDockTab> SLevelEditor::SummonDetailsPanel( FName TabIdentifier )
 
 	return DocTab;
 }
+
+TSharedRef<ISceneOutliner> SLevelEditor::CreateSceneOutliner(FName TabIdentifier)
+{
+	FSceneOutlinerInitializationOptions InitOptions;
+	InitOptions.bShowTransient = true;
+	{
+		UToolMenus* ToolMenus = UToolMenus::Get();
+		static const FName MenuName = "LevelEditor.LevelEditorSceneOutliner.ContextMenu";
+		if (!ToolMenus->IsMenuRegistered(MenuName))
+		{
+			UToolMenu* Menu = ToolMenus->RegisterMenu(MenuName, "SceneOutliner.DefaultContextMenuBase");
+			FToolMenuSection& Section = Menu->AddDynamicSection("LevelEditorContextMenu", FNewToolMenuDelegate::CreateLambda([SelectionSet = TWeakObjectPtr<const UTypedElementSelectionSet>(GetElementSelectionSet())](UToolMenu* InMenu)
+			{
+				FName LevelContextMenuName = FLevelEditorContextMenu::GetContextMenuName(ELevelEditorMenuContext::SceneOutliner, SelectionSet.Get());
+				if (LevelContextMenuName != NAME_None)
+				{
+					// Extend the menu even if no actors selected, as Edit menu should always exist for scene outliner
+					UToolMenu* OtherMenu = UToolMenus::Get()->GenerateMenu(LevelContextMenuName, InMenu->Context);
+					InMenu->Sections.Append(OtherMenu->Sections);
+				}
+			}));
+			Section.InsertPosition = FToolMenuInsert("MainSection", EToolMenuInsertType::Before);
+		}
+
+		TWeakPtr<SLevelEditor> WeakLevelEditor = SharedThis(this);
+		InitOptions.ModifyContextMenu.BindLambda([=](FName& OutMenuName, FToolMenuContext& MenuContext)
+			{
+				OutMenuName = MenuName;
+
+				if (WeakLevelEditor.IsValid())
+				{
+					FLevelEditorContextMenu::InitMenuContext(MenuContext, WeakLevelEditor, ELevelEditorMenuContext::SceneOutliner);
+				}
+			});
+	}
+
+	InitOptions.OutlinerIdentifier = TabIdentifier;
+
+	FSceneOutlinerModule& SceneOutlinerModule = FModuleManager::Get().LoadModuleChecked<FSceneOutlinerModule>("SceneOutliner");
+	TSharedRef<ISceneOutliner> SceneOutlinerRef = SceneOutlinerModule.CreateActorBrowser(
+		InitOptions, GetWorld());
+
+	// Update the most recently created outliner
+	SceneOutlinerPtr = SceneOutlinerRef;
+
+	// Add this to the map of all outliners
+	SceneOutliners.Add(TabIdentifier, SceneOutlinerRef);
+
+	return SceneOutlinerRef;
+}
+
 /** Method to call when a tab needs to be spawned by the FLayoutService */
 TSharedRef<SDockTab> SLevelEditor::SpawnLevelEditorTab( const FSpawnTabArgs& Args, FName TabIdentifier, FString InitializationPayload )
 {
@@ -729,49 +787,11 @@ TSharedRef<SDockTab> SLevelEditor::SpawnLevelEditorTab( const FSpawnTabArgs& Arg
 	}
 	else if (TabIdentifier == LevelEditorTabIds::LevelEditorSceneOutliner || TabIdentifier == LevelEditorTabIds::LevelEditorSceneOutliner2 || TabIdentifier == LevelEditorTabIds::LevelEditorSceneOutliner3 || TabIdentifier == LevelEditorTabIds::LevelEditorSceneOutliner4)
 	{
-		FSceneOutlinerInitializationOptions InitOptions;
-		InitOptions.bShowTransient = true;
-		{
-			UToolMenus* ToolMenus = UToolMenus::Get();
-			static const FName MenuName = "LevelEditor.LevelEditorSceneOutliner.ContextMenu";
-			if (!ToolMenus->IsMenuRegistered(MenuName))
-			{
-				UToolMenu* Menu = ToolMenus->RegisterMenu(MenuName, "SceneOutliner.DefaultContextMenuBase");
-				FToolMenuSection& Section = Menu->AddDynamicSection("LevelEditorContextMenu", FNewToolMenuDelegate::CreateLambda([SelectionSet = TWeakObjectPtr<const UTypedElementSelectionSet>(GetElementSelectionSet())](UToolMenu* InMenu)
-				{
-					FName LevelContextMenuName = FLevelEditorContextMenu::GetContextMenuName(ELevelEditorMenuContext::SceneOutliner, SelectionSet.Get());
-					if (LevelContextMenuName != NAME_None)
-					{
-						// Extend the menu even if no actors selected, as Edit menu should always exist for scene outliner
-						UToolMenu* OtherMenu = UToolMenus::Get()->GenerateMenu(LevelContextMenuName, InMenu->Context);
-						InMenu->Sections.Append(OtherMenu->Sections);
-					}
-				}));
-				Section.InsertPosition = FToolMenuInsert("MainSection", EToolMenuInsertType::Before);
-			}
-
-			TWeakPtr<SLevelEditor> WeakLevelEditor = SharedThis(this);
-			InitOptions.ModifyContextMenu.BindLambda([=](FName& OutMenuName, FToolMenuContext& MenuContext)
-			{
-				OutMenuName = MenuName;
-
-				if (WeakLevelEditor.IsValid())
-				{
-					FLevelEditorContextMenu::InitMenuContext(MenuContext, WeakLevelEditor, ELevelEditorMenuContext::SceneOutliner);
-				}
-			});
-		}
+		TSharedRef<ISceneOutliner> SceneOutlinerRef = CreateSceneOutliner(TabIdentifier);
 
+		FText Label = NSLOCTEXT("LevelEditor", "SceneOutlinerTabTitle", "Outliner");
 
-		FText Label = NSLOCTEXT( "LevelEditor", "SceneOutlinerTabTitle", "Outliner" );
-		InitOptions.OutlinerIdentifier = "LevelEditorOutliner";
-
-		FSceneOutlinerModule& SceneOutlinerModule = FModuleManager::Get().LoadModuleChecked<FSceneOutlinerModule>("SceneOutliner");
-		TSharedRef<ISceneOutliner> SceneOutlinerRef = SceneOutlinerModule.CreateActorBrowser(
-			InitOptions);
-		SceneOutlinerPtr = SceneOutlinerRef;
-
-		return SNew( SDockTab )
+		TSharedRef<SDockTab> SceneOutlinerTab = SNew( SDockTab )
 			.Label( Label )
 			.ToolTip( IDocumentation::Get()->CreateToolTip( Label, nullptr, "Shared/LevelEditor", "SceneOutlinerTab" ) )
 			[
@@ -783,6 +803,11 @@ TSharedRef<SDockTab> SLevelEditor::SpawnLevelEditorTab( const FSpawnTabArgs& Arg
 					SceneOutlinerRef
 				]
 			];
+
+		// Add it to the map of all the outliner tabs
+		SceneOutlinerTabs.Add(TabIdentifier, SceneOutlinerTab);
+
+		return SceneOutlinerTab;
 	}
 	else if(TabIdentifier == LevelEditorTabIds::LevelEditorLayerBrowser)
 	{
@@ -1704,6 +1729,25 @@ void SLevelEditor::HandleEditorMapChange( uint32 MapChangeFlags )
 	{
 		WorldSettingsView->SetObject(GetWorld()->GetWorldSettings(), true);
 	}
+
+	// We want to recreate all the Scene Outliners on loading a new map
+	if (MapChangeFlags == MapChangeEventFlags::NewMap)
+	{
+		for (TPair<FName, TWeakPtr<SDockTab>> OutlinerTabIdentifier : SceneOutlinerTabs)
+		{
+			if (TSharedPtr<SDockTab> SceneOutlinerTabPin = OutlinerTabIdentifier.Value.Pin())
+			{
+				// Create the new outliner
+				TSharedRef<ISceneOutliner> SceneOutlinerRef = CreateSceneOutliner(OutlinerTabIdentifier.Key);
+
+				// Add it to the list of outliners
+				SceneOutliners.Add(OutlinerTabIdentifier.Key, SceneOutlinerRef);
+
+				// Set the tabs content to the newly created outliner
+				SceneOutlinerTabPin->SetContent(SceneOutlinerRef);
+			}
+		}
+	}
 }
 
 void SLevelEditor::HandleAssetsDeleted(const TArray<UClass*>& DeletedClasses)
diff --git a/Engine/Source/Editor/LevelEditor/Private/SLevelEditor.h b/Engine/Source/Editor/LevelEditor/Private/SLevelEditor.h
--- a/Engine/Source/Editor/LevelEditor/Private/SLevelEditor.h
+++ b/Engine/Source/Editor/LevelEditor/Private/SLevelEditor.h
@@ -150,9 +150,14 @@ public:
 	/** Attaches a sequencer asset editor used to animate objects in the level to this level editor */
 	void AttachSequencer( TSharedPtr<SWidget> SequencerWidget, TSharedPtr<IAssetEditorInstance> NewSequencerAssetEditor );
 
-	/** Returns current scene outliner associated with level editor's scene outliner tab, if it exists */
+	/** Returns the most recently created scene outliner. Use GetAllSceneOutliners() to get all potentially active outliners
+	 *  TODO: Add a way to identify the most recently interacted with outliner, and change/deprecate this function to return it
+	 */
 	virtual TSharedPtr<ISceneOutliner> GetSceneOutliner() const override { return SceneOutlinerPtr.Pin();  }
 
+	/** Get an array containing weak pointers to all 4 Scene Outliners which could be potentially active */
+	virtual TArray<TWeakPtr<ISceneOutliner>> GetAllSceneOutliners() const override;
+
 	TSharedRef<SWidget> GetTitleBarMessageWidget() const { return TtileBarMessageBox.ToSharedRef(); }
 private:
 	
@@ -258,6 +263,9 @@ private:
 	/** @return All valid actor details panels */
 	TArray<TSharedRef<SActorDetails>> GetAllActorDetails() const;
 
+	/** Create a Scene Outliner for the Level Editor */
+	TSharedRef<ISceneOutliner> CreateSceneOutliner(FName TabIdentifier);
+
 private:
 
 	// Tracking the active viewports in this level editor.
@@ -308,9 +316,15 @@ private:
 	/** Weak pointer to the level editor's Sequencer widget */
 	TWeakPtr<SWidget> SequencerWidgetPtr;
 
-	/** Weak pointer to the level editor's scene outliner */
+	/** Weak pointer to the level editor's most recently created scene outliner */
 	TWeakPtr<ISceneOutliner> SceneOutlinerPtr;
 
+	/** Map containing Weak pointers to all the Outliners in the level editor */
+	TMap<FName, TWeakPtr<ISceneOutliner>> SceneOutliners;
+
+	/** Map containing Weak pointers to all the Scene Outliner Tabs in the level editor */
+	TMap<FName, TWeakPtr<SDockTab>> SceneOutlinerTabs;
+
 	/** Handle to the registered OnPreviewFeatureLevelChanged delegate. */
 	FDelegateHandle PreviewFeatureLevelChangedHandle;
 
diff --git a/Engine/Source/Editor/LevelEditor/Public/ILevelEditor.h b/Engine/Source/Editor/LevelEditor/Public/ILevelEditor.h
--- a/Engine/Source/Editor/LevelEditor/Public/ILevelEditor.h
+++ b/Engine/Source/Editor/LevelEditor/Public/ILevelEditor.h
@@ -89,6 +89,8 @@ public:
 	virtual void SetActorDetailsSCSEditorUICustomization(TSharedPtr<class ISCSEditorUICustomization> ActorDetailsSCSEditorUICustomization) = 0;
 
 	virtual TSharedPtr<ISceneOutliner> GetSceneOutliner() const = 0;
+
+	virtual TArray<TWeakPtr<ISceneOutliner>> GetAllSceneOutliners() const = 0;
 };
 
 
diff --git a/Engine/Source/Editor/SceneOutliner/Private/ActorBrowsingMode.cpp b/Engine/Source/Editor/SceneOutliner/Private/ActorBrowsingMode.cpp
--- a/Engine/Source/Editor/SceneOutliner/Private/ActorBrowsingMode.cpp
+++ b/Engine/Source/Editor/SceneOutliner/Private/ActorBrowsingMode.cpp
@@ -431,28 +431,32 @@ void FActorBrowsingMode::RegisterContextMenu()
 
 					if (Context->bShowParentTree && Context->NumSelectedItems > 0 && Context->SceneOutliner.IsValid())
 					{
-						// If selection contains some unpinned items, show the pin option
-						// If the selection contains folders, always show the pin option
-						if (Context->NumPinnedItems != Context->NumSelectedItems || Context->NumSelectedFolders > 0)
+						// Only add the menu option to wp levels
+						if (Context->bRepresentingPartitionedWorld)
 						{
-							Section.AddMenuEntry(
-								"PinItems",
-								LOCTEXT("Pin", "Pin"),
-								FText(),
-								FSlateIcon(),
-								FUIAction(FExecuteAction::CreateSP(Context->SceneOutliner.Pin().Get(), &SSceneOutliner::PinSelectedItems)));
-						}
+							// If selection contains some unpinned items, show the pin option
+							// If the selection contains folders, always show the pin option
+							if (Context->NumPinnedItems != Context->NumSelectedItems || Context->NumSelectedFolders > 0)
+							{
+								Section.AddMenuEntry(
+									"PinItems",
+									LOCTEXT("Pin", "Pin"),
+									FText(),
+									FSlateIcon(),
+									FUIAction(FExecuteAction::CreateSP(Context->SceneOutliner.Pin().Get(), &SSceneOutliner::PinSelectedItems)));
+							}
 
-						// If the selection contains some pinned items, show the unpin option
-						// If the selection contains folders, always show the unpin option
-						if (Context->NumPinnedItems != 0 || Context->NumSelectedFolders > 0)
-						{
-							Section.AddMenuEntry(
-								"UnpinItems",
-								LOCTEXT("Unpin", "Unpin"),
-								FText(),
-								FSlateIcon(),
-								FUIAction(FExecuteAction::CreateSP(Context->SceneOutliner.Pin().Get(), &SSceneOutliner::UnpinSelectedItems)));
+							// If the selection contains some pinned items, show the unpin option
+							// If the selection contains folders, always show the unpin option
+							if (Context->NumPinnedItems != 0 || Context->NumSelectedFolders > 0)
+							{
+								Section.AddMenuEntry(
+									"UnpinItems",
+									LOCTEXT("Unpin", "Unpin"),
+									FText(),
+									FSlateIcon(),
+									FUIAction(FExecuteAction::CreateSP(Context->SceneOutliner.Pin().Get(), &SSceneOutliner::UnpinSelectedItems)));
+							}
 						}
 					}
 				}
@@ -477,6 +481,7 @@ TSharedPtr<SWidget> FActorBrowsingMode::BuildContextMenu()
 	ContextObject->NumSelectedItems = ItemSelection.Num();
 	ContextObject->NumSelectedFolders = ItemSelection.Num<FFolderTreeItem>();
 	ContextObject->NumWorldsSelected = ItemSelection.Num<FWorldTreeItem>();
+	ContextObject->bRepresentingPartitionedWorld = bRepresentingWorldPartitionedWorld;
 
 	int32 NumPinnedItems = 0;
 	if (const UWorldPartition* const WorldPartition = RepresentingWorld->GetWorldPartition())
diff --git a/Engine/Source/Editor/SceneOutliner/Private/SceneOutlinerModule.cpp b/Engine/Source/Editor/SceneOutliner/Private/SceneOutlinerModule.cpp
--- a/Engine/Source/Editor/SceneOutliner/Private/SceneOutlinerModule.cpp
+++ b/Engine/Source/Editor/SceneOutliner/Private/SceneOutlinerModule.cpp
@@ -97,7 +97,7 @@ TSharedRef<ISceneOutliner> FSceneOutlinerModule::CreateActorPicker(const FSceneO
 		InitOptions.ColumnMap.Add(FSceneOutlinerBuiltInColumnTypes::Label(), FSceneOutlinerColumnInfo(ESceneOutlinerColumnVisibility::Visible, 0, FCreateSceneOutlinerColumn(), false, TOptional<float>(), FSceneOutlinerBuiltInColumnTypes::Label_Localized()));
 		InitOptions.ColumnMap.Add(FSceneOutlinerBuiltInColumnTypes::ActorInfo(), FSceneOutlinerColumnInfo(ESceneOutlinerColumnVisibility::Visible, 10, FCreateSceneOutlinerColumn(), true, TOptional<float>(), FSceneOutlinerBuiltInColumnTypes::ActorInfo_Localized()));
 		InitOptions.ColumnMap.Add(FSceneOutlinerBuiltInColumnTypes::Pinned(), FSceneOutlinerColumnInfo(ESceneOutlinerColumnVisibility::Invisible, 5, FCreateSceneOutlinerColumn(), true, TOptional<float>(), FSceneOutlinerBuiltInColumnTypes::Pinned_Localized()));
-		CreateActorInfoColumns(InitOptions);
+		CreateActorInfoColumns(InitOptions, SpecifiedWorld);
 	}
 	return CreateSceneOutliner(InitOptions);
 }
@@ -133,7 +133,7 @@ TSharedRef<ISceneOutliner> FSceneOutlinerModule::CreateComponentPicker(const FSc
 	{
 		InitOptions.ColumnMap.Add(FSceneOutlinerBuiltInColumnTypes::Label(), FSceneOutlinerColumnInfo(ESceneOutlinerColumnVisibility::Visible, 0, FCreateSceneOutlinerColumn(), false, TOptional<float>(), FSceneOutlinerBuiltInColumnTypes::Label_Localized()));
 		InitOptions.ColumnMap.Add(FSceneOutlinerBuiltInColumnTypes::ActorInfo(), FSceneOutlinerColumnInfo(ESceneOutlinerColumnVisibility::Visible, 10, FCreateSceneOutlinerColumn(), true, TOptional<float>(), FSceneOutlinerBuiltInColumnTypes::ActorInfo_Localized()));
-		CreateActorInfoColumns(InitOptions);
+		CreateActorInfoColumns(InitOptions, SpecifiedWorld);
 	}
 	return CreateSceneOutliner(InitOptions);
 }
@@ -153,13 +153,22 @@ TSharedRef< ISceneOutliner > FSceneOutlinerModule::CreateActorBrowser(const FSce
 		InitOptions.ColumnMap.Add(FSceneOutlinerBuiltInColumnTypes::Gutter(), FSceneOutlinerColumnInfo(ESceneOutlinerColumnVisibility::Visible, 0, FCreateSceneOutlinerColumn(), true, TOptional<float>(), FSceneOutlinerBuiltInColumnTypes::Gutter_Localized()));
 		InitOptions.ColumnMap.Add(FSceneOutlinerBuiltInColumnTypes::ActorInfo(), FSceneOutlinerColumnInfo(ESceneOutlinerColumnVisibility::Visible, 20, FCreateSceneOutlinerColumn(), true, TOptional<float>(), FSceneOutlinerBuiltInColumnTypes::ActorInfo_Localized()));
 		InitOptions.ColumnMap.Add(FSceneOutlinerBuiltInColumnTypes::SourceControl(), FSceneOutlinerColumnInfo(ESceneOutlinerColumnVisibility::Invisible, 30, FCreateSceneOutlinerColumn(), true, TOptional<float>(), FSceneOutlinerBuiltInColumnTypes::SourceControl_Localized()));
-		InitOptions.ColumnMap.Add(FSceneOutlinerBuiltInColumnTypes::Pinned(), FSceneOutlinerColumnInfo(ESceneOutlinerColumnVisibility::Visible, 5, FCreateSceneOutlinerColumn(), true, TOptional<float>(), FSceneOutlinerBuiltInColumnTypes::Pinned_Localized()));
-		CreateActorInfoColumns(InitOptions);
+
+		// We don't want the pinned column in non wp levels
+		if (UWorld* WorldPtr = SpecifiedWorld.Get())
+		{
+			if (WorldPtr->IsPartitionedWorld())
+			{
+				InitOptions.ColumnMap.Add(FSceneOutlinerBuiltInColumnTypes::Pinned(), FSceneOutlinerColumnInfo(ESceneOutlinerColumnVisibility::Visible, 5, FCreateSceneOutlinerColumn(), true, TOptional<float>(), FSceneOutlinerBuiltInColumnTypes::Pinned_Localized()));
+			}
+		}
+		
+		CreateActorInfoColumns(InitOptions, SpecifiedWorld);
 	}
 	return CreateSceneOutliner(InitOptions);
 }
 
-void FSceneOutlinerModule::CreateActorInfoColumns(FSceneOutlinerInitializationOptions& InInitOptions) const
+void FSceneOutlinerModule::CreateActorInfoColumns(FSceneOutlinerInitializationOptions& InInitOptions, TWeakObjectPtr<UWorld> SpecifiedWorld) const
 {
 	FGetTextForItem MobilityInfoText = FGetTextForItem::CreateLambda([](const ISceneOutlinerTreeItem& Item) -> FString
 	{
@@ -193,21 +202,6 @@ void FSceneOutlinerModule::CreateActorInfoColumns(FSceneOutlinerInitializationOp
 		return FString();
 	});
 
-	FGetTextForItem LevelInfoText = FGetTextForItem::CreateLambda([](const ISceneOutlinerTreeItem& Item) -> FString
-	{
-		if (const FActorTreeItem* ActorItem = Item.CastTo<FActorTreeItem>())
-		{
-			AActor* Actor = ActorItem->Actor.Get();
-			if (!Actor)
-			{
-				return FString();
-			}
-
-			return FPackageName::GetShortName(Actor->GetOutermostObject()->GetName());
-		}
-		return FString();
-	});
-
 	FGetTextForItem LayerInfoText = FGetTextForItem::CreateLambda([](const ISceneOutlinerTreeItem& Item) -> FString
 	{
 		FString Result;
@@ -327,7 +321,7 @@ void FSceneOutlinerModule::CreateActorInfoColumns(FSceneOutlinerInitializationOp
 		return FString();
 	});
 
-	FGetTextForItem PackageShortNameInfoText = FGetTextForItem::CreateLambda([](const ISceneOutlinerTreeItem& Item) -> FString
+	FGetTextForItem LevelInfoText = FGetTextForItem::CreateLambda([](const ISceneOutlinerTreeItem& Item) -> FString
 	{
 		if (const FActorTreeItem* ActorItem = Item.CastTo<FActorTreeItem>())
 		{
@@ -387,12 +381,25 @@ void FSceneOutlinerModule::CreateActorInfoColumns(FSceneOutlinerInitializationOp
 				ColumnName));
 	};
 
+	// The "Level" column should be named "Package Short Name" in wp enabled levels
+	auto LevelColumnName = TAttribute<FText>::CreateLambda([SpecifiedWorld]() -> FText
+	{
+		if (UWorld* WorldPtr = SpecifiedWorld.Get())
+		{
+			if (WorldPtr->IsPartitionedWorld())
+			{
+				return FSceneOutlinerBuiltInColumnTypes::PackageShortName_Localized();
+			}
+		}
+
+		return FSceneOutlinerBuiltInColumnTypes::Level_Localized();
+	});
+
 	AddTextInfoColumn(FSceneOutlinerBuiltInColumnTypes::Mobility(), FSceneOutlinerBuiltInColumnTypes::Mobility_Localized(), MobilityInfoText);
-	AddTextInfoColumn(FSceneOutlinerBuiltInColumnTypes::Level(), FSceneOutlinerBuiltInColumnTypes::Level_Localized(), LevelInfoText);
+	AddTextInfoColumn(FSceneOutlinerBuiltInColumnTypes::Level(), LevelColumnName, LevelInfoText);
 	AddTextInfoColumn(FSceneOutlinerBuiltInColumnTypes::Layer(), FSceneOutlinerBuiltInColumnTypes::Layer_Localized(), LayerInfoText);
 	AddTextInfoColumn(FSceneOutlinerBuiltInColumnTypes::DataLayer(), FSceneOutlinerBuiltInColumnTypes::DataLayer_Localized(), DataLayerInfoText);
 	AddTextInfoColumn(FSceneOutlinerBuiltInColumnTypes::IDName(), FSceneOutlinerBuiltInColumnTypes::IDName_Localized(), InternalNameInfoText);
-	AddTextInfoColumn(FSceneOutlinerBuiltInColumnTypes::PackageShortName(), FSceneOutlinerBuiltInColumnTypes::PackageShortName_Localized(), PackageShortNameInfoText);
 	AddTextInfoColumn(FSceneOutlinerBuiltInColumnTypes::UncachedLights(), FSceneOutlinerBuiltInColumnTypes::UncachedLights_Localized(), UncachedLightsInfoText);
 }
 
diff --git a/Engine/Source/Editor/SceneOutliner/Public/SceneOutlinerMenuContext.h b/Engine/Source/Editor/SceneOutliner/Public/SceneOutlinerMenuContext.h
--- a/Engine/Source/Editor/SceneOutliner/Public/SceneOutlinerMenuContext.h
+++ b/Engine/Source/Editor/SceneOutliner/Public/SceneOutlinerMenuContext.h
@@ -19,6 +19,7 @@ public:
 		: Super(ObjectInitializer)
 	{
 		bShowParentTree = false;
+		bRepresentingPartitionedWorld = false;
 		NumSelectedItems = 0;
 		NumSelectedFolders = 0;
 		NumWorldsSelected = 0;
@@ -28,6 +29,7 @@ public:
 	TWeakPtr<SSceneOutliner> SceneOutliner;
 
 	bool bShowParentTree;
+	bool bRepresentingPartitionedWorld;
 	int32 NumSelectedItems;
 	int32 NumSelectedFolders;
 	int32 NumWorldsSelected;
diff --git a/Engine/Source/Editor/SceneOutliner/Public/SceneOutlinerModule.h b/Engine/Source/Editor/SceneOutliner/Public/SceneOutlinerModule.h
--- a/Engine/Source/Editor/SceneOutliner/Public/SceneOutlinerModule.h
+++ b/Engine/Source/Editor/SceneOutliner/Public/SceneOutlinerModule.h
@@ -132,7 +132,7 @@ private:
 	FOnColumnPermissionListChanged ColumnPermissionListChanged;
 
 	// Function to recreate what originally was FActorInfoColumn as separate columns using FTextInfoColumn
-	void CreateActorInfoColumns(FSceneOutlinerInitializationOptions& InInitOptions) const;
+	void CreateActorInfoColumns(FSceneOutlinerInitializationOptions& InInitOptions, TWeakObjectPtr<UWorld> SpecifiedWorld) const;
 
 public:
 
diff --git a/Engine/Source/Editor/WorldPartitionEditor/Private/WorldPartition/SWorldPartitionEditorGrid.cpp b/Engine/Source/Editor/WorldPartitionEditor/Private/WorldPartition/SWorldPartitionEditorGrid.cpp
--- a/Engine/Source/Editor/WorldPartitionEditor/Private/WorldPartition/SWorldPartitionEditorGrid.cpp
+++ b/Engine/Source/Editor/WorldPartitionEditor/Private/WorldPartition/SWorldPartitionEditorGrid.cpp
@@ -92,10 +92,14 @@ void SWorldPartitionEditorGrid::Refresh()
 	TWeakPtr<class ILevelEditor> LevelEditor = FModuleManager::GetModuleChecked<FLevelEditorModule>(TEXT("LevelEditor")).GetLevelEditorInstance();
 	if (LevelEditor.IsValid())
 	{
-		TSharedPtr<class ISceneOutliner> SceneOutlinerPtr = LevelEditor.Pin()->GetSceneOutliner();
-		if (SceneOutlinerPtr)
+		TArray<TWeakPtr<class ISceneOutliner>> SceneOutlinerPtrs = LevelEditor.Pin()->GetAllSceneOutliners();
+
+		for (TWeakPtr<class ISceneOutliner> SceneOutlinerPtr : SceneOutlinerPtrs)
 		{
-			SceneOutlinerPtr->FullRefresh();
+			if (TSharedPtr<class ISceneOutliner> SceneOutlinerPin = SceneOutlinerPtr.Pin())
+			{
+				SceneOutlinerPin->FullRefresh();
+			}
 		}
 	}
 }
