Adding a button to details pannel and a mecanism to register the promotion of elements. (Hydration)

diff --git a/Engine/Source/Editor/EditorStyle/Private/StarshipStyle.cpp b/Engine/Source/Editor/EditorStyle/Private/StarshipStyle.cpp
--- a/Engine/Source/Editor/EditorStyle/Private/StarshipStyle.cpp
+++ b/Engine/Source/Editor/EditorStyle/Private/StarshipStyle.cpp
@@ -473,6 +473,12 @@ void FStarshipEditorStyle::FStyle::SetupGeneralStyles()
 		Set("Icons.Toolbar.Import", new CORE_IMAGE_BRUSH_SVG( "Starship/Common/import", Icon20x20) );
 	}
 
+	// Typed Elements Icons
+	{
+		Set("Icons.PromoteElements", new IMAGE_BRUSH_SVG("Starship/Common/convert", Icon20x20));
+		Set("Icons.DemoteElements", new IMAGE_BRUSH_SVG("Starship/Common/convert", Icon20x20));
+	}
+
 	Set("UnrealDefaultThumbnail", new IMAGE_BRUSH("Starship/Common/Unreal_DefaultThumbnail", FVector2D(256, 256)));
 
 	Set( "WarningStripe", new IMAGE_BRUSH( "Common/WarningStripe", FVector2D(20,6), FLinearColor::White, ESlateBrushTileType::Horizontal ) );
diff --git a/Engine/Source/Editor/LevelEditor/Private/SActorDetails.cpp b/Engine/Source/Editor/LevelEditor/Private/SActorDetails.cpp
--- a/Engine/Source/Editor/LevelEditor/Private/SActorDetails.cpp
+++ b/Engine/Source/Editor/LevelEditor/Private/SActorDetails.cpp
@@ -11,6 +11,7 @@
 #include "Elements/Framework/EngineElementsLibrary.h"
 #include "Elements/Framework/TypedElementRegistry.h"
 #include "Elements/Framework/TypedElementSelectionSet.h"
+#include "Elements/Interfaces/TypedElementWorldInterface.h"
 #include "Engine/Blueprint.h"
 #include "Engine/BlueprintGeneratedClass.h"
 #include "Engine/Selection.h"
@@ -22,6 +23,7 @@
 #include "Kismet2/KismetEditorUtilities.h"
 #include "LevelEditor.h"
 #include "LevelEditorGenericDetails.h"
+#include "LevelEditorSubsystem.h"
 #include "Modules/ModuleManager.h"
 #include "PropertyEditorModule.h"
 #include "SSubobjectEditor.h"
@@ -34,12 +36,136 @@
 #include "UnrealEdGlobals.h"
 #include "Widgets/Docking/SDockTab.h"
 #include "Widgets/Images/SImage.h"
+#include "Widgets/Input/SButton.h"
 #include "Widgets/Layout/SBorder.h"
 #include "Widgets/Layout/SBox.h"
 #include "Widgets/Layout/SSplitter.h"
 #include "Widgets/SBoxPanel.h"
 #include "Widgets/Text/SRichTextBlock.h"
 
+#define LOCTEXT_NAMESPACE "SActorDetails"
+
+namespace UE::LevelEditor::Private
+{
+	class SElementSelectionDetailsButtons : public SCompoundWidget
+	{
+	public:
+		SLATE_BEGIN_ARGS(SElementSelectionDetailsButtons)
+		{}
+
+		SLATE_END_ARGS()
+
+		void Construct(const FArguments& InArgs, TSharedRef<SWidget> SubobjectEditorButtonBox, TFunction<TArray<TTypedElement<ITypedElementDetailsInterface>>()>&& InGetDetailsHandles)
+		{
+			GetDetailsHandles = MoveTemp(InGetDetailsHandles);
+	
+			ChildSlot
+				[
+					SNew(SHorizontalBox)
+					+ SHorizontalBox::Slot()
+					.VAlign(VAlign_Center)
+					.Padding(0.0f, 0.0f, 4.0f, 0.0f)
+					.AutoWidth()
+					[
+						SAssignNew(PromoteElementButton, SButton)
+						.ButtonStyle(FAppStyle::Get(), "DetailsView.NameAreaButton")
+						.ContentPadding(0)
+						.HAlign(HAlign_Center)
+						.VAlign(VAlign_Center)
+						.AddMetaData<FTagMetaData>(FTagMetaData(TEXT("Element.PromoteElement")))
+						.OnClicked(this, &SElementSelectionDetailsButtons::OnPromoteElement)
+						.ToolTipText(LOCTEXT("PromoteElementTooltip", "Promote the selected elements."))
+						[
+							SNew(SImage)
+							.Image(FAppStyle::Get().GetBrush("Icons.PromoteElements"))
+						]
+					]
+					+ SHorizontalBox::Slot()
+					.VAlign(VAlign_Center)
+					.AutoWidth()
+					[
+						SubobjectEditorButtonBox
+					]
+				];
+		}
+
+		void UpdateSelection(TArrayView<const TTypedElement<ITypedElementDetailsInterface>>* InDetailsElementsPtr)
+		{
+			TArrayView<const TTypedElement<ITypedElementDetailsInterface>> DetailsElementsView;
+			TArray<TTypedElement<ITypedElementDetailsInterface>> DetailsElementsContainer;
+
+			if (InDetailsElementsPtr)
+			{
+				DetailsElementsView = *InDetailsElementsPtr;
+			}
+			else
+			{
+				DetailsElementsContainer = GetDetailsHandles();
+				DetailsElementsView = MakeArrayView(DetailsElementsContainer);
+			}
+
+			bool bCanPromoteAnElement = false;
+			UTypedElementRegistry* Registry = UTypedElementRegistry::GetInstance();
+			for (const TTypedElement<ITypedElementDetailsInterface>& DetailsElement : DetailsElementsView)
+			{
+				if (TTypedElement<ITypedElementWorldInterface> WorldElement = Registry->GetElement<ITypedElementWorldInterface>(DetailsElement))
+				{
+					if (WorldElement.CanPromoteElement())
+					{
+						bCanPromoteAnElement = true;
+						break;
+					}
+				}
+			}
+
+			if (bCanPromoteAnElement)
+			{
+				PromoteElementButton->SetVisibility(EVisibility::Visible);
+			}
+			else
+			{
+				PromoteElementButton->SetVisibility(EVisibility::Collapsed);
+			}
+		}
+
+	private:
+		FReply OnPromoteElement()
+		{
+			FScopedTransaction Transaction(LOCTEXT("PromoteElementsTransaction", "Promote Elements"));
+
+			TArray<TTypedElement<ITypedElementDetailsInterface>> DetailsElements = GetDetailsHandles();
+
+			TArray<AActor*> Actors;
+			Actors.Reserve(DetailsElements.Num());
+
+			UTypedElementSelectionSet* SelectionSet = GEditor->GetEditorSubsystem<ULevelEditorSubsystem>()->GetSelectionSet();
+			FTypedElementSelectionOptions SelectionOptions;
+			UTypedElementRegistry* Registry = UTypedElementRegistry::GetInstance();
+
+			for (const TTypedElement<ITypedElementDetailsInterface>& DetailsElement : DetailsElements)
+			{
+				if (TTypedElement<ITypedElementWorldInterface> WorldElement = Registry->GetElement<ITypedElementWorldInterface>(DetailsElement))
+				{
+					if (WorldElement.CanPromoteElement())
+					{
+						SelectionSet->DeselectElement(WorldElement, SelectionOptions);
+					}
+
+					if (FTypedElementHandle PromotedElement = WorldElement.PromoteElement())
+					{
+						SelectionSet->SelectElement(PromotedElement, SelectionOptions);
+					}
+				}
+			}
+
+			return FReply::Handled();
+		}
+
+		TSharedPtr<SButton> PromoteElementButton;
+		TFunction<TArray<TTypedElement<ITypedElementDetailsInterface>>()> GetDetailsHandles;
+	};
+}
+
 class SActorDetailsUneditableComponentWarning : public SCompoundWidget
 {
 public:
@@ -152,9 +278,55 @@ void SActorDetails::Construct(const FArguments& InArgs, UTypedElementSelectionSe
 
 	ComponentsBox->SetContent(SubobjectEditor.ToSharedRef());
 
-	TSharedRef<SWidget> ButtonBox = SubobjectEditor->GetToolButtonsBox().ToSharedRef();
-	ButtonBox->SetVisibility(MakeAttributeSP(this, &SActorDetails::GetComponentEditorVisibility));
-	DetailsView->SetNameAreaCustomContent( ButtonBox );
+	TSharedRef<SWidget> SubobjectEditorButtonBox = SubobjectEditor->GetToolButtonsBox().ToSharedRef();
+	SubobjectEditorButtonBox->SetVisibility(MakeAttributeSP(this, &SActorDetails::GetComponentEditorVisibility));
+
+
+	TFunction<TArray<TTypedElement<ITypedElementDetailsInterface>>()> GetDetailsHandles = [this]()
+	{
+		TArray<TTypedElement<ITypedElementDetailsInterface>> DetailsElements;
+
+		// Regenerate the details handles
+		if (bHasSelectionOverride)
+		{
+			UTypedElementRegistry* Registry = UTypedElementRegistry::GetInstance();
+			DetailsElements.Reserve(SelectionOverrideActors.Num());
+
+			for (AActor* Actor : SelectionOverrideActors)
+			{
+				if (FTypedElementHandle ActorElementHandle = UEngineElementsLibrary::AcquireEditorActorElementHandle(Actor))
+				{
+					if (TTypedElement<ITypedElementDetailsInterface> ActorDetailsHandle = Registry->GetElement<ITypedElementDetailsInterface>(ActorElementHandle))
+					{
+						// Check if the actor element does to provide a details object
+						if (TUniquePtr<ITypedElementDetailsObject> ElementDetailsObject = ActorDetailsHandle.GetDetailsObject())
+						{
+							DetailsElements.Add(ActorDetailsHandle);
+						}
+					}
+				}
+			}
+		}
+		else
+		{
+			DetailsElements.Reserve(SelectionSet->GetNumSelectedElements());
+			SelectionSet->ForEachSelectedElement<ITypedElementDetailsInterface>([&DetailsElements](const TTypedElement<ITypedElementDetailsInterface>& InDetailsElement)
+				{
+					// Check if the element does to provide a details object
+					if (TUniquePtr<ITypedElementDetailsObject> ElementDetailsObject = InDetailsElement.GetDetailsObject())
+					{
+						DetailsElements.Add(InDetailsElement);
+					}
+					return true;
+				});
+		}
+
+		return DetailsElements;
+	};
+
+
+	TSharedRef<SWidget> ButtonBox = SAssignNew(ElementSelectionDetailsButtons, UE::LevelEditor::Private::SElementSelectionDetailsButtons, SubobjectEditorButtonBox, MoveTemp(GetDetailsHandles));
+	DetailsView->SetNameAreaCustomContent(ButtonBox);
 
 	ChildSlot
 	[
@@ -274,7 +446,7 @@ void SActorDetails::OverrideSelection(const TArray<AActor*>& InActors, const boo
 	UTypedElementRegistry* Registry = UTypedElementRegistry::GetInstance();
 
 	TArray<TTypedElement<ITypedElementDetailsInterface>> DetailsElements;
-	DetailsElements.Reserve(SelectionSet->GetNumSelectedElements());
+	DetailsElements.Reserve(InActors.Num());
 	for (AActor* Actor : InActors)
 	{
 		if (FTypedElementHandle ActorElementHandle = UEngineElementsLibrary::AcquireEditorActorElementHandle(Actor))
@@ -313,11 +485,10 @@ void SActorDetails::RefreshTopLevelElements(TArrayView<const TTypedElement<IType
 		}
 	}
 
-	// Update the underlying details view
-	SetElementDetailsObjects(TopLevelElements, bForceRefresh, bOverrideLock);
+	// Update the underlying details view and the Elements buttons
+	SetElementDetailsObjects(TopLevelElements, bForceRefresh, bOverrideLock, &InDetailsElements);
 
 	// Update the Subobject tree if we were asked to edit a single actor
-	
 	if (AActor* Actor = GetActorContext())
 	{
 		// Enable the selection guard to prevent OnTreeSelectionChanged() from altering the editor's component selection
@@ -402,7 +573,7 @@ void SActorDetails::RefreshSubobjectTreeElements(TArrayView<const FSubobjectEdit
 	}
 }
 
-void SActorDetails::SetElementDetailsObjects(TArrayView<const TUniquePtr<ITypedElementDetailsObject>> InElementDetailsObjects, const bool bForceRefresh, const bool bOverrideLock)
+void SActorDetails::SetElementDetailsObjects(TArrayView<const TUniquePtr<ITypedElementDetailsObject>> InElementDetailsObjects, const bool bForceRefresh, const bool bOverrideLock, TArrayView<const TTypedElement<ITypedElementDetailsInterface>>* InDetailsElementsPtr)
 {
 	TArray<UObject*> DetailsObjects;
 	DetailsObjects.Reserve(InElementDetailsObjects.Num());
@@ -414,6 +585,7 @@ void SActorDetails::SetElementDetailsObjects(TArrayView<const TUniquePtr<ITypedE
 		}
 	}
 	DetailsView->SetObjects(DetailsObjects, bForceRefresh, bOverrideLock);
+	ElementSelectionDetailsButtons->UpdateSelection(InDetailsElementsPtr);
 }
 
 void SActorDetails::PostUndo(bool bSuccess)
@@ -880,3 +1052,5 @@ void SActorDetails::OnObjectsReplaced(const TMap<UObject*, UObject*>& InReplacem
 		SubobjectEditor->UpdateTree();
 	}
 }
+
+#undef LOCTEXT_NAMESPACE
diff --git a/Engine/Source/Editor/LevelEditor/Private/SActorDetails.h b/Engine/Source/Editor/LevelEditor/Private/SActorDetails.h
--- a/Engine/Source/Editor/LevelEditor/Private/SActorDetails.h
+++ b/Engine/Source/Editor/LevelEditor/Private/SActorDetails.h
@@ -25,6 +25,11 @@ class SSubobjectEditor;
 class ISCSEditorUICustomization;
 class FSubobjectEditorTreeNode;
 
+namespace UE::LevelEditor::Private
+{
+	class SElementSelectionDetailsButtons;
+}
+
 /**
  * Wraps a details panel customized for viewing actors
  */
@@ -76,7 +81,7 @@ public:
 private:
 	void RefreshTopLevelElements(TArrayView<const TTypedElement<ITypedElementDetailsInterface>> InDetailsElements, const bool bForceRefresh, const bool bOverrideLock);
 	void RefreshSubobjectTreeElements(TArrayView<const TSharedPtr<FSubobjectEditorTreeNode>> InSelectedNodes, const bool bForceRefresh, const bool bOverrideLock);
-	void SetElementDetailsObjects(TArrayView<const TUniquePtr<ITypedElementDetailsObject>> InElementDetailsObjects, const bool bForceRefresh, const bool bOverrideLock);
+	void SetElementDetailsObjects(TArrayView<const TUniquePtr<ITypedElementDetailsObject>> InElementDetailsObjects, const bool bForceRefresh, const bool bOverrideLock, TArrayView<const TTypedElement<ITypedElementDetailsInterface>>* InDetailsElementsPtr = nullptr);
 
 	AActor* GetActorContext() const;
 	UObject* GetActorContextAsObject() const;
@@ -109,6 +114,8 @@ private:
 	/** The subobject editor provides a tree widget that allows for editing of subobjects */
 	TSharedPtr<SSubobjectEditor> SubobjectEditor;
 
+	TSharedPtr<UE::LevelEditor::Private::SElementSelectionDetailsButtons> ElementSelectionDetailsButtons;
+
 	// The selection set this details panel is observing
 	UTypedElementSelectionSet* SelectionSet = nullptr;
 
diff --git a/Engine/Source/Editor/UnrealEd/Private/Elements/SMInstance/SMInstanceElementEditorWorldInterface.cpp b/Engine/Source/Editor/UnrealEd/Private/Elements/SMInstance/SMInstanceElementEditorWorldInterface.cpp
--- a/Engine/Source/Editor/UnrealEd/Private/Elements/SMInstance/SMInstanceElementEditorWorldInterface.cpp
+++ b/Engine/Source/Editor/UnrealEd/Private/Elements/SMInstance/SMInstanceElementEditorWorldInterface.cpp
@@ -2,11 +2,13 @@
 
 #include "Elements/SMInstance/SMInstanceElementEditorWorldInterface.h"
 #include "Elements/SMInstance/SMInstanceElementData.h"
-#include "Components/InstancedStaticMeshComponent.h"
 
 #include "Elements/Framework/EngineElementsLibrary.h"
 #include "Elements/Framework/TypedElementSelectionSet.h"
 
+#include "Components/InstancedStaticMeshComponent.h"
+#include "Engine/StaticMeshActor.h"
+
 bool USMInstanceElementEditorWorldInterface::CanDeleteElement(const FTypedElementHandle& InElementHandle)
 {
 	const FSMInstanceManager SMInstance = SMInstanceElementDataUtil::GetSMInstanceFromHandle(InElementHandle);
@@ -138,3 +140,78 @@ void USMInstanceElementEditorWorldInterface::DuplicateElements(TArrayView<const
 		}
 	}
 }
+
+bool USMInstanceElementEditorWorldInterface::CanPromoteElement(const FTypedElementHandle& InElementHandle)
+{
+	if (const FSMInstanceManager SMInstance = SMInstanceElementDataUtil::GetSMInstanceFromHandle(InElementHandle))
+	{
+		if (UWorld* ElementWorld = SMInstance.GetISMComponent()->GetWorld())
+		{
+			return !ElementWorld->IsGameWorld() && SMInstance.CanDeleteSMInstance();
+		}
+	}
+
+	return false;
+}
+
+FTypedElementHandle USMInstanceElementEditorWorldInterface::PromoteElement(const FTypedElementHandle& InElementHandle, UWorld* OverrideWorld /* = nullptr */)
+{
+	UWorld* World = OverrideWorld;
+	if (CanPromoteElement(InElementHandle))
+	{
+		if (const FSMInstanceManager SMInstance = SMInstanceElementDataUtil::GetSMInstanceFromHandle(InElementHandle))
+		{
+			UInstancedStaticMeshComponent* OldComponent = SMInstance.GetISMComponent();
+			if (!World)
+			{
+				World = SMInstance.GetISMComponent()->GetWorld();
+			}
+
+			// Create Actor
+			FActorSpawnParameters SpawnParams;
+			SpawnParams.ObjectFlags = RF_Transactional;
+			SpawnParams.OverrideLevel = World->GetCurrentLevel();
+
+			AStaticMeshActor* StaticMeshActor = World->SpawnActor<AStaticMeshActor>(SpawnParams);
+
+			UStaticMeshComponent* NewComponent = StaticMeshActor->GetStaticMeshComponent();
+
+			// Copy editable properties
+			FComponentReregisterContext ComponentRegisterContext(NewComponent);
+			FProperty* Property = UStaticMeshComponent::StaticClass()->PropertyLink;
+			while (Property)
+			{
+				if (Property->HasAllPropertyFlags(CPF_Edit))
+				{
+					void* DestinationAddress = Property->ContainerPtrToValuePtr<void>(NewComponent);
+					void* SourceAddress = Property->ContainerPtrToValuePtr<void>(OldComponent);
+					Property->CopyCompleteValue(DestinationAddress, SourceAddress);
+				}
+
+				Property = Property->PropertyLinkNext;
+			}
+
+			// Set the transform
+			FTransform WorldTransform;
+			SMInstance.GetSMInstanceTransform(WorldTransform, /*bWorldSpace*/true);
+			StaticMeshActor->SetActorTransform(WorldTransform);
+
+			// Move the custom data
+			int32 CustomDataStart = OldComponent->NumCustomDataFloats * SMInstance.GetISMInstanceIndex();
+			for (int32 Index = 0; Index < OldComponent->NumCustomDataFloats; ++Index)
+			{
+				const float CustomData = OldComponent->PerInstanceSMCustomData[CustomDataStart + Index];
+				NewComponent->SetCustomPrimitiveDataFloat(Index, CustomData);
+			}
+
+			StaticMeshActor->PostEditChange();
+			StaticMeshActor->PostEditMove(true);
+
+			SMInstance.DeleteSMInstance();
+
+			return UEngineElementsLibrary::AcquireEditorActorElementHandle(StaticMeshActor);
+		}
+	}
+
+	return FTypedElementHandle();
+}
diff --git a/Engine/Source/Editor/UnrealEd/Public/Elements/SMInstance/SMInstanceElementEditorWorldInterface.h b/Engine/Source/Editor/UnrealEd/Public/Elements/SMInstance/SMInstanceElementEditorWorldInterface.h
--- a/Engine/Source/Editor/UnrealEd/Public/Elements/SMInstance/SMInstanceElementEditorWorldInterface.h
+++ b/Engine/Source/Editor/UnrealEd/Public/Elements/SMInstance/SMInstanceElementEditorWorldInterface.h
@@ -17,4 +17,6 @@ public:
 	virtual bool CanDuplicateElement(const FTypedElementHandle& InElementHandle) override;
 	virtual FTypedElementHandle DuplicateElement(const FTypedElementHandle& InElementHandle, UWorld* InWorld, const FVector& InLocationOffset) override;
 	virtual void DuplicateElements(TArrayView<const FTypedElementHandle> InElementHandles, UWorld* InWorld, const FVector& InLocationOffset, TArray<FTypedElementHandle>& OutNewElements) override;
+	virtual bool CanPromoteElement(const FTypedElementHandle& InElementHandle) override;
+	virtual FTypedElementHandle PromoteElement(const FTypedElementHandle& InElementHandle, UWorld* OverrideWorld /* = nullptr */) override;
 };
diff --git a/Engine/Source/Runtime/Engine/Private/Elements/Framework/EngineElementsLibrary.cpp b/Engine/Source/Runtime/Engine/Private/Elements/Framework/EngineElementsLibrary.cpp
--- a/Engine/Source/Runtime/Engine/Private/Elements/Framework/EngineElementsLibrary.cpp
+++ b/Engine/Source/Runtime/Engine/Private/Elements/Framework/EngineElementsLibrary.cpp
@@ -571,6 +571,21 @@ FTypedElementHandle UEngineElementsLibrary::AcquireEditorSMInstanceElementHandle
 	return FTypedElementHandle();
 }
 
+FTypedElementHandle UEngineElementsLibrary::AcquireEditorSMInstanceElementHandle(const FSMInstanceElementId& SMInstanceElementId)
+{
+	if (GIsEditor)
+	{
+		TTypedElementOwnerScopedAccess<FSMInstanceElementData> SMInstanceElement = GSMInstanceElementOwnerStore.FindElementOwner(SMInstanceElementId);
+
+		if (SMInstanceElement)
+		{
+			return SMInstanceElement->AcquireHandle();
+		}
+	}
+
+	return FTypedElementHandle();
+}
+
 void UEngineElementsLibrary::ReplaceEditorSMInstanceElementHandles(const TMap<FSMInstanceId, FSMInstanceId>& ReplacementSMInstanceIds)
 {
 	TMap<FSMInstanceElementId, FSMInstanceElementId> PotentialSMInstanceReplacements;
diff --git a/Engine/Source/Runtime/Engine/Private/Elements/Interfaces/TypedElementWorldInterface.cpp b/Engine/Source/Runtime/Engine/Private/Elements/Interfaces/TypedElementWorldInterface.cpp
--- a/Engine/Source/Runtime/Engine/Private/Elements/Interfaces/TypedElementWorldInterface.cpp
+++ b/Engine/Source/Runtime/Engine/Private/Elements/Interfaces/TypedElementWorldInterface.cpp
@@ -255,6 +255,30 @@ FScriptTypedElementHandle ITypedElementWorldInterface::DuplicateElement(const FS
 	return GetRegistry().CreateScriptHandle(DuplicateElement(NativeHandle, InWorld, InLocationOffset).GetId());
 }
 
+bool ITypedElementWorldInterface::CanPromoteElement(const FScriptTypedElementHandle& InElementHandle)
+{
+	FTypedElementHandle NativeHandle = InElementHandle.GetTypedElementHandle();
+	if (!NativeHandle)
+	{
+		FFrame::KismetExecutionMessage(TEXT("InElementHandle is not a valid handle."), ELogVerbosity::Error);
+		return false;
+	}
+
+	return CanPromoteElement(NativeHandle);
+}
+
+FScriptTypedElementHandle ITypedElementWorldInterface::PromoteElement(const FScriptTypedElementHandle& InElementHandle, UWorld* OverrideWorld)
+{
+	FTypedElementHandle NativeHandle = InElementHandle.GetTypedElementHandle();
+	if (!NativeHandle)
+	{
+		FFrame::KismetExecutionMessage(TEXT("InElementHandle is not a valid handle."), ELogVerbosity::Error);
+		return FScriptTypedElementHandle();
+	}
+
+	return GetRegistry().CreateScriptHandle(PromoteElement(NativeHandle, OverrideWorld).GetId());
+}
+
 class UTypedElementRegistry& ITypedElementWorldInterface::GetRegistry() const
 {
 	return *UTypedElementRegistry::GetInstance();
diff --git a/Engine/Source/Runtime/Engine/Public/Elements/Framework/EngineElementsLibrary.h b/Engine/Source/Runtime/Engine/Public/Elements/Framework/EngineElementsLibrary.h
--- a/Engine/Source/Runtime/Engine/Public/Elements/Framework/EngineElementsLibrary.h
+++ b/Engine/Source/Runtime/Engine/Public/Elements/Framework/EngineElementsLibrary.h
@@ -71,6 +71,7 @@ public:
 	UFUNCTION(BlueprintCallable, BlueprintPure=false, Category="TypedElementFramework|StaticMeshInstance", meta=(DisplayName="Acquire Editor SMInstance Element Handle", ScriptMethod = "AcquireEditorElementHandle"))
 	static FScriptTypedElementHandle K2_AcquireEditorSMInstanceElementHandle(const UInstancedStaticMeshComponent* ISMComponent, const int32 InstanceIndex, const bool bAllowCreate = true);
 	static FTypedElementHandle AcquireEditorSMInstanceElementHandle(const FSMInstanceId& SMInstanceId, const bool bAllowCreate = true);
+	static FTypedElementHandle AcquireEditorSMInstanceElementHandle(const FSMInstanceElementId& SMInstanceElementId);
 	static void ReplaceEditorSMInstanceElementHandles(const TMap<FSMInstanceId, FSMInstanceId>& ReplacementSMInstanceIds);
 #endif
 
diff --git a/Engine/Source/Runtime/Engine/Public/Elements/Interfaces/TypedElementWorldInterface.h b/Engine/Source/Runtime/Engine/Public/Elements/Interfaces/TypedElementWorldInterface.h
--- a/Engine/Source/Runtime/Engine/Public/Elements/Interfaces/TypedElementWorldInterface.h
+++ b/Engine/Source/Runtime/Engine/Public/Elements/Interfaces/TypedElementWorldInterface.h
@@ -192,6 +192,27 @@ public:
 	{
 	}
 
+	/**
+	 * Can the element be promoted
+	 * Generally available when the element is lighter representation of another element.
+	 * Like an instance for example.
+	 */
+	virtual bool CanPromoteElement(const FTypedElementHandle& InElementHandle)
+	{
+		return false;
+	}
+
+	/**
+	 * Promote a element when possible
+	 * Generally available when the element is lighter representation of another element.
+	 * Like an instance for example.
+	 * 
+	 * @param OverrideWorld Override the world in which the promotion might create new elements. Leave it to null to use the world from the handle.
+	 */
+	virtual FTypedElementHandle PromoteElement(const FTypedElementHandle& InElementHandle, UWorld* OverrideWorld = nullptr)
+	{
+		return FTypedElementHandle();
+	}
 
 	/**
 	 * Script Api
@@ -313,6 +334,24 @@ public:
 	UFUNCTION(BlueprintCallable, Category="TypedElementInterfaces|World")
 	virtual FScriptTypedElementHandle DuplicateElement(const FScriptTypedElementHandle& InElementHandle, UWorld* InWorld, const FVector& InLocationOffset);
 
+	/**
+	 * Can the element be promoted
+	 * Generally available when the element is a lighter representation of another element.
+	 * Like an instance for example.
+	 */
+	UFUNCTION(BlueprintCallable, Category="TypedElementInterfaces|World")
+	virtual bool CanPromoteElement(const FScriptTypedElementHandle& InElementHandle);
+
+	/**
+	 * Promote an element when possible
+	 * Generally available when the element is a lighter representation of another element.
+	 * Like an instance for example.
+	 * 
+	 * @param OverrideWorld Override the world in which the promotion might create new elements. Leave it to null to use the world from the handle.
+	 */
+	UFUNCTION(BlueprintCallable, Category="TypedElementInterfaces|World")
+	virtual FScriptTypedElementHandle PromoteElement(const FScriptTypedElementHandle& InElementHandle, UWorld* OverrideWorld = nullptr);
+
 private:
 	/**
 	 * Return the registry associated with this interface implementation
@@ -344,4 +383,6 @@ struct TTypedElement<ITypedElementWorldInterface> : public TTypedElementBase<ITy
 	bool DeleteElement(UWorld* InWorld, UTypedElementSelectionSet* InSelectionSet, const FTypedElementDeletionOptions& InDeletionOptions) const { return InterfacePtr->DeleteElement(*this, InWorld, InSelectionSet, InDeletionOptions); }
 	bool CanDuplicateElement() const { return InterfacePtr->CanDuplicateElement(*this); }
 	FTypedElementHandle DuplicateElement(UWorld* InWorld, const FVector& InLocationOffset) const { return InterfacePtr->DuplicateElement(*this, InWorld, InLocationOffset); }
+	bool CanPromoteElement() const { return InterfacePtr->CanPromoteElement(*this); }
+	FTypedElementHandle PromoteElement(UWorld* OverrideWorld = nullptr) const { return InterfacePtr->PromoteElement(*this, OverrideWorld); }
 };
