Level Editor: Check for valid actor before dereferncing in outliner filter

diff --git a/Engine/Source/Editor/LevelEditor/Private/LevelEditorToolBar.cpp b/Engine/Source/Editor/LevelEditor/Private/LevelEditorToolBar.cpp
--- a/Engine/Source/Editor/LevelEditor/Private/LevelEditorToolBar.cpp
+++ b/Engine/Source/Editor/LevelEditor/Private/LevelEditorToolBar.cpp
@@ -2053,7 +2053,7 @@ void FLevelEditorToolBar::RegisterCinematicsMenu()
 
 			// Only display MovieScene actors
 			auto ActorFilter = [&](const AActor* Actor) {
-				return Actor->IsA(ALevelSequenceActor::StaticClass());
+				return Actor && Actor->IsA(ALevelSequenceActor::StaticClass());
 			};
 			InitOptions.Filters->AddFilterPredicate<FActorTreeItem>(FActorTreeItem::FFilterPredicate::CreateLambda(ActorFilter));
 		}
