Editor Modes: hide shortcut editors for non-visible modes

A lot of editor modes depend on context derived from other tooling or are situated in their own windows, so assigning a shortcut and then using it will either crash or do nothing, respectively. This makes it so that the only modes with sensible manual assignment behavior can have shortcuts created for them.

diff --git a/Engine/Source/Editor/LevelEditor/Private/LevelEditorModesActions.cpp b/Engine/Source/Editor/LevelEditor/Private/LevelEditorModesActions.cpp
--- a/Engine/Source/Editor/LevelEditor/Private/LevelEditorModesActions.cpp
+++ b/Engine/Source/Editor/LevelEditor/Private/LevelEditorModesActions.cpp
@@ -25,7 +25,7 @@ void FLevelEditorModesCommands::RegisterCommands()
 			FInputBindingManager::Get().FindCommandInContext(GetContextName(), EditorModeCommandName);
 
 		// If a command isn't yet registered for this mode, we need to register one.
-		if ( !EditorModeCommand.IsValid() )
+		if (!EditorModeCommand.IsValid() && Mode.IsVisible())
 		{
 			FFormatNamedArguments Args;
 			FText ModeName = Mode.Name;
