Fix startup crash when assigning shortcuts to editor modes from plugins

diff --git a/Engine/Source/Editor/LevelEditor/Private/SLevelEditor.cpp b/Engine/Source/Editor/LevelEditor/Private/SLevelEditor.cpp
--- a/Engine/Source/Editor/LevelEditor/Private/SLevelEditor.cpp
+++ b/Engine/Source/Editor/LevelEditor/Private/SLevelEditor.cpp
@@ -1571,7 +1571,7 @@ void SLevelEditor::RefreshEditorModeCommands()
 				FInputBindingManager::Get().FindCommandInContext(Commands.GetContextName(), EditorModeCommandName);
 
 			// If a command isn't yet registered for this mode, we need to register one.
-			if (EditorModeCommand.IsValid() && !LevelEditorCommands->IsActionMapped(Commands.EditorModeCommands[CommandIndex]))
+			if (EditorModeCommand.IsValid() && Mode.IsVisible() && !LevelEditorCommands->IsActionMapped(Commands.EditorModeCommands[CommandIndex]))
 			{
 				LevelEditorCommands->MapAction(
 					Commands.EditorModeCommands[CommandIndex],
@@ -1579,9 +1579,8 @@ void SLevelEditor::RefreshEditorModeCommands()
 					FCanExecuteAction(),
 					FIsActionChecked::CreateSP(SharedThis(this), &SLevelEditor::IsModeActive, Mode.ID),
 					FIsActionButtonVisible::CreateSP(SharedThis(this), &SLevelEditor::ShouldShowModeInToolbar, Mode.ID));
+				CommandIndex++;
 			}
-
-			CommandIndex++;
 		}
 	}
 }
