Rename "Select" Mode to "Selection"

diff --git a/Engine/Source/Editor/LevelEditor/Private/LevelEditorModesActions.cpp b/Engine/Source/Editor/LevelEditor/Private/LevelEditorModesActions.cpp
--- a/Engine/Source/Editor/LevelEditor/Private/LevelEditorModesActions.cpp
+++ b/Engine/Source/Editor/LevelEditor/Private/LevelEditorModesActions.cpp
@@ -34,7 +34,7 @@ void FLevelEditorModesCommands::RegisterCommands()
 				ModeName = FText::FromName(Mode.ID);
 			}
 			Args.Add(TEXT("Mode"), ModeName);
-			const FText Tooltip = FText::Format( NSLOCTEXT("LevelEditor", "ModeTooltipF", "Activate {Mode} Editing Mode"), Args );
+			const FText Tooltip = FText::Format( NSLOCTEXT("LevelEditor", "ModeTooltipF", "Activate {Mode} Mode"), Args );
 
 			FInputChord DefaultKeyBinding;
 			if (Mode.IsVisible() && editorMode < EdModeKeys.Num())
diff --git a/Engine/Source/Editor/UnrealEd/Private/Tools/DefaultEdMode.cpp b/Engine/Source/Editor/UnrealEd/Private/Tools/DefaultEdMode.cpp
--- a/Engine/Source/Editor/UnrealEd/Private/Tools/DefaultEdMode.cpp
+++ b/Engine/Source/Editor/UnrealEd/Private/Tools/DefaultEdMode.cpp
@@ -26,7 +26,7 @@ UEdModeDefault::UEdModeDefault()
 {
 	Info = FEditorModeInfo(
 		FBuiltinEditorModes::EM_Default,
-		NSLOCTEXT("DefaultMode", "DisplayName", "Select"),
+		NSLOCTEXT("DefaultMode", "DisplayName", "Selection"),
 		FSlateIcon(FEditorStyle::GetStyleSetName(), "LevelEditor.SelectMode", "LevelEditor.SelectMode.Small"),
 		true, 0);
 }
