- Actor Editor Context : System can register themselves to the ActorEditorContextSubsystem and provide the necessary to apply their context value to newly created actors and optionally show a widget in the viewport.
- Refactored viewport's current level combobox and added "current data layers" in data layer outliner and "current folder" in world outliner to use this system.
- Can optionally be hidden from viewport using advanced flag in ULevelEditorViewportSettings.

diff --git a/Engine/Plugins/MovieScene/LevelSequenceEditor/Source/LevelSequenceEditor/Private/CinematicViewport/SCinematicLevelViewport.cpp b/Engine/Plugins/MovieScene/LevelSequenceEditor/Source/LevelSequenceEditor/Private/CinematicViewport/SCinematicLevelViewport.cpp
--- a/Engine/Plugins/MovieScene/LevelSequenceEditor/Source/LevelSequenceEditor/Private/CinematicViewport/SCinematicLevelViewport.cpp
+++ b/Engine/Plugins/MovieScene/LevelSequenceEditor/Source/LevelSequenceEditor/Private/CinematicViewport/SCinematicLevelViewport.cpp
@@ -140,7 +140,7 @@ class SCinematicPreviewViewport : public SLevelViewport
 {
 public:
 	virtual const FSlateBrush* OnGetViewportBorderBrush() const override { return nullptr; }
-	virtual EVisibility GetCurrentLevelTextVisibility() const override { return EVisibility::Collapsed; }
+	virtual bool IsActorEditorContextVisible() const { return false; }
 	virtual EVisibility GetSelectedActorsCurrentLevelTextVisibility() const override { return EVisibility::Collapsed; }
 	virtual EVisibility GetViewportControlsVisibility() const override { return EVisibility::Collapsed; }
 
diff --git a/Engine/Source/Editor/DataLayerEditor/Private/DataLayer/DataLayerEditorSubsystem.cpp b/Engine/Source/Editor/DataLayerEditor/Private/DataLayer/DataLayerEditorSubsystem.cpp
--- a/Engine/Source/Editor/DataLayerEditor/Private/DataLayer/DataLayerEditorSubsystem.cpp
+++ b/Engine/Source/Editor/DataLayerEditor/Private/DataLayer/DataLayerEditorSubsystem.cpp
@@ -18,6 +18,7 @@
 #include "Engine/Brush.h"
 #include "EngineUtils.h"
 #include "Editor.h"
+#include "Subsystems/ActorEditorContextSubsystem.h"
 
 //////////////////////////////////////////////////////////////////////////
 // FDataLayersBroadcast
@@ -119,6 +120,8 @@ UDataLayerEditorSubsystem* UDataLayerEditorSubsystem::Get()
 
 void UDataLayerEditorSubsystem::Initialize(FSubsystemCollectionBase& Collection)
 {
+	Collection.InitializeDependency<UActorEditorContextSubsystem>();
+
 	Super::Initialize(Collection);
 
 	// Set up the broadcast functions for DataLayerEditorSubsystem
@@ -128,15 +131,113 @@ void UDataLayerEditorSubsystem::Initialize(FSubsystemCollectionBase& Collection)
 	{
 		World->PersistentLevel->OnLoadedActorAddedToLevelEvent.AddLambda([this](AActor& InActor) { InitializeNewActorDataLayers(&InActor); });
 	}
+
+	UActorEditorContextSubsystem::Get()->RegisterClient(this);
 }
 
 void UDataLayerEditorSubsystem::Deinitialize()
 {
+	UActorEditorContextSubsystem::Get()->UnregisterClient(this);
+
 	Super::Deinitialize();
 
 	DataLayersBroadcast->Deinitialize();
 }
 
+void UDataLayerEditorSubsystem::OnExecuteActorEditorContextAction(UWorld* InWorld, const EActorEditorContextAction& InType, AActor* InActor)
+{
+	AWorldDataLayers* WorldDataLayer = InWorld->GetWorldDataLayers();
+	if (!WorldDataLayer)
+	{
+		return;
+	}
+
+	switch (InType)
+	{
+	case EActorEditorContextAction::ApplyContext:
+		check(InActor && InActor->GetWorld() == InWorld);
+		{
+			AddActorToDataLayers(InActor, WorldDataLayer->GetActorEditorContextDataLayers());
+		}
+		break;
+	case EActorEditorContextAction::ResetContext:
+		for (UDataLayer* DataLayer : WorldDataLayer->GetActorEditorContextDataLayers())
+		{
+			RemoveFromActorEditorContext(DataLayer);
+		}
+		break;
+	case EActorEditorContextAction::PushContext:
+		WorldDataLayer->PushActorEditorContext();
+		BroadcastDataLayerChanged(EDataLayerAction::Reset, NULL, NAME_None);
+		break;
+	case EActorEditorContextAction::PopContext:
+		WorldDataLayer->PopActorEditorContext();
+		BroadcastDataLayerChanged(EDataLayerAction::Reset, NULL, NAME_None);
+		break;
+	}
+}
+
+bool UDataLayerEditorSubsystem::GetActorEditorContextDisplayInfo(UWorld* InWorld, FActorEditorContextClientDisplayInfo& OutDiplayInfo) const
+{
+	AWorldDataLayers* WorldDataLayer = InWorld->GetWorldDataLayers();
+	if (WorldDataLayer && !WorldDataLayer->GetActorEditorContextDataLayers().IsEmpty())
+	{
+		OutDiplayInfo.Title = TEXT("Current Data Layers");
+		OutDiplayInfo.Brush = FEditorStyle::GetBrush(TEXT("DataLayer.Editor"));
+		return true;
+	}
+	return false;
+}
+
+TSharedRef<SWidget> UDataLayerEditorSubsystem::GetActorEditorContextWidget(UWorld* InWorld) const
+{
+	TSharedRef<SVerticalBox> OutWidget = SNew(SVerticalBox);
+	if (AWorldDataLayers* WorldDataLayer = InWorld->GetWorldDataLayers())
+	{
+		TArray<UDataLayer*> DataLayers = WorldDataLayer->GetActorEditorContextDataLayers();
+		for (UDataLayer* DataLayer : DataLayers)
+		{
+			check(IsValid(DataLayer));
+			OutWidget->AddSlot().AutoHeight()
+			[
+				SNew(SHorizontalBox)
+				+ SHorizontalBox::Slot()
+				.AutoWidth()
+				.Padding(0.0f, 1.0f, 1.0f, 1.0f)
+				[
+					SNew(SImage)
+					.ColorAndOpacity(DataLayer->GetDebugColor())
+					.Image(FAppStyle::Get().GetBrush("Level.ColorIcon"))
+				]
+				+ SHorizontalBox::Slot()
+				.AutoWidth()
+				.Padding(4.0f, 1.0f, 1.0f, 1.0f)
+				[
+					SNew(STextBlock)
+					.Text(FText::FromString(DataLayer->GetDataLayerLabel().ToString()))
+				]
+			];
+		}
+	}
+	return OutWidget;
+}
+
+void UDataLayerEditorSubsystem::AddToActorEditorContext(UDataLayer* InDataLayer)
+{
+	if (InDataLayer->AddToActorEditorContext())
+	{
+		BroadcastDataLayerChanged(EDataLayerAction::Modify, InDataLayer, NAME_None);
+	}
+}
+
+void UDataLayerEditorSubsystem::RemoveFromActorEditorContext(UDataLayer* InDataLayer)
+{
+	if (InDataLayer->RemoveFromActorEditorContext())
+	{
+		BroadcastDataLayerChanged(EDataLayerAction::Modify, InDataLayer, NAME_None);
+	}
+}
+
 bool UDataLayerEditorSubsystem::RefreshWorldPartitionEditorCells(bool bIsFromUserChange)
 {
 	if (UWorldPartition* WorldPartition = GetWorld() ? GetWorld()->GetWorldPartition() : nullptr)
@@ -951,6 +1052,7 @@ void UDataLayerEditorSubsystem::BroadcastDataLayerChanged(const EDataLayerAction
 {
 	RebuildSelectedDataLayersFromEditorSelection();
 	DataLayerChanged.Broadcast(Action, ChangedDataLayer, ChangedProperty);
+	ActorEditorContextClientChanged.Broadcast(this);
 }
 
 void UDataLayerEditorSubsystem::OnSelectionChanged()
diff --git a/Engine/Source/Editor/DataLayerEditor/Private/DataLayer/DataLayerMode.cpp b/Engine/Source/Editor/DataLayerEditor/Private/DataLayer/DataLayerMode.cpp
--- a/Engine/Source/Editor/DataLayerEditor/Private/DataLayer/DataLayerMode.cpp
+++ b/Engine/Source/Editor/DataLayerEditor/Private/DataLayer/DataLayerMode.cpp
@@ -22,6 +22,7 @@
 #include "DataLayersActorDescTreeItem.h"
 #include "EditorActorFolders.h"
 #include "Algo/Transform.h"
+#include "Algo/AnyOf.h"
 #include "ToolMenus.h"
 #include "Selection.h"
 #include "Editor.h"
@@ -899,6 +900,7 @@ void FDataLayerMode::RegisterContextMenu()
 			TArray<UDataLayer*> SelectedDataLayers = GetSelectedDataLayers(SceneOutliner);
 			const bool bSelectedDataLayersContainsLocked = !!SelectedDataLayers.FindByPredicate([](const UDataLayer* DataLayer) { return DataLayer->IsLocked(); });
 
+			bool bHasActorEditorContextDataLayers = false;
 			TArray<const UDataLayer*> AllDataLayers;
 			if (const AWorldDataLayers* WorldDataLayers = RepresentingWorld.IsValid() ? RepresentingWorld->GetWorldDataLayers() : nullptr)
 			{
@@ -907,6 +909,8 @@ void FDataLayerMode::RegisterContextMenu()
 					AllDataLayers.Add(DataLayer);
 					return true;
 				});
+
+				bHasActorEditorContextDataLayers = !WorldDataLayers->GetActorEditorContextDataLayers().IsEmpty();
 			}
 
 			{
@@ -1102,6 +1106,57 @@ void FDataLayerMode::RegisterContextMenu()
 					));
 			}
 
+			{
+				FToolMenuSection& Section = InMenu->AddSection("DataLayerActorEditorContext", LOCTEXT("DataLayerActorEditorContext", "Actor Editor Context"));
+
+				Section.AddMenuEntry("MakeCurrentDataLayers", LOCTEXT("MakeCurrentDataLayers", "Make Current Data Layer(s)"), FText(), FSlateIcon(),
+					FUIAction(
+						FExecuteAction::CreateLambda([this, SceneOutliner, SelectedDataLayers]() 
+						{
+							const FScopedDataLayerTransaction Transaction(LOCTEXT("MakeCurrentDataLayers", "Make Current Data Layer(s)"), RepresentingWorld.Get());
+							for (TWeakObjectPtr<const UDataLayer>& DataLayer : SelectedDataLayersSet)
+							{
+								if (DataLayer.IsValid() && !DataLayer->IsLocked())
+								{
+									DataLayerEditorSubsystem->AddToActorEditorContext(const_cast<UDataLayer*>(DataLayer.Get()));
+								}
+							}
+						}),
+						FCanExecuteAction::CreateLambda([this] { return Algo::AnyOf(SelectedDataLayersSet, [](const TWeakObjectPtr<const UDataLayer>& DataLayer) { return DataLayer.IsValid() && !DataLayer->IsLocked() && !DataLayer->IsInActorEditorContext(); }); })
+					));
+
+				Section.AddMenuEntry("RemoveCurrentDataLayers", LOCTEXT("RemoveCurrentDataLayers", "Remove Current Data Layer(s)"), FText(), FSlateIcon(),
+					FUIAction(
+						FExecuteAction::CreateLambda([this, SceneOutliner, SelectedDataLayers]()
+						{
+							const FScopedDataLayerTransaction Transaction(LOCTEXT("RemoveCurrentDataLayers", "Remove Current Data Layer(s)"), RepresentingWorld.Get());
+							for (TWeakObjectPtr<const UDataLayer>& DataLayer : SelectedDataLayersSet)
+							{
+								if (DataLayer.IsValid() && !DataLayer->IsLocked())
+								{
+									DataLayerEditorSubsystem->RemoveFromActorEditorContext(const_cast<UDataLayer*>(DataLayer.Get()));
+								}
+							}
+						}),
+						FCanExecuteAction::CreateLambda([this] { return Algo::AnyOf(SelectedDataLayersSet, [](const TWeakObjectPtr<const UDataLayer>& DataLayer) { return DataLayer.IsValid() && !DataLayer->IsLocked() && DataLayer->IsInActorEditorContext(); }); })
+					));
+
+				Section.AddMenuEntry("ClearCurrentDataLayers", LOCTEXT("ClearCurrentDataLayers", "Clear Current Data Layers"), FText(), FSlateIcon(),
+					FUIAction(
+						FExecuteAction::CreateLambda([this, AllDataLayers]() 
+						{
+							check(!AllDataLayers.IsEmpty());
+							{
+								const FScopedDataLayerTransaction Transaction(LOCTEXT("ClearCurrentDataLayers", "Clear Current Data Layers"), RepresentingWorld.Get());
+								for (const UDataLayer* DataLayer : AllDataLayers)
+								{
+									DataLayerEditorSubsystem->RemoveFromActorEditorContext(const_cast<UDataLayer*>(DataLayer));
+								}
+							}}),
+						FCanExecuteAction::CreateLambda([bHasActorEditorContextDataLayers] { return bHasActorEditorContextDataLayers; })
+					));
+			}
+
 			{
 				FToolMenuSection& Section = InMenu->AddSection("DataLayerVisibility", LOCTEXT("DataLayerVisibility", "Visibility"));
 
diff --git a/Engine/Source/Editor/DataLayerEditor/Private/DataLayer/SDataLayerTreeLabel.cpp b/Engine/Source/Editor/DataLayerEditor/Private/DataLayer/SDataLayerTreeLabel.cpp
--- a/Engine/Source/Editor/DataLayerEditor/Private/DataLayer/SDataLayerTreeLabel.cpp
+++ b/Engine/Source/Editor/DataLayerEditor/Private/DataLayer/SDataLayerTreeLabel.cpp
@@ -63,6 +63,21 @@ void SDataLayerTreeLabel::Construct(const FArguments& InArgs, FDataLayerTreeItem
 		SNew(SHorizontalBox)
 		+ SHorizontalBox::Slot()
 		.AutoWidth()
+		.HAlign(HAlign_Right)
+		.VAlign(VAlign_Center)
+		[
+			SNew(SBox)
+			.WidthOverride(FSceneOutlinerDefaultTreeItemMetrics::IconSize())
+			.HeightOverride(FSceneOutlinerDefaultTreeItemMetrics::IconSize())
+			[
+				SNew(SImage)
+				.Visibility_Lambda([this]() { return IsInActorEditorContext() ? EVisibility::Visible : EVisibility::Collapsed; })
+				.Image(FEditorStyle::Get().GetBrush(TEXT("SceneOutliner.MarkedAsCurrent")))
+				.ColorAndOpacity(FLinearColor::White)
+			]
+		]
+		+ SHorizontalBox::Slot()
+		.AutoWidth()
 		.VAlign(VAlign_Center)
 		.Padding(FSceneOutlinerDefaultTreeItemMetrics::IconPadding())
 		[
@@ -105,6 +120,12 @@ bool SDataLayerTreeLabel::ShouldBeHighlighted() const
 	return DataLayerTreeItem && DataLayerTreeItem->ShouldBeHighlighted();
 }
 
+bool SDataLayerTreeLabel::IsInActorEditorContext() const
+{
+	const UDataLayer* DataLayer = DataLayerPtr.Get();
+	return DataLayer && DataLayer->IsInActorEditorContext();
+}
+
 FSlateFontInfo SDataLayerTreeLabel::GetDisplayNameFont() const
 {
 	if (ShouldBeHighlighted())
diff --git a/Engine/Source/Editor/DataLayerEditor/Private/DataLayer/SDataLayerTreeLabel.h b/Engine/Source/Editor/DataLayerEditor/Private/DataLayer/SDataLayerTreeLabel.h
--- a/Engine/Source/Editor/DataLayerEditor/Private/DataLayer/SDataLayerTreeLabel.h
+++ b/Engine/Source/Editor/DataLayerEditor/Private/DataLayer/SDataLayerTreeLabel.h
@@ -33,6 +33,7 @@ private:
 	bool OnVerifyItemLabelChanged(const FText& InLabel, FText& OutErrorMessage);
 	void OnLabelCommitted(const FText& InLabel, ETextCommit::Type InCommitInfo);
 	bool ShouldBeHighlighted() const;
+	bool IsInActorEditorContext() const;
 
 	TWeakPtr<FDataLayerTreeItem> TreeItemPtr;
 	TWeakObjectPtr<UDataLayer> DataLayerPtr;
diff --git a/Engine/Source/Editor/DataLayerEditor/Public/DataLayer/DataLayerEditorSubsystem.h b/Engine/Source/Editor/DataLayerEditor/Public/DataLayer/DataLayerEditorSubsystem.h
--- a/Engine/Source/Editor/DataLayerEditor/Public/DataLayer/DataLayerEditorSubsystem.h
+++ b/Engine/Source/Editor/DataLayerEditor/Public/DataLayer/DataLayerEditorSubsystem.h
@@ -7,6 +7,7 @@
 #include "DataLayerAction.h"
 #include "WorldPartition/DataLayer/ActorDataLayer.h"
 #include "WorldPartition/DataLayer/DataLayer.h"
+#include "IActorEditorContextClient.h"
 #include "DataLayerEditorSubsystem.generated.h"
 
 class AActor;
@@ -18,7 +19,7 @@ class UWorld;
 template<typename TItemType> class IFilter;
 
 UCLASS()
-class DATALAYEREDITOR_API UDataLayerEditorSubsystem final : public UEditorSubsystem
+class DATALAYEREDITOR_API UDataLayerEditorSubsystem final : public UEditorSubsystem, public IActorEditorContextClient
 {
 	GENERATED_BODY()
 
@@ -42,6 +43,16 @@ public:
 	 */
 	virtual ~UDataLayerEditorSubsystem() {}
 
+	//~ Begin IActorEditorContextClient interface
+	virtual void OnExecuteActorEditorContextAction(UWorld* InWorld, const EActorEditorContextAction& InType, AActor* InActor = nullptr) override;
+	virtual bool GetActorEditorContextDisplayInfo(UWorld* InWorld, FActorEditorContextClientDisplayInfo& OutDiplayInfo) const override;
+	virtual bool CanResetContext(UWorld* InWorld) const override { return true; };
+	virtual TSharedRef<SWidget> GetActorEditorContextWidget(UWorld* InWorld) const override;
+	virtual FOnActorEditorContextClientChanged& GetOnActorEditorContextClientChanged() override { return ActorEditorContextClientChanged; }
+	//~ End IActorEditorContextClient interface
+	void AddToActorEditorContext(UDataLayer* InDataLayer);
+	void RemoveFromActorEditorContext(UDataLayer* InDataLayer);
+
 	/* Broadcasts whenever one or more DataLayers are modified
 	 *
 	 * Actions
@@ -648,5 +659,8 @@ private:
 	/** Auxiliary class that sets the callback functions for multiple delegates */
 	TSharedPtr<class FDataLayersBroadcast> DataLayersBroadcast;
 
+	/** Delegate used to notify changes to ActorEditorContextSubsystem */
+	FOnActorEditorContextClientChanged ActorEditorContextClientChanged;
+
 	friend class FDataLayersBroadcast;
 };
diff --git a/Engine/Source/Editor/EditorStyle/Private/SlateEditorStyle.cpp b/Engine/Source/Editor/EditorStyle/Private/SlateEditorStyle.cpp
--- a/Engine/Source/Editor/EditorStyle/Private/SlateEditorStyle.cpp
+++ b/Engine/Source/Editor/EditorStyle/Private/SlateEditorStyle.cpp
@@ -1447,6 +1447,7 @@ void FSlateEditorStyle::FStyle::SetupGeneralStyles()
 		Set( "SceneOutliner.NewFolderIcon", new IMAGE_BRUSH("Icons/icon_AddFolder_16x", Icon16x16 ) );
 		Set( "SceneOutliner.MoveToRoot", new IMAGE_BRUSH("Icons/icon_NoFolder_16x", Icon16x16 ) );
 		Set( "SceneOutliner.ChangedItemHighlight", new BOX_BRUSH( "Common/EditableTextSelectionBackground", FMargin(4.f/16.f) ) );
+		Set( "SceneOutliner.MarkedAsCurrent", new IMAGE_BRUSH("Icons/MarkedAsCurrent", Icon12x12));
 		Set( "SceneOutliner.World", new IMAGE_BRUSH( "Icons/icon_world_16x", Icon16x16 ) );
 
 		// Selection color should still be orange to align with the editor viewport.
diff --git a/Engine/Source/Editor/EditorStyle/Private/StarshipStyle.cpp b/Engine/Source/Editor/EditorStyle/Private/StarshipStyle.cpp
--- a/Engine/Source/Editor/EditorStyle/Private/StarshipStyle.cpp
+++ b/Engine/Source/Editor/EditorStyle/Private/StarshipStyle.cpp
@@ -1224,6 +1224,7 @@ void FStarshipEditorStyle::FStyle::SetupGeneralStyles()
 		Set( "SceneOutliner.FolderClosed",  new CORE_IMAGE_BRUSH_SVG("Starship/Common/folder-closed", Icon16x16, FStyleColors::AccentFolder ) );
 		Set( "SceneOutliner.FolderOpen",    new CORE_IMAGE_BRUSH_SVG("Starship/Common/folder-open", Icon16x16, FStyleColors::AccentFolder ) );
 		Set( "SceneOutliner.World", 		new CORE_IMAGE_BRUSH_SVG("Starship/Common/world", Icon16x16 ) );
+		Set( "SceneOutliner.MarkedAsCurrent", new IMAGE_BRUSH("Icons/MarkedAsCurrent", Icon12x12));
 		Set( "SceneOutliner.ChangedItemHighlight", new FSlateRoundedBoxBrush( FStyleColors::White, 1.0) );
 
 		Set( "SceneOutliner.TableViewRow", FTableRowStyle(NormalTableRowStyle));
diff --git a/Engine/Source/Editor/FoliageEdit/Private/FoliageEdMode.cpp b/Engine/Source/Editor/FoliageEdit/Private/FoliageEdMode.cpp
--- a/Engine/Source/Editor/FoliageEdit/Private/FoliageEdMode.cpp
+++ b/Engine/Source/Editor/FoliageEdit/Private/FoliageEdMode.cpp
@@ -40,7 +40,6 @@
 #include "Misc/ScopeExit.h"
 
 #include "ActorPartition/ActorPartitionSubsystem.h"
-#include "WorldPartition/DataLayer/DataLayerSubsystem.h"
 
 //Slate dependencies
 #include "IAssetViewport.h"
@@ -394,11 +393,6 @@ void FEdModeFoliage::Enter()
 	// Load UI settings from config file
 	UISettings.Load();
 
-	// For now, Only Foliage sets DataLayerSubsystem's DataLayerEditorContext, validate it's empty before updating it
-	const UDataLayerSubsystem* DataLayerSubsystem = GetWorld()->GetSubsystem<UDataLayerSubsystem>();
-	check(!DataLayerSubsystem || !DataLayerSubsystem->HasDataLayerEditorContext());
-	SetDataLayerEditorContext(UISettings.GetDataLayer());
-
 	// Bind to editor callbacks
 	FEditorDelegates::NewCurrentLevel.AddSP(this, &FEdModeFoliage::NotifyNewCurrentLevel);
 	FWorldDelegates::LevelAddedToWorld.AddSP(this, &FEdModeFoliage::NotifyLevelAddedToWorld);
@@ -451,11 +445,6 @@ void FEdModeFoliage::Enter()
 /** FEdMode: Called when the mode is exited */
 void FEdModeFoliage::Exit()
 {
-	DataLayerEditorContext.Reset();
-	// For now, Only Foliage sets DataLayerSubsystem's DataLayerEditorContext, validate it's back to be empty
-	const UDataLayerSubsystem* DataLayerSubsystem = GetWorld()->GetSubsystem<UDataLayerSubsystem>();
-	check(!DataLayerSubsystem || !DataLayerSubsystem->HasDataLayerEditorContext());
-
 	if (bToolActive)
 	{
 		EndFoliageBrushTrace();
@@ -518,15 +507,6 @@ void FEdModeFoliage::Exit()
 	FEdMode::Exit();
 }
 
-void FEdModeFoliage::SetDataLayerEditorContext(const FActorDataLayer& DataLayer)
-{
-	// Update Data Layer and update DataLayerSubsystem's DataLayerEditorContext with it
-	UISettings.SetDataLayer(DataLayer);
-	// Since this is a scope, first call Reset on TUniquePtr to delete old context before assigning the new one
-	DataLayerEditorContext.Reset();
-	DataLayerEditorContext.Reset(new FScopeChangeDataLayerEditorContext(GetWorld(), UISettings.GetDataLayer()));
-}
-
 EFoliageEditingState FEdModeFoliage::GetEditingState() const
 {
 	UWorld* World = GetWorld();
diff --git a/Engine/Source/Editor/FoliageEdit/Private/FoliageEdMode.h b/Engine/Source/Editor/FoliageEdit/Private/FoliageEdMode.h
--- a/Engine/Source/Editor/FoliageEdit/Private/FoliageEdMode.h
+++ b/Engine/Source/Editor/FoliageEdit/Private/FoliageEdMode.h
@@ -8,7 +8,6 @@
 #include "UnrealWidgetFwd.h"
 #include "EdMode.h"
 #include "Widgets/Views/SHeaderRow.h"
-#include "WorldPartition/DataLayer/ActorDataLayer.h"
 
 class AInstancedFoliageActor;
 class FCanvas;
@@ -22,7 +21,6 @@ class ULandscapeComponent;
 class UPrimitiveComponent;
 class UStaticMeshComponent;
 struct FViewportClick;
-struct FScopeChangeDataLayerEditorContext;
 
 //
 // Forward declarations.
@@ -85,8 +83,6 @@ struct FFoliageUISettings
 	void SetPaintDensity(float InPaintDensity) { PaintDensity = InPaintDensity; }
 	float GetUnpaintDensity() const { return UnpaintDensity; }
 	void SetUnpaintDensity(float InUnpaintDensity) { UnpaintDensity = InUnpaintDensity; }
-	const FActorDataLayer& GetDataLayer() const { return DataLayer; }
-	void SetDataLayer(const FActorDataLayer& InDataLayer) { DataLayer = InDataLayer; }
 	bool GetFilterLandscape() const { return bFilterLandscape ? true : false; }
 	void SetFilterLandscape(bool InbFilterLandscape) { bFilterLandscape = InbFilterLandscape; }
 	bool GetFilterStaticMesh() const { return bFilterStaticMesh ? true : false; }
@@ -188,8 +184,6 @@ private:
 	float PaintDensity;
 	float UnpaintDensity;
 
-	FActorDataLayer DataLayer;
-
 	bool IsInSingleInstantiationMode;
 	bool IsInQuickSingleInstantiationMode;
 	bool IsInQuickEraseMode;
@@ -578,9 +572,6 @@ public:
 	/*Find the relevant foliage actor with the foliage type and run the operation*/
 	static void ForEachFoliageInfo(UWorld* InWorld, const UFoliageType* FoliageType, const FSphere& BrushSphere, TFunctionRef<bool(AInstancedFoliageActor* IFA, FFoliageInfo* FoliageInfo, const UFoliageType* FoliageType)> InOperation);
 	
-	/** Changes the Data Layer editing context */
-	void SetDataLayerEditorContext(const FActorDataLayer& DataLayer);
-
 private:
 
 	void BindCommands();
@@ -722,7 +713,5 @@ private:
 
 	/** Flag to know when we are tracking a transaction in mouse delta */
 	bool bTracking;
-
-	TUniquePtr<FScopeChangeDataLayerEditorContext> DataLayerEditorContext;
 };
 
diff --git a/Engine/Source/Editor/FoliageEdit/Private/SFoliageEdit.cpp b/Engine/Source/Editor/FoliageEdit/Private/SFoliageEdit.cpp
--- a/Engine/Source/Editor/FoliageEdit/Private/SFoliageEdit.cpp
+++ b/Engine/Source/Editor/FoliageEdit/Private/SFoliageEdit.cpp
@@ -232,31 +232,6 @@ void SFoliageEdit::Construct(const FArguments& InArgs)
 							.Text(LOCTEXT("DataLayer_Text", "Data Layer"))
 							.Font(StandardFont)
 						]
-						+ SHorizontalBox::Slot()
-						.Padding(StandardRightPadding)
-						.FillWidth(2.0f)
-						.MaxWidth(100.f)
-						.VAlign(VAlign_Center)
-						[
-							SNew(SComboButton)
-							.OnGetMenuContent_Lambda([this]()
-								{ 
-									return FDataLayerPropertyTypeCustomizationHelper::CreateDataLayerMenu([this](const UDataLayer* DataLayer)
-									{ 
-										FoliageEditMode->SetDataLayerEditorContext(FActorDataLayer(DataLayer ? DataLayer->GetFName() : NAME_None)); 
-									});
-								})
-							.ContentPadding(2)
-							.ButtonContent()
-							[
-								SNew(STextBlock)
-								.Text_Lambda([this]()
-								{
-									const UDataLayer* DataLayer = UDataLayerEditorSubsystem::Get()->GetDataLayerFromName(FoliageEditMode->UISettings.GetDataLayer().Name);
-									return UDataLayer::GetDataLayerText(DataLayer);
-								})
-							]
-						]						
 					]
 					
 					+ SVerticalBox::Slot()
diff --git a/Engine/Source/Editor/LevelEditor/Private/LevelEditorSubsystem.cpp b/Engine/Source/Editor/LevelEditor/Private/LevelEditorSubsystem.cpp
--- a/Engine/Source/Editor/LevelEditor/Private/LevelEditorSubsystem.cpp
+++ b/Engine/Source/Editor/LevelEditor/Private/LevelEditorSubsystem.cpp
@@ -18,6 +18,11 @@
 #include "ToolMenus.h"
 #include "UnrealEdGlobals.h"
 #include "EditorModeManager.h"
+#include "Styling/SlateIconFinder.h"
+#include "Subsystems/ActorEditorContextSubsystem.h"
+#include "LevelInstance/LevelInstanceSubsystem.h"
+#include "LevelInstance/LevelInstanceActor.h"
+#include "ClassIconFinder.h"
 
 DEFINE_LOG_CATEGORY_STATIC(LevelEditorSubsystem, Log, All);
 
@@ -52,11 +57,19 @@ TSharedPtr<SLevelViewport> GetLevelViewport(const FName& ViewportConfigKey)
 
 void ULevelEditorSubsystem::Initialize(FSubsystemCollectionBase& Collection)
 {
+	Collection.InitializeDependency<UActorEditorContextSubsystem>();
+
 	UToolMenus::RegisterStartupCallback(FSimpleMulticastDelegate::FDelegate::CreateUObject(this, &ULevelEditorSubsystem::ExtendQuickActionMenu));
+	UActorEditorContextSubsystem::Get()->RegisterClient(this);
+	FWorldDelegates::LevelAddedToWorld.AddUObject(this, &ULevelEditorSubsystem::OnLevelAddedOrRemoved);
+	FWorldDelegates::LevelRemovedFromWorld.AddUObject(this, &ULevelEditorSubsystem::OnLevelAddedOrRemoved);
 }
 
 void ULevelEditorSubsystem::Deinitialize()
 {
+	UActorEditorContextSubsystem::Get()->UnregisterClient(this);
+	FWorldDelegates::LevelAddedToWorld.RemoveAll(this);
+	FWorldDelegates::LevelRemovedFromWorld.RemoveAll(this);
 	UToolMenus::UnRegisterStartupCallback(this);
 	UToolMenus::UnregisterOwner(this);
 }
@@ -649,4 +662,187 @@ FEditorModeTools* ULevelEditorSubsystem::GetLevelEditorModeManager()
 	return nullptr;
 }
 
+// Widget used to show current level in viewport
+class SCurrentLevelWidget : public SCompoundWidget
+{
+public:
+	SLATE_BEGIN_ARGS(SCurrentLevelWidget) {}
+	SLATE_ARGUMENT(UWorld*, World)
+	SLATE_END_ARGS()
+
+	~SCurrentLevelWidget()
+	{
+		GEditor->GetEditorWorldContext().RemoveRef(World);
+	}
+	void Construct(const FArguments& InArgs)
+	{
+		World = (InArgs._World);
+		GEditor->GetEditorWorldContext().AddRef(World);
+		CommandList = MakeShareable(new FUICommandList);
+
+		// No option to change current level for partitioned worlds
+		if (World && World->IsPartitionedWorld())
+		{
+			ChildSlot
+			[
+				// Current Level
+				SNew(SVerticalBox)
+				+ SVerticalBox::Slot()
+				.AutoHeight()
+				[
+					SNew(STextBlock)
+					.Visibility(this, &SCurrentLevelWidget::GetCurrentLevelTextVisibility)
+					.Text(this, &SCurrentLevelWidget::GetCurrentLevelText)
+				]
+				// Referencing Level Instance (if any)
+				+ SVerticalBox::Slot()
+				.AutoHeight()
+				[
+					SNew(SHorizontalBox)
+					.Visibility_Lambda([this]() { return GetEditingLevelInstance() ? EVisibility::Visible : EVisibility::Collapsed; })
+					+ SHorizontalBox::Slot()
+					.AutoWidth()
+					[
+						SNew(STextBlock)
+						.Text(LOCTEXT("FromBegin", "("))
+					]
+					+ SHorizontalBox::Slot()
+					.AutoWidth()
+					.Padding(1.f, 1.f, 1.f, 1.f)
+					[
+						SNew(SBox)
+						.WidthOverride(16)
+						.HeightOverride(16)
+						[
+							SNew(SImage)
+							.Image_Lambda([this]() { return FClassIconFinder::FindIconForActor(GetEditingLevelInstance()); })
+							.ColorAndOpacity(FAppStyle::Get().GetSlateColor("Colors.AccentGreen"))
+						]
+					]
+					+ SHorizontalBox::Slot()
+					.AutoWidth()
+					[
+						SNew(STextBlock)
+						.Text_Lambda([this]() 
+						{
+							ALevelInstance* LevelInstance = GetEditingLevelInstance();
+							return LevelInstance ? FText::FromString(LevelInstance->GetActorLabel()) : FText::GetEmpty();
+						})
+						.ColorAndOpacity(FAppStyle::Get().GetSlateColor("Colors.AccentGreen"))
+					]
+					+ SHorizontalBox::Slot()
+					.AutoWidth()
+					.Padding(0.f, 0.f, 4.f, 0.f)
+					[
+						SNew(STextBlock)
+						.Text(LOCTEXT("FromEnd", ")"))
+					]
+				]
+			];
+		}
+		else
+		{
+			ChildSlot
+			[
+				SNew(SComboButton)
+				.Cursor(EMouseCursor::Default)
+				.VAlign(VAlign_Center)
+				.ComboButtonStyle(FAppStyle::Get(), "SimpleComboButton")
+				.Visibility(this, &SCurrentLevelWidget::GetCurrentLevelButtonVisibility)
+				.OnGetMenuContent(this, &SCurrentLevelWidget::GenerateLevelMenu)
+				.ButtonContent()
+				[
+					SNew(STextBlock)
+					.Visibility(this, &SCurrentLevelWidget::GetCurrentLevelTextVisibility)
+					.Text(this, &SCurrentLevelWidget::GetCurrentLevelText)
+					.ShadowOffset(FVector2D(1, 1))
+				]
+			];
+		}
+	}
+
+private:
+	ALevelInstance* GetEditingLevelInstance() const
+	{
+		ULevelInstanceSubsystem* LevelInstanceSubsystem = GetWorld()->GetSubsystem<ULevelInstanceSubsystem>();
+		return LevelInstanceSubsystem ? LevelInstanceSubsystem->GetEditingLevelInstance() : nullptr;
+	}
+
+	FText GetCurrentLevelText() const
+	{
+		if (GetWorld() && GetWorld()->GetCurrentLevel())
+		{
+			// Get the level name 
+			const FText ActualLevelName = FText::FromName(FPackageName::GetShortFName(GetWorld()->GetCurrentLevel()->GetOutermost()->GetFName()));
+			if (GetWorld()->GetCurrentLevel() == GetWorld()->PersistentLevel)
+			{
+				FFormatNamedArguments Args;
+				Args.Add(TEXT("ActualLevelName"), ActualLevelName);
+				return FText::Format(LOCTEXT("LevelName", "{0} (Persistent)"), ActualLevelName);
+			}
+			return ActualLevelName;
+		}
+		return FText::GetEmpty();
+	}
+	
+	bool IsVisible() const
+	{
+		return (GetWorld() && (GetWorld()->GetCurrentLevel()->OwningWorld->GetLevels().Num() > 1) && (!GetWorld()->IsPartitionedWorld() || (GetWorld()->GetCurrentLevel() != GetWorld()->PersistentLevel)));
+	}
+
+	EVisibility GetCurrentLevelTextVisibility() const
+	{
+		return IsVisible() ? EVisibility::SelfHitTestInvisible : EVisibility::Collapsed;
+	}
+
+	EVisibility GetCurrentLevelButtonVisibility() const
+	{
+		return IsVisible() ? EVisibility::Visible : EVisibility::Collapsed;
+	}
+
+	TSharedRef<SWidget> GenerateLevelMenu() const
+	{
+		// Get all menu extenders for this context menu from the level editor module
+		FLevelEditorModule& LevelEditorModule = FModuleManager::GetModuleChecked<FLevelEditorModule>(TEXT("LevelEditor"));
+		TSharedRef<FUICommandList> InCommandList = CommandList.ToSharedRef();
+		TSharedPtr<FExtender> MenuExtender = LevelEditorModule.AssembleExtenders(InCommandList, LevelEditorModule.GetAllLevelEditorLevelMenuExtenders());
+
+		// Create the menu
+		FMenuBuilder LevelMenuBuilder(/*bShouldCloseWindowAfterMenuSelection*/true, InCommandList, MenuExtender);
+		LevelMenuBuilder.BeginSection("LevelListing", LOCTEXT("Levels", "Levels"));
+		LevelMenuBuilder.EndSection();
+		return LevelMenuBuilder.MakeWidget();
+	}
+
+	UWorld* GetWorld() const
+	{
+		return World;
+	}
+
+	TSharedPtr<FUICommandList> CommandList;
+	UWorld* World;
+};
+
+bool ULevelEditorSubsystem::GetActorEditorContextDisplayInfo(UWorld* InWorld, FActorEditorContextClientDisplayInfo& OutDiplayInfo) const
+{
+	const bool bIsVisible = InWorld && (InWorld->GetCurrentLevel()->OwningWorld->GetLevels().Num() > 1) && (!InWorld->IsPartitionedWorld() || (InWorld->GetCurrentLevel() != InWorld->PersistentLevel));
+	if (bIsVisible)
+	{
+		OutDiplayInfo.Title = TEXT("Current Level");
+		OutDiplayInfo.Brush = FSlateIconFinder::FindIconBrushForClass(UWorld::StaticClass());
+		return true;
+	}
+	return false;
+}
+
+void ULevelEditorSubsystem::OnLevelAddedOrRemoved(ULevel* InLevel, UWorld* InWorld)
+{
+	ActorEditorContextClientChanged.Broadcast(this);
+}
+
+TSharedRef<SWidget> ULevelEditorSubsystem::GetActorEditorContextWidget(UWorld* InWorld) const
+{
+	return SNew(SCurrentLevelWidget).World(InWorld);
+}
+
 #undef LOCTEXT_NAMESPACE
diff --git a/Engine/Source/Editor/LevelEditor/Private/SLevelViewport.cpp b/Engine/Source/Editor/LevelEditor/Private/SLevelViewport.cpp
--- a/Engine/Source/Editor/LevelEditor/Private/SLevelViewport.cpp
+++ b/Engine/Source/Editor/LevelEditor/Private/SLevelViewport.cpp
@@ -90,6 +90,7 @@
 #include "DataLayer/DataLayerEditorSubsystem.h"
 #include "SInViewportDetails.h"
 #include "Viewports/InViewportUIDragOperation.h"
+#include "SActorEditorContext.h"
 
 static const FName LevelEditorName("LevelEditor");
 static FAutoConsoleCommand EnableInViewportMenu(TEXT("Editor.EnableInViewportMenu"), TEXT("Enables the new in-viewport property menu"), FConsoleCommandDelegate::CreateStatic(&SLevelViewport::EnableInViewportMenu));
@@ -351,39 +352,10 @@ void SLevelViewport::ConstructViewportOverlayContent()
 			]
 			+ SVerticalBox::Slot()
 			.AutoHeight()
-			.Padding(2.0f, 1.0f, 2.0f, 1.0f)
 			[
-				SNew(SComboButton)
-				.Cursor(EMouseCursor::Default)
-				// Allows users to drag with the mouse to select options after opening the menu */
-				.VAlign(VAlign_Center)
-				.ComboButtonStyle(FAppStyle::Get(), "SimpleComboButton")
-				.Visibility(this, &SLevelViewport::GetCurrentLevelButtonVisibility)
-				.OnGetMenuContent(this, &SLevelViewport::GenerateLevelMenu)
-				.OnMenuOpenChanged_Lambda([this](bool){	OnFloatingButtonClicked(); })
-				.ButtonContent()
-				[
-					SNew(SHorizontalBox)
-					.Visibility(this, &SLevelViewport::GetCurrentLevelTextVisibility)
-					// Current level label
-					+ SHorizontalBox::Slot()
-					.AutoWidth()
-					.Padding(2.0f, 1.0f, 2.0f, 1.0f)
-					[
-						SNew(STextBlock)
-						.Text(this, &SLevelViewport::GetCurrentLevelText, true)
-						.ShadowOffset(FVector2D(1,1))
-					]
-					// Current level
-					+ SHorizontalBox::Slot()
-					.AutoWidth()
-					.Padding(4.0f, 1.0f, 2.0f, 1.0f)
-					[
-						SNew(STextBlock)
-						.Text(this, &SLevelViewport::GetCurrentLevelText, false)
-						.ShadowOffset(FVector2D(1, 1))
-					]
-				]
+				SNew(SActorEditorContext)
+				.World(GetWorld())
+				.Visibility_Lambda([this]() { return IsActorEditorContextVisible() ? EVisibility::Visible : EVisibility::Collapsed; })
 			]
 		]
 	];
@@ -398,26 +370,16 @@ void SLevelViewport::ConstructViewportOverlayContent()
 	];
 }
 
-
-TSharedRef<SWidget> SLevelViewport::GenerateLevelMenu() const
+bool SLevelViewport::IsActorEditorContextVisible() const
 {
-	FWorldBrowserModule& WorldBrowserModule = FModuleManager::LoadModuleChecked<FWorldBrowserModule>("WorldBrowser");
-	// Get all menu extenders for this context menu from the level editor module
-	FLevelEditorModule& LevelEditorModule = FModuleManager::GetModuleChecked<FLevelEditorModule>(TEXT("LevelEditor"));
-	TSharedRef<FUICommandList> InCommandList = GetCommandList().ToSharedRef();
-	TSharedPtr<FExtender> MenuExtender = LevelEditorModule.AssembleExtenders(InCommandList, LevelEditorModule.GetAllLevelEditorLevelMenuExtenders());
-
-	// Create the menu
-	const bool bShouldCloseWindowAfterMenuSelection = true;
-	FMenuBuilder LevelMenuBuilder(bShouldCloseWindowAfterMenuSelection, InCommandList, MenuExtender);
-	
-	LevelMenuBuilder.BeginSection("LevelListing", LOCTEXT("Levels", "Levels"));
-	LevelMenuBuilder.EndSection();
-
-	return LevelMenuBuilder.MakeWidget();
+	return ActiveViewport.IsValid() && 
+		(ActiveViewport->GetPlayInEditorIsSimulate() || !ActiveViewport->GetClient()->GetWorld()->IsGameWorld()) && 
+		(&GetLevelViewportClient() == GCurrentLevelEditingViewportClient) && 
+		GetDefault<ULevelEditorViewportSettings>()->bShowActorEditorContext &&
+		GetWorld() && 
+		GetWorld()->GetCurrentLevel();
 }
 
-
 void SLevelViewport::ConstructLevelEditorViewportClient(FLevelEditorViewportInstanceSettings& ViewportInstanceSettings)
 {
 	if (!LevelViewportClient.IsValid())
@@ -3825,72 +3787,6 @@ FText SLevelViewport::GetCurrentScreenPercentageText() const
 	return FText::FromString(FString::Printf(TEXT("%3d%%"), int32(GetLevelViewportClient().GetPreviewScreenPercentage())));
 }
 
-FText SLevelViewport::GetCurrentLevelText( bool bDrawOnlyLabel ) const
-{
-	// Display the current level and current level grid volume in the status bar
-	FText LabelName;
-	FText CurrentLevelName;
-
-	
-	if( ActiveViewport.IsValid() && (&GetLevelViewportClient() == GCurrentLevelEditingViewportClient) && GetWorld() && GetWorld()->GetCurrentLevel() != nullptr )
-	{
-		if( ActiveViewport->GetPlayInEditorIsSimulate() || !ActiveViewport->GetClient()->GetWorld()->IsGameWorld() )
-		{
-			if(bDrawOnlyLabel)
-			{
-				LabelName = LOCTEXT("CurrentLevelLabel", "Level");
-			}
-			else
-			{
-				// Get the level name 
-				FText ActualLevelName = FText::FromName(FPackageName::GetShortFName(GetWorld()->GetCurrentLevel()->GetOutermost()->GetFName()));
-
-				if(GetWorld()->GetCurrentLevel() == GetWorld()->PersistentLevel)
-				{
-					FFormatNamedArguments Args;
-					Args.Add(TEXT("ActualLevelName"), ActualLevelName);
-					CurrentLevelName = FText::Format(LOCTEXT("LevelName", "{0} (Persistent)"), ActualLevelName);
-				}
-				else
-				{
-					CurrentLevelName = ActualLevelName;
-				}
-			}
-
-			if(bDrawOnlyLabel)
-			{
-				return LabelName;
-			}
-		}
-	}
-
-	return CurrentLevelName;
-}
-
-EVisibility SLevelViewport::GetCurrentLevelTextVisibility() const
-{
-	EVisibility ContentVisibility = OnGetViewportContentVisibility();
-	if (ContentVisibility == EVisibility::Visible)
-	{
-		ContentVisibility = EVisibility::SelfHitTestInvisible;
-	}
-	return (&GetLevelViewportClient() == GCurrentLevelEditingViewportClient) 
-		&& !IsPlayInEditorViewportActive() 
-		&& GetWorld() && GetWorld()->GetCurrentLevel()->OwningWorld->GetLevels().Num() > 1
-		&& !GetWorld()->IsPartitionedWorld()
-		?  ContentVisibility : EVisibility::Collapsed;
-}
-
-EVisibility SLevelViewport::GetCurrentLevelButtonVisibility() const
-{
-	EVisibility TextVisibility = GetCurrentLevelTextVisibility();
-	if (TextVisibility == EVisibility::SelfHitTestInvisible)
-	{
-		TextVisibility = EVisibility::Visible;
-	}
-	return TextVisibility;
-}
-
 EVisibility SLevelViewport::GetSelectedActorsCurrentLevelTextVisibility() const
 {
 	EVisibility ContentVisibility = OnGetViewportContentVisibility();
@@ -3902,6 +3798,7 @@ EVisibility SLevelViewport::GetSelectedActorsCurrentLevelTextVisibility() const
 		&& (GEditor->GetSelectedActorCount() > 0) 
 		&& !IsPlayInEditorViewportActive() 
 		&& GetWorld() && GetWorld()->GetCurrentLevel() && GetWorld()->GetCurrentLevel()->OwningWorld->GetLevels().Num() > 1
+		&& !GetWorld()->IsPartitionedWorld()
 		? ContentVisibility : EVisibility::Collapsed;
 }
 
diff --git a/Engine/Source/Editor/LevelEditor/Public/LevelEditorSubsystem.h b/Engine/Source/Editor/LevelEditor/Public/LevelEditorSubsystem.h
--- a/Engine/Source/Editor/LevelEditor/Public/LevelEditorSubsystem.h
+++ b/Engine/Source/Editor/LevelEditor/Public/LevelEditorSubsystem.h
@@ -4,6 +4,7 @@
 
 #include "CoreMinimal.h"
 #include "EditorSubsystem.h"
+#include "IActorEditorContextClient.h"
 
 #include "LevelEditorSubsystem.generated.h"
 
@@ -16,7 +17,7 @@ class FEditorModeTools;
 * Subsystem for exposing Level Editor related functionality to scripts
 */
 UCLASS()
-class LEVELEDITOR_API ULevelEditorSubsystem : public UEditorSubsystem
+class LEVELEDITOR_API ULevelEditorSubsystem : public UEditorSubsystem, public IActorEditorContextClient
 {
 	GENERATED_BODY()
 
@@ -135,4 +136,20 @@ public:
 	 * The mode manager is not created in commandlet environments, because modes inherently imply user interactions.
 	 */
 	FEditorModeTools* GetLevelEditorModeManager();
+
+	//~ Begin IActorEditorContextClient interface
+	virtual void OnExecuteActorEditorContextAction(UWorld* InWorld, const EActorEditorContextAction& InType, class AActor* InActor = nullptr) {}
+	virtual bool GetActorEditorContextDisplayInfo(UWorld* InWorld, FActorEditorContextClientDisplayInfo& OutDiplayInfo) const override;
+	virtual bool CanResetContext(UWorld* InWorld) const override { return false; };
+	virtual TSharedRef<SWidget> GetActorEditorContextWidget(UWorld* InWorld) const override;
+	virtual FOnActorEditorContextClientChanged& GetOnActorEditorContextClientChanged() override { return ActorEditorContextClientChanged; }
+	//~ End IActorEditorContextClient interface
+
+private:
+
+	/** Called when a Level is added to a world or remove from a world */
+	void OnLevelAddedOrRemoved(ULevel* InLevel, UWorld* InWorld);
+
+	/** Delegate used to notify changes to ActorEditorContextSubsystem */
+	FOnActorEditorContextClientChanged ActorEditorContextClientChanged;
 };
\ No newline at end of file
diff --git a/Engine/Source/Editor/LevelEditor/Public/SLevelViewport.h b/Engine/Source/Editor/LevelEditor/Public/SLevelViewport.h
--- a/Engine/Source/Editor/LevelEditor/Public/SLevelViewport.h
+++ b/Engine/Source/Editor/LevelEditor/Public/SLevelViewport.h
@@ -58,8 +58,6 @@ public:
 	 */
 	void ConstructViewportOverlayContent();
 
-	TSharedRef<SWidget> GenerateLevelMenu() const;
-
 	/**
 	 * Constructs the level editor viewport client
 	 */
@@ -311,17 +309,16 @@ public:
 	/** Get the parent level editor for this viewport */
 	TWeakPtr<ILevelEditor> GetParentLevelEditor() const { return ParentLevelEditor; }
 
-	/** Called to get the level text */
-	FText GetCurrentLevelText( bool bDrawOnlyLabel ) const;
+	/** @return whether the the actor editor context should be displayed for this viewport */
+	virtual bool IsActorEditorContextVisible() const;
 
 	/** Called to get the screen percentage preview text */
 	FText GetCurrentScreenPercentageText() const;
 
-	/** @return The visibility of the current level text display */
-	virtual EVisibility GetCurrentLevelTextVisibility() const;
-
-	/** @return The visibility of the current level button display */
-	virtual EVisibility GetCurrentLevelButtonVisibility() const;
+	UE_DEPRECATED(5.1, "GetCurrentLevelTextVisibility not used anymore.")
+	virtual EVisibility GetCurrentLevelTextVisibility() const { return EVisibility::Collapsed; }
+	UE_DEPRECATED(5.1, "GetCurrentLevelButtonVisibility not used anymore.")
+	virtual EVisibility GetCurrentLevelButtonVisibility() const { return EVisibility::Collapsed; }
 
 	/** @return The visibility of the current level text display */
 	virtual EVisibility GetSelectedActorsCurrentLevelTextVisibility() const;
diff --git a/Engine/Source/Editor/SceneOutliner/Private/ActorBrowsingMode.cpp b/Engine/Source/Editor/SceneOutliner/Private/ActorBrowsingMode.cpp
--- a/Engine/Source/Editor/SceneOutliner/Private/ActorBrowsingMode.cpp
+++ b/Engine/Source/Editor/SceneOutliner/Private/ActorBrowsingMode.cpp
@@ -461,6 +461,64 @@ void FActorBrowsingMode::RegisterContextMenu()
 					}
 				}
 			}));
+
+		Menu->AddDynamicSection("DynamicActorEditorContext", FNewToolMenuDelegate::CreateLambda([this](UToolMenu* InMenu)
+			{
+				USceneOutlinerMenuContext* Context = InMenu->FindContext<USceneOutlinerMenuContext>();
+				if (Context && Context->SceneOutliner.IsValid() && Context->bShowParentTree)
+				{
+					FToolMenuSection& Section = InMenu->AddSection("ActorEditorContextSection", LOCTEXT("ActorEditorContextSectionName", "Actor Editor Context"));
+					SSceneOutliner* SceneOutliner = Context->SceneOutliner.Pin().Get();
+
+					if ((Context->NumSelectedItems == 1) && (Context->NumSelectedFolders == 1))
+					{
+						SceneOutliner->GetTree().GetSelectedItems()[0]->GenerateContextMenu(InMenu, *SceneOutliner);
+
+						Section.AddMenuEntry(
+							"MakeCurrentFolder",
+							LOCTEXT("MakeCurrentFolder", "Make Current Folder"),
+							FText(),
+							FSlateIcon(),
+							FUIAction(
+								FExecuteAction::CreateLambda([SceneOutliner]()
+								{
+									const FSceneOutlinerItemSelection& Selection = SceneOutliner->GetSelection();
+									if (Selection.SelectedItems.Num() == 1)
+									{
+										FSceneOutlinerTreeItemPtr Item = Selection.SelectedItems[0].Pin();
+										if (FActorFolderTreeItem* FolderItem = Item->CastTo<FActorFolderTreeItem>())
+										{
+											if (FolderItem->World.IsValid())
+											{
+												const FScopedTransaction Transaction(NSLOCTEXT("UnrealEd", "MakeCurrentActorFolder", "Make Current Actor Folder"));
+												FActorFolders::Get().SetActorEditorContextFolder(*FolderItem->World, FolderItem->GetFolder());
+											}
+										}
+									}
+								})
+							)
+						);
+					}
+
+					Section.AddMenuEntry(
+						"ClearCurrentFolder",
+						LOCTEXT("ClearCurrentFolder", "Clear Current Folder"),
+						FText(),
+						FSlateIcon(),
+						FUIAction(
+							FExecuteAction::CreateLambda([this, SceneOutliner]()
+							{
+								if (RepresentingWorld.IsValid())
+								{
+									const FScopedTransaction Transaction(NSLOCTEXT("UnrealEd", "ClearCurrentActorFolder", "Clear Current Actor Folder"));
+									FActorFolders::Get().SetActorEditorContextFolder(*RepresentingWorld.Get(), FFolder());
+								}
+							}),
+							FCanExecuteAction::CreateLambda([this] { return RepresentingWorld.IsValid() && !FActorFolders::Get().GetActorEditorContextFolder(*RepresentingWorld.Get()).IsNone(); })
+						)
+					);
+				}
+			}));
 	}
 
 	if (!ToolMenus->IsMenuRegistered(DefaultContextMenuName))
diff --git a/Engine/Source/Editor/SceneOutliner/Private/ActorFolderTreeItem.cpp b/Engine/Source/Editor/SceneOutliner/Private/ActorFolderTreeItem.cpp
--- a/Engine/Source/Editor/SceneOutliner/Private/ActorFolderTreeItem.cpp
+++ b/Engine/Source/Editor/SceneOutliner/Private/ActorFolderTreeItem.cpp
@@ -53,35 +53,55 @@ struct SActorFolderTreeLabel : FSceneOutlinerCommonLabelData, public SCompoundWi
 		ChildSlot
 			[
 				SNew(SHorizontalBox)
-
 				+ SHorizontalBox::Slot()
-			.AutoWidth()
-			.VAlign(VAlign_Center)
-			.Padding(FSceneOutlinerDefaultTreeItemMetrics::IconPadding())
-			[
-				SNew(SBox)
-				.WidthOverride(FSceneOutlinerDefaultTreeItemMetrics::IconSize())
-				.HeightOverride(FSceneOutlinerDefaultTreeItemMetrics::IconSize())
-			[
-				SNew(SImage)
-				.Image(this, &SActorFolderTreeLabel::GetIcon)
-				.ColorAndOpacity(FSlateColor::UseForeground())
-			]
-			]
-
-		+ SHorizontalBox::Slot()
-			.FillWidth(1.0f)
-			.VAlign(VAlign_Center)
-			.Padding(0.0f, 2.0f)
-			[
-				InlineTextBlock.ToSharedRef()
-			]
+				.AutoWidth()
+				.HAlign(HAlign_Right)
+				.VAlign(VAlign_Center)
+				[
+					SNew(SBox)
+					.WidthOverride(FSceneOutlinerDefaultTreeItemMetrics::IconSize())
+					.HeightOverride(FSceneOutlinerDefaultTreeItemMetrics::IconSize())
+					[
+						SNew(SImage)
+						.Visibility_Lambda([this]() { return IsInActorEditorContext() ? EVisibility::Visible : EVisibility::Collapsed; })
+						.Image(FEditorStyle::Get().GetBrush(TEXT("SceneOutliner.MarkedAsCurrent")))
+						.ColorAndOpacity(FSlateColor::UseForeground())
+					]
+				]
+				+ SHorizontalBox::Slot()
+				.AutoWidth()
+				.VAlign(VAlign_Center)
+				.Padding(FSceneOutlinerDefaultTreeItemMetrics::IconPadding())
+				[
+					SNew(SBox)
+					.WidthOverride(FSceneOutlinerDefaultTreeItemMetrics::IconSize())
+					.HeightOverride(FSceneOutlinerDefaultTreeItemMetrics::IconSize())
+					[
+						SNew(SImage)
+						.Image(this, &SActorFolderTreeLabel::GetIcon)
+						.ColorAndOpacity(FSlateColor::UseForeground())
+					]
+				]
+				+ SHorizontalBox::Slot()
+				.FillWidth(1.0f)
+				.VAlign(VAlign_Center)
+				.Padding(0.0f, 2.0f)
+				[
+					InlineTextBlock.ToSharedRef()
+				]
 			];
 	}
 
 private:
 	TWeakPtr<FActorFolderTreeItem> TreeItemPtr;
 
+	bool IsInActorEditorContext() const
+	{
+		auto Folder = TreeItemPtr.Pin();
+		const bool bIsCurrentFolder = (Folder.IsValid() && Folder->World.IsValid()) ? (FActorFolders::Get().GetActorEditorContextFolder(*Folder->World) == Folder->GetFolder()) : false;
+		return bIsCurrentFolder;
+	}
+
 	FText GetDisplayText() const
 	{
 		auto Folder = TreeItemPtr.Pin();
diff --git a/Engine/Source/Editor/UnrealEd/Classes/Settings/LevelEditorViewportSettings.h b/Engine/Source/Editor/UnrealEd/Classes/Settings/LevelEditorViewportSettings.h
--- a/Engine/Source/Editor/UnrealEd/Classes/Settings/LevelEditorViewportSettings.h
+++ b/Engine/Source/Editor/UnrealEd/Classes/Settings/LevelEditorViewportSettings.h
@@ -277,6 +277,10 @@ class UNREALED_API ULevelEditorViewportSettings
 	UPROPERTY(EditAnywhere, config, Category=LookAndFeel, meta=( DisplayName = "Clicking BSP Enables Brush" ), AdvancedDisplay)
 	uint32 bClickBSPSelectsBrush:1;
 
+	/** If true, viewport will show actor editor context (current level, current data layer(s), current folder) */
+	UPROPERTY(EditAnywhere, config, Category = LookAndFeel, meta=(AdvancedDisplay))
+	uint32 bShowActorEditorContext : 1;
+
 	/** How fast the perspective camera moves when flying through the world. */
 	UPROPERTY(config, meta=(UIMin = "1", UIMax = "8", ClampMin="1", ClampMax="8"))
 	int32 CameraSpeed;
@@ -529,7 +533,7 @@ public:
 	/** The scale to apply to spline tangent lengths */
 	UPROPERTY(EditAnywhere, config, Category = LookAndFeel, AdvancedDisplay, meta = (ClampMin = "0.00"))
 	float SplineTangentScale = 0.5f;
-	
+
 	UPROPERTY(config)
 	FVector2D LastInViewportMenuLocation;
 
diff --git a/Engine/Source/Editor/UnrealEd/Private/EditorActorFolders.cpp b/Engine/Source/Editor/UnrealEd/Private/EditorActorFolders.cpp
--- a/Engine/Source/Editor/UnrealEd/Private/EditorActorFolders.cpp
+++ b/Engine/Source/Editor/UnrealEd/Private/EditorActorFolders.cpp
@@ -17,6 +17,7 @@
 #include "UObject/ObjectSaveContext.h"
 #include "LevelInstance/LevelInstanceActor.h"
 #include "ActorFolder.h"
+#include "Subsystems/ActorEditorContextSubsystem.h"
 
 #define LOCTEXT_NAMESPACE "FActorFolders"
 
@@ -72,10 +73,12 @@ FActorFolders& FActorFolders::Get()
 void FActorFolders::Init()
 {
 	Singleton = new FActorFolders;
+	UActorEditorContextSubsystem::Get()->RegisterClient(Singleton);
 }
 
 void FActorFolders::Cleanup()
 {
+	UActorEditorContextSubsystem::Get()->UnregisterClient(Singleton);
 	delete Singleton;
 	Singleton = nullptr;
 }
@@ -100,6 +103,7 @@ void FActorFolders::BroadcastOnActorFolderCreated(UWorld& InWorld, const FFolder
 		OnFolderCreate.Broadcast(InWorld, InFolder.GetPath());
 	}
 	PRAGMA_ENABLE_DEPRECATION_WARNINGS
+	BroadcastOnActorEditorContextClientChanged();
 }
 
 void FActorFolders::BroadcastOnActorFolderDeleted(UWorld& InWorld, const FFolder& InFolder)
@@ -111,6 +115,7 @@ void FActorFolders::BroadcastOnActorFolderDeleted(UWorld& InWorld, const FFolder
 		OnFolderDelete.Broadcast(InWorld, InFolder.GetPath());
 	}
 	PRAGMA_ENABLE_DEPRECATION_WARNINGS
+	BroadcastOnActorEditorContextClientChanged();
 }
 
 void FActorFolders::BroadcastOnActorFolderMoved(UWorld& InWorld, const FFolder& InSrcFolder, const FFolder& InDstFolder)
@@ -122,6 +127,7 @@ void FActorFolders::BroadcastOnActorFolderMoved(UWorld& InWorld, const FFolder&
 		OnFolderMove.Broadcast(InWorld, InSrcFolder.GetPath(), InDstFolder.GetPath());
 	}
 	PRAGMA_ENABLE_DEPRECATION_WARNINGS
+	BroadcastOnActorEditorContextClientChanged();
 }
 
 void FActorFolders::OnLevelActorListChanged()
@@ -513,6 +519,84 @@ void FActorFolders::SetIsFolderExpanded(UWorld& InWorld, const FFolder& InFolder
 	GetOrCreateWorldFolders(InWorld).SetIsFolderExpanded(InFolder, bIsExpanded);
 }
 
+FFolder FActorFolders::GetActorEditorContextFolder(UWorld& InWorld) const
+{
+	const UWorldFolders** Folders = (const UWorldFolders **)WorldFolders.Find(&InWorld);
+	if (Folders)
+	{
+		return (*Folders)->GetActorEditorContextFolder();
+	}
+	return FFolder();
+}
+
+void FActorFolders::SetActorEditorContextFolder(UWorld& InWorld, const FFolder& InFolder)
+{
+	if (GetOrCreateWorldFolders(InWorld).SetActorEditorContextFolder(InFolder))
+	{
+		BroadcastOnActorEditorContextClientChanged();
+	}
+}
+
+void FActorFolders::OnExecuteActorEditorContextAction(UWorld* InWorld, const EActorEditorContextAction& InType, AActor* InActor)
+{
+	switch (InType)
+	{
+		case EActorEditorContextAction::ApplyContext:
+			check(InActor && InActor->GetWorld() == InWorld);
+			{
+				FFolder Folder = GetActorEditorContextFolder(*InWorld);
+				if (!Folder.IsNone())
+				{
+					// Currently not supported to change rootobject through this interface
+					if (Folder.GetRootObject() == InActor->GetFolderRootObject())
+					{
+						InActor->SetFolderPath_Recursively(Folder.GetPath());
+					}
+				}
+			}
+			break;
+		case EActorEditorContextAction::ResetContext:
+			SetActorEditorContextFolder(*InWorld, FFolder());
+			break;
+		case EActorEditorContextAction::PushContext:
+			if (UWorldFolders** Folders = (UWorldFolders**)WorldFolders.Find(InWorld))
+			{
+				(*Folders)->PushActorEditorContext();
+			}
+			break;
+		case EActorEditorContextAction::PopContext:
+			if (UWorldFolders** Folders = (UWorldFolders**)WorldFolders.Find(InWorld))
+			{
+				(*Folders)->PopActorEditorContext();
+			}
+			break;
+	}
+}
+
+bool FActorFolders::GetActorEditorContextDisplayInfo(UWorld* InWorld, FActorEditorContextClientDisplayInfo& OutDiplayInfo) const
+{
+	const FFolder Folder = GetActorEditorContextFolder(*InWorld);
+	if (!Folder.IsNone())
+	{
+		OutDiplayInfo.Title = TEXT("Current Actor Folder");
+		OutDiplayInfo.Brush = FEditorStyle::GetBrush(TEXT("SceneOutliner.FolderClosed"));
+		return true;
+	}
+	return false;
+}
+
+TSharedRef<SWidget> FActorFolders::GetActorEditorContextWidget(UWorld* InWorld) const
+{
+	const FFolder Folder = GetActorEditorContextFolder(*InWorld);
+	FText Text = (!Folder.IsNone()) ? FText::FromName(Folder.GetLeafName()) : FText::GetEmpty();
+	return SNew(STextBlock).Text(Text);
+}
+
+void FActorFolders::BroadcastOnActorEditorContextClientChanged()
+{
+	ActorEditorContextClientChanged.Broadcast(this);
+}
+
 void FActorFolders::ForEachFolder(UWorld& InWorld, TFunctionRef<bool(const FFolder&)> Operation)
 {
 	GetOrCreateWorldFolders(InWorld).ForEachFolder(Operation);
diff --git a/Engine/Source/Editor/UnrealEd/Private/SActorEditorContext.cpp b/Engine/Source/Editor/UnrealEd/Private/SActorEditorContext.cpp
new file mode 100644
--- /dev/null
+++ b/Engine/Source/Editor/UnrealEd/Private/SActorEditorContext.cpp
@@ -0,0 +1,123 @@
+// Copyright Epic Games, Inc. All Rights Reserved.
+
+#include "SActorEditorContext.h"
+#include "IActorEditorContextClient.h"
+#include "Subsystems/ActorEditorContextSubsystem.h"
+#include "Widgets/Input/SButton.h"
+#include "Widgets/Images/SImage.h"
+#include "Widgets/Text/STextBlock.h"
+#include "Widgets/SBoxPanel.h"
+#include "Editor.h"
+
+#define LOCTEXT_NAMESPACE "ActorEditorContext"
+
+void SActorEditorContext::Construct(const FArguments& InArgs)
+{
+	World = (InArgs._World);
+	GEditor->GetEditorWorldContext().AddRef(World);
+	UActorEditorContextSubsystem::Get()->OnActorEditorContextSubsystemChanged().AddSP(this, &SActorEditorContext::Rebuild);
+	FEditorDelegates::MapChange.AddSP(this, &SActorEditorContext::OnEditorMapChange);
+	Rebuild();
+}
+
+SActorEditorContext::~SActorEditorContext()
+{
+	GEditor->GetEditorWorldContext().RemoveRef(World);
+	FEditorDelegates::MapChange.RemoveAll(this);
+	UActorEditorContextSubsystem::Get()->OnActorEditorContextSubsystemChanged().RemoveAll(this);
+}
+
+void SActorEditorContext::Rebuild()
+{
+	TArray<IActorEditorContextClient*> Clients = UActorEditorContextSubsystem::Get()->GetDisplayableClients();
+	if (Clients.Num() > 0 && World)
+	{
+		TSharedPtr<SVerticalBox> VBox;
+		ChildSlot.Padding(2, 2, 2, 2)
+		[
+			SAssignNew(VBox, SVerticalBox)
+		];
+
+		for (IActorEditorContextClient* Client : Clients)
+		{
+			FActorEditorContextClientDisplayInfo Info;
+			Client->GetActorEditorContextDisplayInfo(World, Info);
+
+			VBox->AddSlot()
+				.AutoHeight()
+				.Padding(0, 2)
+			[
+				SNew(SBorder)
+				.Visibility_Lambda([Client, this]()
+				{ 
+					FActorEditorContextClientDisplayInfo Info;
+					return Client->GetActorEditorContextDisplayInfo(World, Info) ? EVisibility::Visible : EVisibility::Collapsed; 
+				})
+				.BorderImage(FCoreStyle::Get().GetBrush("Docking.Sidebar.Border"))
+				.Content()
+				[
+					SNew(SVerticalBox)
+					+ SVerticalBox::Slot()
+					.AutoHeight()
+					[
+						SNew(SHorizontalBox)
+						+ SHorizontalBox::Slot()
+						.AutoWidth()
+						.Padding(4.0f, 1.0f, 2.0f, 1.0f)
+						[
+							SNew(SImage)
+							.Image(Info.Brush ? Info.Brush : FStyleDefaults::GetNoBrush())
+						]
+						+ SHorizontalBox::Slot()
+						.AutoWidth()
+						.Padding(4.0f, 1.0f, 2.0f, 1.0f)
+						[
+							SNew(STextBlock)
+							.Text(FText::FromString(Info.Title))
+							.ShadowOffset(FVector2D(1, 1))
+							.Font(FCoreStyle::GetDefaultFontStyle("Bold", 10))
+							.ColorAndOpacity(FLinearColor::White)
+						]
+						+ SHorizontalBox::Slot()
+						.AutoWidth()
+						[
+							SNew(SButton)
+							.HAlign(HAlign_Right)
+							.Cursor(EMouseCursor::Default)
+							.ButtonStyle(FAppStyle::Get(), "SimpleButton")
+							.ContentPadding(0)
+							.Visibility_Lambda([Client, this]() { return (Client && Client->CanResetContext(World)) ? EVisibility::Visible : EVisibility::Collapsed; })
+							.ToolTipText(FText::Format(LOCTEXT("ResetActorEditorContextTooltip", "Reset {0}"), FText::FromString(Info.Title)))
+							.OnClicked_Lambda([Client, this]()
+							{
+								if (Client && Client->CanResetContext(World))
+								{
+									UActorEditorContextSubsystem::Get()->ResetContext(Client);
+								}
+								return FReply::Handled();
+							})
+							.Content()
+							[
+								SNew(SImage)
+								.Image(FEditorStyle::GetBrush("PropertyWindow.DiffersFromDefault"))
+								.ColorAndOpacity(FSlateColor::UseForeground())
+							]
+						]
+					]
+					+ SVerticalBox::Slot()
+					.AutoHeight()
+					.Padding(8, 2, 0, 2)
+					[
+						World ? Client->GetActorEditorContextWidget(World) : SNullWidget::NullWidget
+					]
+				]
+			];
+		}
+	}
+	else
+	{
+		ChildSlot[SNullWidget::NullWidget];
+	}
+}
+
+#undef LOCTEXT_NAMESPACE
\ No newline at end of file
diff --git a/Engine/Source/Editor/UnrealEd/Private/Settings/SettingsClasses.cpp b/Engine/Source/Editor/UnrealEd/Private/Settings/SettingsClasses.cpp
--- a/Engine/Source/Editor/UnrealEd/Private/Settings/SettingsClasses.cpp
+++ b/Engine/Source/Editor/UnrealEd/Private/Settings/SettingsClasses.cpp
@@ -883,6 +883,7 @@ ULevelEditorViewportSettings::ULevelEditorViewportSettings( const FObjectInitial
 	MeasuringToolUnits = MeasureUnits_Centimeters;
 	bAllowArcballRotate = false;
 	bAllowScreenRotate = false;
+	bShowActorEditorContext = true;
 	MouseSensitivty = .2f;
 	// Set a default preview mesh
 	PreviewMeshes.Add(FSoftObjectPath("/Engine/EditorMeshes/ColorCalibrator/SM_ColorCalibrator.SM_ColorCalibrator"));
diff --git a/Engine/Source/Editor/UnrealEd/Private/Subsystems/ActorEditorContextSubsystem.cpp b/Engine/Source/Editor/UnrealEd/Private/Subsystems/ActorEditorContextSubsystem.cpp
new file mode 100644
--- /dev/null
+++ b/Engine/Source/Editor/UnrealEd/Private/Subsystems/ActorEditorContextSubsystem.cpp
@@ -0,0 +1,155 @@
+// Copyright Epic Games, Inc. All Rights Reserved.
+
+#include "Subsystems/ActorEditorContextSubsystem.h"
+#include "GameFramework/Actor.h"
+#include "IActorEditorContextClient.h"
+#include "ScopedTransaction.h"
+#include "Editor.h"
+
+#define LOCTEXT_NAMESPACE "ActorEditorContext"
+
+UActorEditorContextSubsystem* UActorEditorContextSubsystem::Get()
+{ 
+	check(GEditor);
+	return GEditor->GetEditorSubsystem<UActorEditorContextSubsystem>();
+}
+
+void UActorEditorContextSubsystem::Initialize(FSubsystemCollectionBase& Collection)
+{
+	Super::Initialize(Collection);
+
+	GEditor->OnLevelActorAdded().AddUObject(this, &UActorEditorContextSubsystem::ApplyContext);
+}
+
+void UActorEditorContextSubsystem::Deinitialize()
+{
+	GEditor->OnLevelActorAdded().RemoveAll(this);
+
+	Super::Deinitialize();
+}
+
+void UActorEditorContextSubsystem::RegisterClient(IActorEditorContextClient* Client)
+{
+	if (Client && !Clients.Contains(Client))
+	{
+		Client->GetOnActorEditorContextClientChanged().AddUObject(this, &UActorEditorContextSubsystem::OnActorEditorContextClientChanged);
+		Clients.Add(Client);
+	}
+}
+
+void UActorEditorContextSubsystem::UnregisterClient(IActorEditorContextClient* Client)
+{
+	if (Client)
+	{
+		if (Clients.Remove(Client))
+		{
+			Client->GetOnActorEditorContextClientChanged().RemoveAll(this);
+		}
+	}
+}
+
+void UActorEditorContextSubsystem::ApplyContext(AActor* InActor)
+{
+	UWorld* World = GetWorld();
+	if (Clients.IsEmpty() || !World || !InActor || (InActor->GetWorld() != World))
+	{
+		return;
+	}
+
+	for (IActorEditorContextClient* Client : Clients)
+	{
+		Client->OnExecuteActorEditorContextAction(World, EActorEditorContextAction::ApplyContext, InActor);
+	}
+}
+
+void UActorEditorContextSubsystem::ResetContext()
+{
+	UWorld* World = GetWorld();
+	if (Clients.IsEmpty() || !World)
+	{
+		return;
+	}
+		
+	const FScopedTransaction Transaction(LOCTEXT("Reset Actor Editor Context", "Reset Actor Editor Context"));
+	for (IActorEditorContextClient* Client : Clients)
+	{
+		Client->OnExecuteActorEditorContextAction(World, EActorEditorContextAction::ResetContext);
+	}
+	ActorEditorContextSubsystemChanged.Broadcast();
+}
+
+void UActorEditorContextSubsystem::ResetContext(IActorEditorContextClient* Client)
+{
+	UWorld* World = GetWorld();
+	if (Clients.IsEmpty() || !World)
+	{
+		return;
+	}
+
+	if (Clients.Contains(Client))
+	{
+		const FScopedTransaction Transaction(LOCTEXT("Reset Actor Editor Context", "Reset Actor Editor Context"));
+		Client->OnExecuteActorEditorContextAction(World, EActorEditorContextAction::ResetContext);
+		ActorEditorContextSubsystemChanged.Broadcast();
+	}
+}
+
+void UActorEditorContextSubsystem::PushContext()
+{
+	UWorld* World = GetWorld();
+	if (Clients.IsEmpty() || !World)
+	{
+		return;
+	}
+
+	for (IActorEditorContextClient* Client : Clients)
+	{
+		Client->OnExecuteActorEditorContextAction(World, EActorEditorContextAction::PushContext);
+	}
+	ActorEditorContextSubsystemChanged.Broadcast();
+}
+
+void UActorEditorContextSubsystem::PopContext()
+{
+	UWorld* World = GetWorld();
+	if (Clients.IsEmpty() || !World)
+	{
+		return;
+	}
+
+	for (IActorEditorContextClient* Client : Clients)
+	{
+		Client->OnExecuteActorEditorContextAction(World, EActorEditorContextAction::PopContext);
+	}
+	ActorEditorContextSubsystemChanged.Broadcast();
+}
+
+TArray<IActorEditorContextClient*> UActorEditorContextSubsystem::GetDisplayableClients() const
+{
+	TArray<IActorEditorContextClient*> DisplayableClients;
+	UWorld* World = GetWorld();
+	if (!Clients.IsEmpty() && World)
+	{
+		for (IActorEditorContextClient* Client : Clients)
+		{
+			FActorEditorContextClientDisplayInfo Info;
+			if (Client->GetActorEditorContextDisplayInfo(World, Info))
+			{
+				DisplayableClients.Add(Client);
+			}
+		}
+	}
+	return DisplayableClients;
+}
+
+UWorld* UActorEditorContextSubsystem::GetWorld() const
+{
+	return GEditor->GetEditorWorldContext().World();
+}
+
+void UActorEditorContextSubsystem::OnActorEditorContextClientChanged(IActorEditorContextClient* Client)
+{
+	ActorEditorContextSubsystemChanged.Broadcast();
+}
+
+#undef LOCTEXT_NAMESPACE
\ No newline at end of file
diff --git a/Engine/Source/Editor/UnrealEd/Private/WorldFolders.cpp b/Engine/Source/Editor/UnrealEd/Private/WorldFolders.cpp
--- a/Engine/Source/Editor/UnrealEd/Private/WorldFolders.cpp
+++ b/Engine/Source/Editor/UnrealEd/Private/WorldFolders.cpp
@@ -9,6 +9,7 @@
 #include "Misc/Paths.h"
 #include "EngineUtils.h"
 #include "EditorActorFolders.h"
+#include "LevelInstance/LevelInstanceSubsystem.h"
 #include "LevelInstance/LevelInstanceActor.h"
 
 #define LOCTEXT_NAMESPACE "UnrealEd.WorldFolders"
@@ -119,7 +120,30 @@ bool UWorldFolders::RenameFolder(const FFolder& InOldFolder, const FFolder& InNe
 	check(IsValid(World.Get()));
 	check(InOldFolder.GetRootObject() == InNewFolder.GetRootObject());
 
-	return GetImpl(InOldFolder).RenameFolder(InOldFolder, InNewFolder);
+	bool bSuccess = GetImpl(InOldFolder).RenameFolder(InOldFolder, InNewFolder);
+	if (bSuccess)
+	{
+		bool bChanged = false;
+		for (int StackIndex=0; StackIndex < CurrentFolderStack.Num(); ++StackIndex)
+		{
+			FActorPlacementFolder& StackElement = CurrentFolderStack[StackIndex];
+			if (StackElement.GetFolder() == InOldFolder)
+			{
+				StackElement.Path = InNewFolder.GetPath();
+				bChanged = true;
+			}
+		}
+		if (CurrentFolder.GetFolder() == InOldFolder)
+		{
+			CurrentFolder.Path = InNewFolder.GetPath();
+			bChanged = true;
+		}
+		if (bChanged)
+		{
+			FActorFolders::Get().BroadcastOnActorEditorContextClientChanged();
+		}
+	}
+	return bSuccess;
 }
 
 void UWorldFolders::BroadcastOnActorFolderCreated(const FFolder& InFolder)
@@ -156,6 +180,49 @@ bool UWorldFolders::SetIsFolderExpanded(const FFolder& InFolder, bool bIsExpande
 	return false;
 }
 
+FFolder UWorldFolders::GetActorEditorContextFolder() const
+{
+	FFolder Folder = CurrentFolder.GetFolder();
+	if (ContainsFolder(Folder))
+	{
+		ULevelInstanceSubsystem* LevelInstanceSubsystem = GetWorld()->GetSubsystem<ULevelInstanceSubsystem>();
+		UObject* EditingLevelInstance = LevelInstanceSubsystem ? LevelInstanceSubsystem->GetEditingLevelInstance() : nullptr;
+		FFolder::FRootObject RootObject = FFolder::FRootObject(EditingLevelInstance ? EditingLevelInstance : (GetWorld()->GetCurrentLevel() == GetWorld()->PersistentLevel) ? nullptr : GetWorld()->GetCurrentLevel());
+		if (Folder.GetRootObject() == RootObject)
+		{
+			return Folder;
+		}
+	}
+	return FFolder();
+}
+
+bool UWorldFolders::SetActorEditorContextFolder(const FFolder& InFolder)
+{
+	if (InFolder != CurrentFolder.GetFolder())
+	{
+		Modify();
+		CurrentFolder.Path = InFolder.GetPath();
+		CurrentFolder.RootObjectPtr = InFolder.GetRootObjectPtr();
+		return true;
+	}
+	return false;
+}
+
+void UWorldFolders::PushActorEditorContext()
+{
+	Modify();
+	CurrentFolderStack.Push(CurrentFolder);
+	CurrentFolder.Reset();
+}
+
+void UWorldFolders::PopActorEditorContext()
+{
+	check(!CurrentFolderStack.IsEmpty());
+
+	Modify();
+	CurrentFolder = CurrentFolderStack.Pop();
+}
+
 bool UWorldFolders::ContainsFolder(const FFolder& InFolder) const
 {
 	return GetImpl(InFolder).ContainsFolder(InFolder);
@@ -196,6 +263,8 @@ void UWorldFolders::Serialize(FArchive& Ar)
 
 	check(PersistentFolders.IsValid());
 	Ar << FoldersProperties;
+
+	Super::Serialize(Ar);
 }
 
 void UWorldFolders::OnWorldSaved()
diff --git a/Engine/Source/Editor/UnrealEd/Public/EditorActorFolders.h b/Engine/Source/Editor/UnrealEd/Public/EditorActorFolders.h
--- a/Engine/Source/Editor/UnrealEd/Public/EditorActorFolders.h
+++ b/Engine/Source/Editor/UnrealEd/Public/EditorActorFolders.h
@@ -8,6 +8,7 @@
 #include "UObject/GCObject.h"
 #include "Folder.h"
 #include "WorldFolders.h"
+#include "IActorEditorContextClient.h"
 
 class FObjectPostSaveContext;
 class AActor;
@@ -27,7 +28,7 @@ DECLARE_MULTICAST_DELEGATE_ThreeParams(FOnActorFolderMoved, UWorld&, const FFold
 
 
 /** Class responsible for managing an in-memory representation of actor folders in the editor */
-struct UNREALED_API FActorFolders : public FGCObject
+struct UNREALED_API FActorFolders : public FGCObject, public IActorEditorContextClient
 {
 	FActorFolders();
 	~FActorFolders();
@@ -151,6 +152,16 @@ struct UNREALED_API FActorFolders : public FGCObject
 	/** Get the folder properties for the specified path. Returns nullptr if no properties exist */
 	FActorFolderProps* GetFolderProperties(UWorld& InWorld, const FFolder& InFolder);
 
+	//~ Begin IActorEditorContextClient interface
+	virtual void OnExecuteActorEditorContextAction(UWorld* InWorld, const EActorEditorContextAction& InType, class AActor* InActor = nullptr) override;
+	virtual bool GetActorEditorContextDisplayInfo(UWorld* InWorld, FActorEditorContextClientDisplayInfo& OutDiplayInfo) const override;
+	virtual bool CanResetContext(UWorld* InWorld) const override { return true; };
+	virtual TSharedRef<SWidget> GetActorEditorContextWidget(UWorld* InWorld) const override;
+	virtual FOnActorEditorContextClientChanged& GetOnActorEditorContextClientChanged() override { return ActorEditorContextClientChanged; }
+	//~ End IActorEditorContextClient interface
+	FFolder GetActorEditorContextFolder(UWorld& InWorld) const;
+	void SetActorEditorContextFolder(UWorld& InWorld, const FFolder& InFolder);
+
 private:
 
 	/** Broadcast when actor folder is created. */
@@ -199,6 +210,12 @@ private:
 	/** Transient map of folders, keyed on world pointer */
 	TMap<TWeakObjectPtr<UWorld>, UWorldFolders*> WorldFolders;
 
+	/** Called when ActorEditorContextClient changed. */
+	void BroadcastOnActorEditorContextClientChanged();
+
+	/** Delegate used to notify changes to ActorEditorContextSubsystem */
+	FOnActorEditorContextClientChanged ActorEditorContextClientChanged;
+
 	/** Singleton instance maintained by the editor */
 	static FActorFolders* Singleton;
 
diff --git a/Engine/Source/Editor/UnrealEd/Public/IActorEditorContextClient.h b/Engine/Source/Editor/UnrealEd/Public/IActorEditorContextClient.h
new file mode 100644
--- /dev/null
+++ b/Engine/Source/Editor/UnrealEd/Public/IActorEditorContextClient.h
@@ -0,0 +1,39 @@
+// Copyright Epic Games, Inc. All Rights Reserved.
+
+#pragma once
+
+#include "CoreMinimal.h"
+#include "Widgets/SWidget.h"
+
+/**
+ * Public interface for Actor Editor Context Clients
+ */
+
+DECLARE_MULTICAST_DELEGATE_OneParam(FOnActorEditorContextClientChanged, struct IActorEditorContextClient*);
+
+class UWorld;
+class AActor;
+
+struct FActorEditorContextClientDisplayInfo
+{
+	FActorEditorContextClientDisplayInfo() : Brush(nullptr) {}
+	FString Title;
+	const FSlateBrush* Brush;
+};
+
+enum class EActorEditorContextAction
+{
+	ApplyContext,
+	ResetContext,
+	PushContext,
+	PopContext
+};
+
+struct IActorEditorContextClient
+{
+	virtual void OnExecuteActorEditorContextAction(UWorld* InWorld, const EActorEditorContextAction& InType, AActor* InActor = nullptr) = 0;
+	virtual bool GetActorEditorContextDisplayInfo(UWorld* InWorld, FActorEditorContextClientDisplayInfo& OutDiplayInfo) const = 0;
+	virtual bool CanResetContext(UWorld* InWorld) const = 0;
+	virtual TSharedRef<SWidget> GetActorEditorContextWidget(UWorld* InWorld) const = 0;
+	virtual FOnActorEditorContextClientChanged& GetOnActorEditorContextClientChanged() = 0;
+};
\ No newline at end of file
diff --git a/Engine/Source/Editor/UnrealEd/Public/SActorEditorContext.h b/Engine/Source/Editor/UnrealEd/Public/SActorEditorContext.h
new file mode 100644
--- /dev/null
+++ b/Engine/Source/Editor/UnrealEd/Public/SActorEditorContext.h
@@ -0,0 +1,23 @@
+// Copyright Epic Games, Inc. All Rights Reserved.
+
+#pragma once
+
+#include "Widgets/SCompoundWidget.h"
+
+class UNREALED_API SActorEditorContext : public SCompoundWidget
+{
+public:
+	SLATE_BEGIN_ARGS(SActorEditorContext) {}
+	SLATE_ARGUMENT(UWorld*, World)
+	SLATE_END_ARGS()
+
+	void Construct(const FArguments& InArgs);
+	~SActorEditorContext();
+
+private:
+
+	void OnEditorMapChange(uint32 MapChangeFlags = 0) { Rebuild(); }
+	void Rebuild();
+	
+	UWorld* World;
+};
\ No newline at end of file
diff --git a/Engine/Source/Editor/UnrealEd/Public/Subsystems/ActorEditorContextSubsystem.h b/Engine/Source/Editor/UnrealEd/Public/Subsystems/ActorEditorContextSubsystem.h
new file mode 100644
--- /dev/null
+++ b/Engine/Source/Editor/UnrealEd/Public/Subsystems/ActorEditorContextSubsystem.h
@@ -0,0 +1,47 @@
+// Copyright Epic Games, Inc. All Rights Reserved.
+
+#pragma once
+
+#include "CoreMinimal.h"
+#include "EditorSubsystem.h"
+
+#include "ActorEditorContextSubsystem.generated.h"
+
+struct IActorEditorContextClient;
+class AActor;
+
+DECLARE_MULTICAST_DELEGATE(FOnActorEditorContextSubsystemChanged);
+
+/**
+* UActorEditorContextSubsystem 
+*/
+UCLASS()
+class UNREALED_API UActorEditorContextSubsystem : public UEditorSubsystem
+{
+public:
+
+	GENERATED_BODY()
+
+	static UActorEditorContextSubsystem* Get();
+
+	virtual void Initialize(FSubsystemCollectionBase& Collection) override;
+	virtual void Deinitialize() override;
+
+	void RegisterClient(IActorEditorContextClient* Client);
+	void UnregisterClient(IActorEditorContextClient* Client);
+	void ResetContext();
+	void ResetContext(IActorEditorContextClient* Client);
+	void PushContext();
+	void PopContext();
+	TArray<IActorEditorContextClient*> GetDisplayableClients() const;
+	FOnActorEditorContextSubsystemChanged& OnActorEditorContextSubsystemChanged() { return ActorEditorContextSubsystemChanged; }
+
+private:
+
+	UWorld* GetWorld() const;
+	void OnActorEditorContextClientChanged(IActorEditorContextClient* Client);
+	void ApplyContext(AActor* InActor);
+
+	FOnActorEditorContextSubsystemChanged ActorEditorContextSubsystemChanged;
+	TArray<IActorEditorContextClient*> Clients;
+};
\ No newline at end of file
diff --git a/Engine/Source/Editor/UnrealEd/Public/WorldFolders.h b/Engine/Source/Editor/UnrealEd/Public/WorldFolders.h
--- a/Engine/Source/Editor/UnrealEd/Public/WorldFolders.h
+++ b/Engine/Source/Editor/UnrealEd/Public/WorldFolders.h
@@ -28,6 +28,29 @@ struct FActorFolderProps
 	bool bIsExpanded;
 };
 
+USTRUCT()
+struct FActorPlacementFolder
+{
+	GENERATED_BODY()
+
+	UPROPERTY()
+	FName Path;
+
+	UPROPERTY()
+	TWeakObjectPtr<UObject> RootObjectPtr;
+
+	void Reset()
+	{
+		Path = NAME_None;
+		RootObjectPtr.Reset();
+	}
+
+	FFolder GetFolder() const
+	{
+		return FFolder(Path, FFolder::FRootObject(RootObjectPtr.Get()));
+	}
+};
+
 /** Per-World Actor Folders UObject (used to support undo/redo reliably) */
 UCLASS()
 class UNREALED_API UWorldFolders : public UObject
@@ -43,6 +66,10 @@ public:
 	bool RenameFolder(const FFolder& InOldFolder, const FFolder& InNewFolder);
 	bool IsFolderExpanded(const FFolder& InFolder) const;
 	bool SetIsFolderExpanded(const FFolder& InFolder, bool bIsExpanded);
+	FFolder GetActorEditorContextFolder() const;
+	bool SetActorEditorContextFolder(const FFolder& InFolder);
+	void PushActorEditorContext();
+	void PopActorEditorContext();
 	bool ContainsFolder(const FFolder& InFolder) const;
 	void ForEachFolder(TFunctionRef<bool(const FFolder&)> Operation);
 	void ForEachFolderWithRootObject(const FFolder::FRootObject& InFolderRootObject, TFunctionRef<bool(const FFolder&)> Operation);
@@ -77,6 +104,11 @@ private:
 	TMap<FFolder, FActorFolderProps> FoldersProperties;
 	TMap<FFolder, FActorFolderProps> LoadedStateFoldersProperties;
 
+	UPROPERTY()
+	FActorPlacementFolder CurrentFolder;
+
+	TArray<FActorPlacementFolder> CurrentFolderStack;
+
 	friend class FWorldFoldersImplementation;
 	friend class FWorldPersistentFolders;
 	friend class FWorldTransientFolders;
diff --git a/Engine/Source/Runtime/Engine/Private/ActorPartition/ActorPartitionSubsystem.cpp b/Engine/Source/Runtime/Engine/Private/ActorPartition/ActorPartitionSubsystem.cpp
--- a/Engine/Source/Runtime/Engine/Private/ActorPartition/ActorPartitionSubsystem.cpp
+++ b/Engine/Source/Runtime/Engine/Private/ActorPartition/ActorPartitionSubsystem.cpp
@@ -393,22 +393,18 @@ APartitionActor* UActorPartitionSubsystem::GetActor(const TSubclassOf<APartition
 	
 	auto WrappedInActorCreated = [World, InActorCreated](APartitionActor* PartitionActor)
 	{
-		if (UDataLayerSubsystem* DataLayerSubsystem = World->GetSubsystem<UDataLayerSubsystem>())
+		if (const AWorldDataLayers* WorldDataLayers = World->GetWorldDataLayers())
 		{
-			if (const AWorldDataLayers* WorldDataLayers = World->GetWorldDataLayers())
+			for (UDataLayer* DataLayer : WorldDataLayers->GetActorEditorContextDataLayers())
 			{
-				for (const FName& DataLayer : DataLayerSubsystem->GetDataLayerEditorContext().GetDataLayers())
-				{
-					PartitionActor->AddDataLayer(WorldDataLayers->GetDataLayerFromName(DataLayer));
-				}
+				PartitionActor->AddDataLayer(DataLayer);
 			}
 		}
-
 		InActorCreated(PartitionActor);
 	};
 
 	UDataLayerSubsystem* DataLayerSubsystem = World->GetSubsystem<UDataLayerSubsystem>();
-	FActorPartitionIdentifier ActorPartitionId(InActorClass, InGuid, DataLayerSubsystem ? DataLayerSubsystem->GetDataLayerEditorContext().GetHash() : FDataLayerEditorContext::EmptyHash);
+	FActorPartitionIdentifier ActorPartitionId(InActorClass, InGuid, DataLayerSubsystem ? DataLayerSubsystem->GetDataLayerEditorContextHash() : FDataLayerEditorContext::EmptyHash);
 	TMap<FActorPartitionIdentifier, TWeakObjectPtr<APartitionActor>>* ActorsPerId = PartitionedActors.Find(InCellCoords);
 	APartitionActor* FoundActor = nullptr;
 	if (!ActorsPerId)
diff --git a/Engine/Source/Runtime/Engine/Private/LevelInstance/LevelInstanceSubsystem.cpp b/Engine/Source/Runtime/Engine/Private/LevelInstance/LevelInstanceSubsystem.cpp
--- a/Engine/Source/Runtime/Engine/Private/LevelInstance/LevelInstanceSubsystem.cpp
+++ b/Engine/Source/Runtime/Engine/Private/LevelInstance/LevelInstanceSubsystem.cpp
@@ -43,6 +43,7 @@
 #include "UObject/UObjectGlobals.h"
 #include "EditorActorFolders.h"
 #include "Misc/MessageDialog.h"
+#include "Subsystems/ActorEditorContextSubsystem.h"
 #endif
 
 #include "HAL/IConsoleManager.h"
@@ -1700,6 +1701,9 @@ bool ULevelInstanceSubsystem::EditLevelInstanceInternal(ALevelInstance* LevelIns
 	// Edit can't be undone
 	GEditor->ResetTransaction(LOCTEXT("LevelInstanceEditResetTrans", "Edit Level Instance"));
 
+	// When editing a Level Instance, push a new empty actor editor context
+	UActorEditorContextSubsystem::Get()->PushContext();
+
 	return true;
 }
 
@@ -1857,6 +1861,9 @@ bool ULevelInstanceSubsystem::CommitLevelInstanceInternal(TUniquePtr<FLevelInsta
 
 	GEngine->BroadcastLevelActorListChanged();
 
+	// Restore actor editor context
+	UActorEditorContextSubsystem::Get()->PopContext();
+
 	return true;
 }
 
diff --git a/Engine/Source/Runtime/Engine/Private/WorldPartition/DataLayer/DataLayer.cpp b/Engine/Source/Runtime/Engine/Private/WorldPartition/DataLayer/DataLayer.cpp
--- a/Engine/Source/Runtime/Engine/Private/WorldPartition/DataLayer/DataLayer.cpp
+++ b/Engine/Source/Runtime/Engine/Private/WorldPartition/DataLayer/DataLayer.cpp
@@ -259,6 +259,23 @@ void UDataLayer::SetIsRuntime(bool bInIsRuntime)
 	}
 }
 
+#if WITH_EDITOR
+bool UDataLayer::IsInActorEditorContext() const
+{
+	return GetOuterAWorldDataLayers()->IsInActorEditorContext(this);
+}
+
+bool UDataLayer::AddToActorEditorContext()
+{
+	return GetOuterAWorldDataLayers()->AddToActorEditorContext(this);
+}
+
+bool UDataLayer::RemoveFromActorEditorContext()
+{
+	return GetOuterAWorldDataLayers()->RemoveFromActorEditorContext(this);
+}
+#endif
+
 void UDataLayer::PropagateIsRuntime()
 {
 	if (IsRuntime())
diff --git a/Engine/Source/Runtime/Engine/Private/WorldPartition/DataLayer/DataLayerEditorContext.cpp b/Engine/Source/Runtime/Engine/Private/WorldPartition/DataLayer/DataLayerEditorContext.cpp
--- a/Engine/Source/Runtime/Engine/Private/WorldPartition/DataLayer/DataLayerEditorContext.cpp
+++ b/Engine/Source/Runtime/Engine/Private/WorldPartition/DataLayer/DataLayerEditorContext.cpp
@@ -40,38 +40,4 @@ FDataLayerEditorContext::FDataLayerEditorContext(UWorld* InWorld, const TArray<F
 	}
 }
 
-/*
- * FScopeChangeDataLayerEditorContext
- */
-
-FScopeChangeDataLayerEditorContext::FScopeChangeDataLayerEditorContext(UWorld* InWorld, const FDataLayerEditorContext& InContext)
-	: World(InWorld)
-{
-	Initialize(InContext);
-}
-
-FScopeChangeDataLayerEditorContext::FScopeChangeDataLayerEditorContext(UWorld* InWorld, const FActorDataLayer& InContextDataLayer)
-	: World(InWorld)
-{
-	const AWorldDataLayers* WorldDataLayers = InWorld->GetWorldDataLayers();
-	Initialize(FDataLayerEditorContext(InWorld, WorldDataLayers ? WorldDataLayers->GetDataLayerNames({ InContextDataLayer }) : TArray<FName>()));
-}
-
-void FScopeChangeDataLayerEditorContext::Initialize(const FDataLayerEditorContext& InContext)
-{
-	if (const UDataLayerSubsystem* DataLayerSubsystem = World.Get() ? World->GetSubsystem<UDataLayerSubsystem>() : nullptr)
-	{
-		OldContext = DataLayerSubsystem->GetDataLayerEditorContext();
-		DataLayerSubsystem->DataLayerEditorContext = InContext;
-	}
-}
-
-FScopeChangeDataLayerEditorContext::~FScopeChangeDataLayerEditorContext()
-{
-	if (const UDataLayerSubsystem* DataLayerSubsystem = World.Get() ? World->GetSubsystem<UDataLayerSubsystem>() : nullptr)
-	{
-		DataLayerSubsystem->DataLayerEditorContext = OldContext;
-	}
-}
-
 #endif
\ No newline at end of file
diff --git a/Engine/Source/Runtime/Engine/Private/WorldPartition/DataLayer/DataLayerSubsystem.cpp b/Engine/Source/Runtime/Engine/Private/WorldPartition/DataLayer/DataLayerSubsystem.cpp
--- a/Engine/Source/Runtime/Engine/Private/WorldPartition/DataLayer/DataLayerSubsystem.cpp
+++ b/Engine/Source/Runtime/Engine/Private/WorldPartition/DataLayer/DataLayerSubsystem.cpp
@@ -521,4 +521,20 @@ void UDataLayerSubsystem::GetDataLayerDebugColors(TMap<FName, FColor>& OutMappin
 		OutMapping.Add(DataLayer->GetFName(), DataLayer->GetDebugColor());
 		return true;
 	});
-}
\ No newline at end of file
+}
+
+#if WITH_EDITOR
+uint32 UDataLayerSubsystem::GetDataLayerEditorContextHash() const
+{
+	if (const AWorldDataLayers* WorldDataLayers = GetWorld()->GetWorldDataLayers())
+	{
+		TArray<FName> DataLayers;
+		for (UDataLayer* DataLayer : WorldDataLayers->GetActorEditorContextDataLayers())
+		{
+			DataLayers.Add(DataLayer->GetFName());
+		}
+		return FDataLayerEditorContext(GetWorld(), DataLayers).GetHash();
+	}
+	return FDataLayerEditorContext::EmptyHash;
+}
+#endif
\ No newline at end of file
diff --git a/Engine/Source/Runtime/Engine/Private/WorldPartition/DataLayer/WorldDataLayers.cpp b/Engine/Source/Runtime/Engine/Private/WorldPartition/DataLayer/WorldDataLayers.cpp
--- a/Engine/Source/Runtime/Engine/Private/WorldPartition/DataLayer/WorldDataLayers.cpp
+++ b/Engine/Source/Runtime/Engine/Private/WorldPartition/DataLayer/WorldDataLayers.cpp
@@ -640,6 +640,72 @@ void AWorldDataLayers::SetAllowRuntimeDataLayerEditing(bool bInAllowRuntimeDataL
 		bAllowRuntimeDataLayerEditing = bInAllowRuntimeDataLayerEditing;
 	}
 }
+
+bool AWorldDataLayers::IsInActorEditorContext(const UDataLayer* InDataLayer) const
+{
+	for (const FName& DataLayerName : CurrentDataLayers.DataLayers)
+	{
+		const UDataLayer* DataLayer = GetDataLayerFromName(DataLayerName);
+		if (DataLayer && (DataLayer == InDataLayer) && !DataLayer->IsLocked())
+		{
+			return true;
+		}
+	}
+	return false;
+}
+
+bool AWorldDataLayers::AddToActorEditorContext(UDataLayer* InDataLayer)
+{
+	check(WorldDataLayers.Contains(InDataLayer));
+	if (!CurrentDataLayers.DataLayers.Contains(InDataLayer->GetFName()))
+	{
+		Modify(/*bDirty*/false);
+		CurrentDataLayers.DataLayers.Add(InDataLayer->GetFName());
+		return true;
+	}
+	return false;
+}
+
+bool AWorldDataLayers::RemoveFromActorEditorContext(UDataLayer* InDataLayer)
+{
+	check(WorldDataLayers.Contains(InDataLayer));
+	if (CurrentDataLayers.DataLayers.Contains(InDataLayer->GetFName()))
+	{
+		Modify(/*bDirty*/false);
+		CurrentDataLayers.DataLayers.Remove(InDataLayer->GetFName());
+		return true;
+	}
+	return false;
+}
+
+void AWorldDataLayers::PushActorEditorContext()
+{
+	Modify(/*bDirty*/false);
+	CurrentDataLayersStack.Push(CurrentDataLayers);
+	CurrentDataLayers.Reset();
+}
+
+void AWorldDataLayers::PopActorEditorContext()
+{
+	check(!CurrentDataLayersStack.IsEmpty());
+	Modify(/*bDirty*/false);
+	CurrentDataLayers = CurrentDataLayersStack.Pop();
+}
+
+TArray<UDataLayer*> AWorldDataLayers::GetActorEditorContextDataLayers() const
+{
+	TArray<UDataLayer*> Result;
+	for (const FName& DataLayerName : CurrentDataLayers.DataLayers)
+	{
+		const UDataLayer* DataLayer = GetDataLayerFromName(DataLayerName);
+		if (DataLayer && !DataLayer->IsLocked())
+		{
+			Result.Add(const_cast<UDataLayer*>(DataLayer));
+		}
+	};
+	return Result;
+}
+
 #endif
 
 bool AWorldDataLayers::ContainsDataLayer(const UDataLayer* InDataLayer) const
diff --git a/Engine/Source/Runtime/Engine/Public/WorldPartition/DataLayer/DataLayer.h b/Engine/Source/Runtime/Engine/Public/WorldPartition/DataLayer/DataLayer.h
--- a/Engine/Source/Runtime/Engine/Public/WorldPartition/DataLayer/DataLayer.h
+++ b/Engine/Source/Runtime/Engine/Public/WorldPartition/DataLayer/DataLayer.h
@@ -90,6 +90,10 @@ public:
 	void SetParent(UDataLayer* InParent);
 	void SetChildParent(UDataLayer* InParent);
 
+	bool IsInActorEditorContext() const;
+	bool AddToActorEditorContext();
+	bool RemoveFromActorEditorContext();
+
 	static FText GetDataLayerText(const UDataLayer* InDataLayer);
 	const TCHAR* GetDataLayerIconName() const;
 #endif
@@ -116,15 +120,15 @@ public:
 	UFUNCTION(Category = "Data Layer|Editor", BlueprintCallable)
 	bool IsEffectiveVisible() const;
 
+	UFUNCTION(Category = "Data Layer|Editor", BlueprintCallable)
+	FColor GetDebugColor() const { return DebugColor; }
+
 	UFUNCTION(Category = "Data Layer|Runtime", BlueprintCallable)
 	bool IsRuntime() const { return bIsRuntime; }
 	
 	UFUNCTION(Category = "Data Layer|Runtime", BlueprintCallable)
 	EDataLayerRuntimeState GetInitialRuntimeState() const { return IsRuntime() ? InitialRuntimeState : EDataLayerRuntimeState::Unloaded; }
 
-	UFUNCTION(Category = "Data Layer|Runtime", BlueprintCallable)
-	FColor GetDebugColor() const { return DebugColor; }
-
 	/** Returns a sanitized version of the provided Data Layer Label */
 	static FName GetSanitizedDataLayerLabel(FName DataLayerLabel);
 
@@ -193,7 +197,7 @@ private:
 	UPROPERTY(Category = "Data Layer|Advanced|Runtime", EditAnywhere, meta = (EditConditionHides, EditCondition = "bIsRuntime"))
 	EDataLayerRuntimeState InitialRuntimeState;
 
-	UPROPERTY(Category = "Data Layer|Advanced|Runtime", EditAnywhere, meta = (EditConditionHides, EditCondition = "bIsRuntime"))
+	UPROPERTY(Category = "Data Layer|Editor", EditAnywhere)
 	FColor DebugColor;
 
 	UPROPERTY()
diff --git a/Engine/Source/Runtime/Engine/Public/WorldPartition/DataLayer/DataLayerEditorContext.h b/Engine/Source/Runtime/Engine/Public/WorldPartition/DataLayer/DataLayerEditorContext.h
--- a/Engine/Source/Runtime/Engine/Public/WorldPartition/DataLayer/DataLayerEditorContext.h
+++ b/Engine/Source/Runtime/Engine/Public/WorldPartition/DataLayer/DataLayerEditorContext.h
@@ -26,20 +26,4 @@ private:
 	uint32 Hash;
 	TArray<FName> DataLayers;
 };
-
-/**
- * FScopeChangeDataLayerEditorContext
- */
-
-struct ENGINE_API FScopeChangeDataLayerEditorContext
-{
-	FScopeChangeDataLayerEditorContext(UWorld* InWorld, const FDataLayerEditorContext& InContext);
-	FScopeChangeDataLayerEditorContext(UWorld* InWorld, const FActorDataLayer& InContextDataLayer);
-	~FScopeChangeDataLayerEditorContext();
-private:
-	void Initialize(const FDataLayerEditorContext& InContext);
-
-	TWeakObjectPtr<UWorld> World;
-	FDataLayerEditorContext OldContext;
-};
 #endif
\ No newline at end of file
diff --git a/Engine/Source/Runtime/Engine/Public/WorldPartition/DataLayer/DataLayerSubsystem.h b/Engine/Source/Runtime/Engine/Public/WorldPartition/DataLayer/DataLayerSubsystem.h
--- a/Engine/Source/Runtime/Engine/Public/WorldPartition/DataLayer/DataLayerSubsystem.h
+++ b/Engine/Source/Runtime/Engine/Public/WorldPartition/DataLayer/DataLayerSubsystem.h
@@ -137,12 +137,9 @@ public:
 
 #if WITH_EDITOR
 public:
-	FORCEINLINE const FDataLayerEditorContext& GetDataLayerEditorContext() const { return DataLayerEditorContext; }
-	FORCEINLINE bool HasDataLayerEditorContext() const { return !DataLayerEditorContext.IsEmpty(); }
+	uint32 GetDataLayerEditorContextHash() const;
 
 private:
-	mutable FDataLayerEditorContext DataLayerEditorContext;
-	friend struct FScopeChangeDataLayerEditorContext;
 #endif
 
 private:
diff --git a/Engine/Source/Runtime/Engine/Public/WorldPartition/DataLayer/WorldDataLayers.h b/Engine/Source/Runtime/Engine/Public/WorldPartition/DataLayer/WorldDataLayers.h
--- a/Engine/Source/Runtime/Engine/Public/WorldPartition/DataLayer/WorldDataLayers.h
+++ b/Engine/Source/Runtime/Engine/Public/WorldPartition/DataLayer/WorldDataLayers.h
@@ -11,6 +11,20 @@
 
 class UDataLayer;
 
+USTRUCT()
+struct FActorPlacementDataLayers
+{
+	GENERATED_BODY()
+
+	UPROPERTY()
+	TArray<FName> DataLayers;
+
+	void Reset()
+	{
+		DataLayers.Reset();
+	}
+};
+
 /**
  * Actor containing all data layers for a world
  */
@@ -38,6 +52,13 @@ public:
 	void SetAllowRuntimeDataLayerEditing(bool bInAllowRuntimeDataLayerEditing);
 	bool GetAllowRuntimeDataLayerEditing() const { return bAllowRuntimeDataLayerEditing; }
 
+	bool IsInActorEditorContext(const UDataLayer* InDataLayer) const;
+	bool AddToActorEditorContext(UDataLayer* InDataLayer);
+	bool RemoveFromActorEditorContext(UDataLayer* InDataLayer);
+	void PushActorEditorContext();
+	void PopActorEditorContext();
+	TArray<UDataLayer*> GetActorEditorContextDataLayers() const;
+
 	//~ Begin Helper Functions
 	TArray<const UDataLayer*> GetDataLayerObjects(const TArray<FActorDataLayer>& DataLayers) const;
 	TArray<const UDataLayer*> GetDataLayerObjects(const TArray<FName>& InDataLayerNames) const;
@@ -116,6 +137,11 @@ private:
 	// True when Runtime Data Layer editing is allowed.
 	UPROPERTY()
 	bool bAllowRuntimeDataLayerEditing;
+
+	UPROPERTY(Transient)
+	FActorPlacementDataLayers CurrentDataLayers;
+
+	TArray<FActorPlacementDataLayers> CurrentDataLayersStack;
 #endif
 
 	UPROPERTY()
diff --git a/Engine/Source/Runtime/Foliage/Private/ProceduralFoliageComponent.cpp b/Engine/Source/Runtime/Foliage/Private/ProceduralFoliageComponent.cpp
--- a/Engine/Source/Runtime/Foliage/Private/ProceduralFoliageComponent.cpp
+++ b/Engine/Source/Runtime/Foliage/Private/ProceduralFoliageComponent.cpp
@@ -15,9 +15,6 @@
 #if WITH_EDITOR
 #include "WorldPartition/WorldPartition.h"
 #include "WorldPartition/WorldPartitionSubsystem.h"
-#include "WorldPartition/DataLayer/WorldDataLayers.h"
-#include "WorldPartition/DataLayer/DataLayerSubsystem.h"
-#include "WorldPartition/DataLayer/DataLayerEditorContext.h"
 #endif
 
 
@@ -300,15 +297,12 @@ bool UProceduralFoliageComponent::ResimulateProceduralFoliage(TFunctionRef<void(
 			{
 				{
 					// Remove old foliage instances
-					FScopeChangeDataLayerEditorContext ScopeContext(GetWorld(), LastSimulationDataLayer);
 					RemoveProceduralContent(false);
 				}
 
 				{
 					// Add new foliage instances
-					FScopeChangeDataLayerEditorContext ScopeContext(GetWorld(), DataLayer);
 					AddInstancesFunc(DesiredFoliageInstances);
-					LastSimulationDataLayer = DataLayer;
 				}
 			}
 
@@ -353,22 +347,4 @@ bool UProceduralFoliageComponent::HasSpawnedAnyInstances()
 	return false;
 }
 
-#if WITH_EDITOR
-bool UProceduralFoliageComponent::CanEditChange(const FProperty* InProperty) const
-{
-	if (InProperty->GetFName() == GET_MEMBER_NAME_CHECKED(UProceduralFoliageComponent, DataLayer))
-	{
-		if (!IsTemplate())
-		{
-			if (!UWorld::IsPartitionedWorld(GetWorld()))
-			{
-				return false;
-			}
-		}
-	}
-	return Super::CanEditChange(InProperty);
-}
-#endif
-
-
 #undef LOCTEXT_NAMESPACE
diff --git a/Engine/Source/Runtime/Foliage/Public/ProceduralFoliageComponent.h b/Engine/Source/Runtime/Foliage/Public/ProceduralFoliageComponent.h
--- a/Engine/Source/Runtime/Foliage/Public/ProceduralFoliageComponent.h
+++ b/Engine/Source/Runtime/Foliage/Public/ProceduralFoliageComponent.h
@@ -7,7 +7,6 @@
 #include "Misc/Guid.h"
 #include "Components/ActorComponent.h"
 #include "InstancedFoliage.h"
-#include "WorldPartition/DataLayer/ActorDataLayer.h"
 #include "ProceduralFoliageComponent.generated.h"
 
 class AVolume;
@@ -77,13 +76,6 @@ class FOLIAGE_API UProceduralFoliageComponent : public UActorComponent
 	UPROPERTY(Category = "ProceduralFoliage", BlueprintReadWrite, EditAnywhere)
 	bool bShowDebugTiles;
 
-	/** Data Layer to assign to foliage actors generated by foliage spawner simulation */
-	UPROPERTY(Category = "ProceduralFoliage", BlueprintReadWrite, EditAnywhere)
-	FActorDataLayer DataLayer;
-#endif
-
-#if WITH_EDITOR
-	virtual bool CanEditChange(const FProperty* InProperty) const;
 #endif
 
 	// UObject interface
@@ -134,8 +126,6 @@ private:
 	FGuid ProceduralGuid;
 
 #if WITH_EDITORONLY_DATA
-	UPROPERTY()
-	FActorDataLayer LastSimulationDataLayer;
 	friend class UProceduralFoliageEditorLibrary;
 #endif
 };
