Checking that the ILevelEditor shared pointer is valid before dereferencing it in the main menu generating code.

diff --git a/Engine/Source/Editor/LevelEditor/Private/LevelEditorMenu.cpp b/Engine/Source/Editor/LevelEditor/Private/LevelEditorMenu.cpp
--- a/Engine/Source/Editor/LevelEditor/Private/LevelEditorMenu.cpp
+++ b/Engine/Source/Editor/LevelEditor/Private/LevelEditorMenu.cpp
@@ -255,13 +255,14 @@ void FLevelEditorMenu::RegisterLevelEditorMenus()
 					TAttribute<FText>::CreateLambda([LevelEditorName]()
 					{
 						FLevelEditorModule& LevelEditorModule = FModuleManager::LoadModuleChecked<FLevelEditorModule>(LevelEditorName);
-						return LevelEditorModule.GetLevelEditorInstance().Pin()->GetLevelViewportContextMenuTitle();
+						TSharedPtr<ILevelEditor> LevelEditorInstancePtr = LevelEditorModule.GetLevelEditorInstance().Pin();
+						return LevelEditorInstancePtr ? LevelEditorInstancePtr->GetLevelViewportContextMenuTitle() : FText::GetEmpty();
 					}),
 					TAttribute<FText>::CreateLambda([LevelEditorName]()
 					{
 						FLevelEditorModule& LevelEditorModule = FModuleManager::LoadModuleChecked<FLevelEditorModule>(LevelEditorName);
-						return FLevelEditorContextMenu::GetContextMenuToolTip(
-							ELevelEditorMenuContext::MainMenu, LevelEditorModule.GetLevelEditorInstance().Pin()->GetElementSelectionSet());
+						TSharedPtr<ILevelEditor> LevelEditorInstancePtr = LevelEditorModule.GetLevelEditorInstance().Pin();
+						return LevelEditorInstancePtr ? FLevelEditorContextMenu::GetContextMenuToolTip(ELevelEditorMenuContext::MainMenu, LevelEditorInstancePtr->GetElementSelectionSet()) : FText::GetEmpty();
 					}),
 					FOnGetContent::CreateLambda([LevelEditorName]()
 					{
@@ -270,8 +271,8 @@ void FLevelEditorMenu::RegisterLevelEditorMenus()
 						// and NOT any extenders registered for the main menu bar.
 						// I have verified that this works properly with the global Mac menu bar.
 						FLevelEditorModule& LevelEditorModule = FModuleManager::LoadModuleChecked<FLevelEditorModule>(LevelEditorName);
-						return FLevelEditorContextMenu::BuildMenuWidget(
-							LevelEditorModule.GetLevelEditorInstance(), ELevelEditorMenuContext::MainMenu, nullptr, FTypedElementHandle());
+						TSharedPtr<ILevelEditor> LevelEditorInstancePtr = LevelEditorModule.GetLevelEditorInstance().Pin();
+						return LevelEditorInstancePtr ? FLevelEditorContextMenu::BuildMenuWidget(LevelEditorInstancePtr, ELevelEditorMenuContext::MainMenu, nullptr, FTypedElementHandle()) : SNullWidget::NullWidget;
 					}));
 			
 			ActionsEntry.InsertPosition = FToolMenuInsert("Help", EToolMenuInsertType::Before);
