Make Actor Editor Context respect "Hide Viewport UI" option.

diff --git a/Engine/Source/Editor/LevelEditor/Private/SLevelViewport.cpp b/Engine/Source/Editor/LevelEditor/Private/SLevelViewport.cpp
--- a/Engine/Source/Editor/LevelEditor/Private/SLevelViewport.cpp
+++ b/Engine/Source/Editor/LevelEditor/Private/SLevelViewport.cpp
@@ -355,7 +355,7 @@ void SLevelViewport::ConstructViewportOverlayContent()
 			[
 				SNew(SActorEditorContext)
 				.World(GetWorld())
-				.Visibility_Lambda([this]() { return IsActorEditorContextVisible() ? EVisibility::Visible : EVisibility::Collapsed; })
+				.Visibility_Lambda([this]() { return IsActorEditorContextVisible() ? OnGetViewportContentVisibility() : EVisibility::Collapsed; })
 			]
 		]
 	];
@@ -372,12 +372,12 @@ void SLevelViewport::ConstructViewportOverlayContent()
 
 bool SLevelViewport::IsActorEditorContextVisible() const
 {
-	return ActiveViewport.IsValid() && 
-		(ActiveViewport->GetPlayInEditorIsSimulate() || !ActiveViewport->GetClient()->GetWorld()->IsGameWorld()) && 
-		(&GetLevelViewportClient() == GCurrentLevelEditingViewportClient) && 
-		GetDefault<ULevelEditorViewportSettings>()->bShowActorEditorContext &&
-		GetWorld() && 
-		GetWorld()->GetCurrentLevel();
+	return GetDefault<ULevelEditorViewportSettings>()->bShowActorEditorContext &&
+		GetWorld() &&
+		GetWorld()->GetCurrentLevel() &&
+		(&GetLevelViewportClient() == GCurrentLevelEditingViewportClient) &&
+		ActiveViewport.IsValid() &&
+		(ActiveViewport->GetPlayInEditorIsSimulate() || !ActiveViewport->GetClient()->GetWorld()->IsGameWorld());
 }
 
 void SLevelViewport::ConstructLevelEditorViewportClient(FLevelEditorViewportInstanceSettings& ViewportInstanceSettings)
