Prevent details panel from getting stuck on deleted elements when locked

It will now unlock and clear itself when an element is deleted

diff --git a/Engine/Source/Editor/LevelEditor/Private/SActorDetails.cpp b/Engine/Source/Editor/LevelEditor/Private/SActorDetails.cpp
--- a/Engine/Source/Editor/LevelEditor/Private/SActorDetails.cpp
+++ b/Engine/Source/Editor/LevelEditor/Private/SActorDetails.cpp
@@ -438,7 +438,8 @@ void SActorDetails::RefreshSelection(const bool bForceRefresh)
 	bHasSelectionOverride = false;
 	SelectionOverrideActors.Reset();
 
-	RefreshTopLevelElements(DetailsElements, bForceRefresh, /*bOverrideLock*/false);
+	const bool bOverrideLock = DetailsElements.IsEmpty();
+	RefreshTopLevelElements(DetailsElements, bForceRefresh, bOverrideLock);
 }
 
 void SActorDetails::OverrideSelection(const TArray<AActor*>& InActors, const bool bForceRefresh)
