Fix a crash when changing layout or viewport type in Level editor viewports and running various exec commands that used typed element common actions.

diff --git a/Engine/Source/Editor/EditorFramework/Public/ViewportTabContent.h b/Engine/Source/Editor/EditorFramework/Public/ViewportTabContent.h
--- a/Engine/Source/Editor/EditorFramework/Public/ViewportTabContent.h
+++ b/Engine/Source/Editor/EditorFramework/Public/ViewportTabContent.h
@@ -14,6 +14,7 @@
 class SDockTab;
 class FEditorViewportLayout;
 class IEditorViewportLayoutEntity;
+class FUICommandList;
 
 class EDITORFRAMEWORK_API FViewportTabContent
 {
@@ -36,6 +37,12 @@ public:
 
 	virtual void SetViewportConfiguration(const FName& ConfigurationName) {}
 
+	/**
+	 * Maps the common commands into the given command list for this layout to the given viewport based on the config key.
+	 * This call is only additive to the command list, and does not clear or erase already mapped actions.
+	 */
+	virtual void BindViewportLayoutCommands(FUICommandList& InOutCommandList, FName ViewportConfigKey) {}
+
 	const TMap< FName, TSharedPtr< IEditorViewportLayoutEntity > >* GetViewports() const;
 
 	using ViewportActionFunction = TFunction<void(FName Name, TSharedPtr<IEditorViewportLayoutEntity>)>;
diff --git a/Engine/Source/Editor/LevelEditor/Private/LevelViewportTabContent.cpp b/Engine/Source/Editor/LevelEditor/Private/LevelViewportTabContent.cpp
--- a/Engine/Source/Editor/LevelEditor/Private/LevelViewportTabContent.cpp
+++ b/Engine/Source/Editor/LevelEditor/Private/LevelViewportTabContent.cpp
@@ -8,6 +8,12 @@
 #include "Widgets/Docking/SDockTab.h"
 #include "LevelViewportLayout.h"
 #include "LevelEditorViewport.h"
+#include "LevelViewportActions.h"
+#include "Framework/Commands/UIAction.h"
+#include "Framework/Commands/UICommandList.h"
+#include "ToolMenus.h"
+#include "ViewportTypeDefinition.h"
+#include "Templates/SharedPointer.h"
 
 
 // FLevelViewportTabContent ///////////////////////////
@@ -52,6 +58,48 @@ void FLevelViewportTabContent::Initialize(AssetEditorViewportFactoryFunction Fun
 	FEditorViewportTabContent::Initialize(Func, InParentTab, InLayoutString);
 }
 
+void FLevelViewportTabContent::BindViewportLayoutCommands(FUICommandList& InOutCommandList, FName ViewportConfigKey)
+{
+	FEditorViewportTabContent::BindViewportLayoutCommands(InOutCommandList, ViewportConfigKey);
+
+	FLevelViewportCommands& ViewportActions = FLevelViewportCommands::Get();
+
+	auto AddViewportConfigurationAction = [&InOutCommandList, this](const TSharedPtr<FUICommandInfo>& InCommandInfo, FName InLayoutType)
+	{
+		InOutCommandList.MapAction(
+			InCommandInfo,
+			FExecuteAction::CreateSP(this, &FLevelViewportTabContent::OnUIActionSetViewportConfiguration, InLayoutType),
+			FCanExecuteAction(),
+			FIsActionChecked::CreateSP(this, &FLevelViewportTabContent::IsViewportConfigurationChecked, InLayoutType));
+	};
+
+	AddViewportConfigurationAction(ViewportActions.ViewportConfig_OnePane, LevelViewportConfigurationNames::OnePane);
+	AddViewportConfigurationAction(ViewportActions.ViewportConfig_TwoPanesH, LevelViewportConfigurationNames::TwoPanesHoriz);
+	AddViewportConfigurationAction(ViewportActions.ViewportConfig_TwoPanesV, LevelViewportConfigurationNames::TwoPanesVert);
+	AddViewportConfigurationAction(ViewportActions.ViewportConfig_ThreePanesLeft, LevelViewportConfigurationNames::ThreePanesLeft);
+	AddViewportConfigurationAction(ViewportActions.ViewportConfig_ThreePanesRight, LevelViewportConfigurationNames::ThreePanesRight);
+	AddViewportConfigurationAction(ViewportActions.ViewportConfig_ThreePanesTop, LevelViewportConfigurationNames::ThreePanesTop);
+	AddViewportConfigurationAction(ViewportActions.ViewportConfig_ThreePanesBottom, LevelViewportConfigurationNames::ThreePanesBottom);
+	AddViewportConfigurationAction(ViewportActions.ViewportConfig_FourPanesLeft, LevelViewportConfigurationNames::FourPanesLeft);
+	AddViewportConfigurationAction(ViewportActions.ViewportConfig_FourPanesRight, LevelViewportConfigurationNames::FourPanesRight);
+	AddViewportConfigurationAction(ViewportActions.ViewportConfig_FourPanesTop, LevelViewportConfigurationNames::FourPanesTop);
+	AddViewportConfigurationAction(ViewportActions.ViewportConfig_FourPanesBottom, LevelViewportConfigurationNames::FourPanesBottom);
+	AddViewportConfigurationAction(ViewportActions.ViewportConfig_FourPanes2x2, LevelViewportConfigurationNames::FourPanes2x2);
+
+	auto ProcessViewportTypeActions = [&InOutCommandList, ViewportConfigKey, this](FName InViewportTypeName, const FViewportTypeDefinition& InDefinition) {
+		if (InDefinition.ActivationCommand.IsValid())
+		{
+			InOutCommandList.MapAction(InDefinition.ActivationCommand, FUIAction(
+				FExecuteAction::CreateSP(this, &FLevelViewportTabContent::OnUIActionSetViewportTypeWithinLayout, ViewportConfigKey, InViewportTypeName),
+				FCanExecuteAction(),
+				FIsActionChecked::CreateSP(this, &FLevelViewportTabContent::IsViewportTypeWithinLayoutEqual, ViewportConfigKey, InViewportTypeName)
+			));
+		}
+	};
+	FLevelEditorModule& LevelEditorModule = FModuleManager::GetModuleChecked<FLevelEditorModule>(TEXT("LevelEditor"));
+	LevelEditorModule.IterateViewportTypes(ProcessViewportTypeActions);
+}
+
 void FLevelViewportTabContent::OnLayoutStartChange(bool bSwitchingLayouts)
 {
 	GCurrentLevelEditingViewportClient = nullptr;
@@ -82,3 +130,49 @@ void FLevelViewportTabContent::OnLayoutChanged()
 	FLevelEditorModule& LevelEditorModule = FModuleManager::LoadModuleChecked<FLevelEditorModule>("LevelEditor");
 	LevelEditorModule.OnTabContentChanged().Broadcast();
 }
+
+void FLevelViewportTabContent::OnUIActionSetViewportConfiguration(FName InConfigurationName)
+{
+	SetViewportConfiguration(InConfigurationName);
+
+	FSlateApplication::Get().DismissAllMenus();
+	UToolMenus::Get()->CleanupStaleWidgetsNextTick(true);
+}
+
+FName FLevelViewportTabContent::GetViewportTypeWithinLayout(FName InConfigKey) const
+{
+	if (ActiveViewportLayout && !InConfigKey.IsNone())
+	{
+		if (TSharedPtr<IEditorViewportLayoutEntity> LayoutEntity = ActiveViewportLayout->GetViewports().FindRef(InConfigKey))
+		{
+			return LayoutEntity->GetType();
+		}
+	}
+
+	return "Default";
+}
+
+void FLevelViewportTabContent::OnUIActionSetViewportTypeWithinLayout(FName InConfigKey, FName InLayoutType)
+{
+	if (!IsViewportTypeWithinLayoutEqual(InConfigKey, InLayoutType))
+	{
+		SaveConfig();
+
+		const FString& IniSection = FLayoutSaveRestore::GetAdditionalLayoutConfigIni();
+		GConfig->SetString(*IniSection, *(InConfigKey.ToString() + TEXT(".TypeWithinLayout")), *InLayoutType.ToString(), GEditorPerProjectIni);
+
+		// Force a refresh of the tab content and clear all menus
+		RefreshViewportConfiguration();
+		FSlateApplication::Get().DismissAllMenus();
+	}
+}
+
+bool FLevelViewportTabContent::IsViewportTypeWithinLayoutEqual(FName InConfigKey, FName InLayoutType) const
+{
+	return (GetViewportTypeWithinLayout(InConfigKey) == InLayoutType);
+}
+
+bool FLevelViewportTabContent::IsViewportConfigurationChecked(FName InLayoutType) const
+{
+	return IsViewportConfigurationSet(InLayoutType);
+}
diff --git a/Engine/Source/Editor/LevelEditor/Private/SLevelViewport.cpp b/Engine/Source/Editor/LevelEditor/Private/SLevelViewport.cpp
--- a/Engine/Source/Editor/LevelEditor/Private/SLevelViewport.cpp
+++ b/Engine/Source/Editor/LevelEditor/Private/SLevelViewport.cpp
@@ -1453,90 +1453,13 @@ void SLevelViewport::BindViewCommands( FUICommandList& OutCommandList )
 		FCanExecuteAction::CreateSP( this, &SLevelViewport::CanExecuteActorLockSelected )
 		);
 
-	OutCommandList.MapAction(
-		ViewportActions.ViewportConfig_OnePane,
-		FExecuteAction::CreateSP(this, &SLevelViewport::OnSetViewportConfiguration, LevelViewportConfigurationNames::OnePane),
-		FCanExecuteAction(),
-		FIsActionChecked::CreateSP(this, &SLevelViewport::IsViewportConfigurationSet, LevelViewportConfigurationNames::OnePane));
-
-	OutCommandList.MapAction(
-		ViewportActions.ViewportConfig_TwoPanesH,
-		FExecuteAction::CreateSP( this, &SLevelViewport::OnSetViewportConfiguration, LevelViewportConfigurationNames::TwoPanesHoriz ),
-		FCanExecuteAction(),
-		FIsActionChecked::CreateSP( this, &SLevelViewport::IsViewportConfigurationSet, LevelViewportConfigurationNames::TwoPanesHoriz ));
-
-	OutCommandList.MapAction(
-		ViewportActions.ViewportConfig_TwoPanesV,
-		FExecuteAction::CreateSP( this, &SLevelViewport::OnSetViewportConfiguration, LevelViewportConfigurationNames::TwoPanesVert ),
-		FCanExecuteAction(),
-		FIsActionChecked::CreateSP( this, &SLevelViewport::IsViewportConfigurationSet, LevelViewportConfigurationNames::TwoPanesVert ));
-
-	OutCommandList.MapAction(
-		ViewportActions.ViewportConfig_ThreePanesLeft,
-		FExecuteAction::CreateSP( this, &SLevelViewport::OnSetViewportConfiguration, LevelViewportConfigurationNames::ThreePanesLeft ),
-		FCanExecuteAction(),
-		FIsActionChecked::CreateSP( this, &SLevelViewport::IsViewportConfigurationSet, LevelViewportConfigurationNames::ThreePanesLeft ));
-
-	OutCommandList.MapAction(
-		ViewportActions.ViewportConfig_ThreePanesRight,
-		FExecuteAction::CreateSP( this, &SLevelViewport::OnSetViewportConfiguration, LevelViewportConfigurationNames::ThreePanesRight ),
-		FCanExecuteAction(),
-		FIsActionChecked::CreateSP( this, &SLevelViewport::IsViewportConfigurationSet, LevelViewportConfigurationNames::ThreePanesRight ));
-
-	OutCommandList.MapAction(
-		ViewportActions.ViewportConfig_ThreePanesTop,
-		FExecuteAction::CreateSP( this, &SLevelViewport::OnSetViewportConfiguration, LevelViewportConfigurationNames::ThreePanesTop ),
-		FCanExecuteAction(),
-		FIsActionChecked::CreateSP( this, &SLevelViewport::IsViewportConfigurationSet, LevelViewportConfigurationNames::ThreePanesTop ));
-
-	OutCommandList.MapAction(
-		ViewportActions.ViewportConfig_ThreePanesBottom,
-		FExecuteAction::CreateSP( this, &SLevelViewport::OnSetViewportConfiguration, LevelViewportConfigurationNames::ThreePanesBottom ),
-		FCanExecuteAction(),
-		FIsActionChecked::CreateSP( this, &SLevelViewport::IsViewportConfigurationSet, LevelViewportConfigurationNames::ThreePanesBottom ));
-
-	OutCommandList.MapAction(
-		ViewportActions.ViewportConfig_FourPanesLeft,
-		FExecuteAction::CreateSP( this, &SLevelViewport::OnSetViewportConfiguration, LevelViewportConfigurationNames::FourPanesLeft ),
-		FCanExecuteAction(),
-		FIsActionChecked::CreateSP( this, &SLevelViewport::IsViewportConfigurationSet, LevelViewportConfigurationNames::FourPanesLeft ));
-
-	OutCommandList.MapAction(
-		ViewportActions.ViewportConfig_FourPanesRight,
-		FExecuteAction::CreateSP( this, &SLevelViewport::OnSetViewportConfiguration, LevelViewportConfigurationNames::FourPanesRight ),
-		FCanExecuteAction(),
-		FIsActionChecked::CreateSP( this, &SLevelViewport::IsViewportConfigurationSet, LevelViewportConfigurationNames::FourPanesRight ));
-
-	OutCommandList.MapAction(
-		ViewportActions.ViewportConfig_FourPanesTop,
-		FExecuteAction::CreateSP( this, &SLevelViewport::OnSetViewportConfiguration, LevelViewportConfigurationNames::FourPanesTop ),
-		FCanExecuteAction(),
-		FIsActionChecked::CreateSP( this, &SLevelViewport::IsViewportConfigurationSet, LevelViewportConfigurationNames::FourPanesTop ));
-
-	OutCommandList.MapAction(
-		ViewportActions.ViewportConfig_FourPanesBottom,
-		FExecuteAction::CreateSP( this, &SLevelViewport::OnSetViewportConfiguration, LevelViewportConfigurationNames::FourPanesBottom ),
-		FCanExecuteAction(),
-		FIsActionChecked::CreateSP( this, &SLevelViewport::IsViewportConfigurationSet, LevelViewportConfigurationNames::FourPanesBottom ));
-
-	OutCommandList.MapAction(
-		ViewportActions.ViewportConfig_FourPanes2x2,
-		FExecuteAction::CreateSP( this, &SLevelViewport::OnSetViewportConfiguration, LevelViewportConfigurationNames::FourPanes2x2 ),
-		FCanExecuteAction(),
-		FIsActionChecked::CreateSP( this, &SLevelViewport::IsViewportConfigurationSet, LevelViewportConfigurationNames::FourPanes2x2 ));
-
-	auto ProcessViewportTypeActions = [&](FName InViewportTypeName, const FViewportTypeDefinition& InDefinition){
-		if (InDefinition.ActivationCommand.IsValid())
+	if (TSharedPtr<FLevelViewportLayout> LayoutPinned = ParentLayout.Pin())
+	{
+		if (TSharedPtr<FEditorViewportTabContent> ParentTabContentPinned = LayoutPinned->GetParentTabContent().Pin())
 		{
-			OutCommandList.MapAction(InDefinition.ActivationCommand, FUIAction(
-				FExecuteAction::CreateSP(this, &SLevelViewport::ToggleViewportTypeActivationWithinLayout, InViewportTypeName),
-				FCanExecuteAction(),
-				FIsActionChecked::CreateSP(this, &SLevelViewport::IsViewportTypeWithinLayoutEqual, InViewportTypeName)
-			));
+			ParentTabContentPinned->BindViewportLayoutCommands(OutCommandList, ConfigKey);
 		}
-	};
-	FLevelEditorModule& LevelEditorModule = FModuleManager::GetModuleChecked<FLevelEditorModule>(TEXT("LevelEditor"));
-	LevelEditorModule.IterateViewportTypes(ProcessViewportTypeActions);
+	}
 
 	FBufferVisualizationMenuCommands::Get().BindCommands(OutCommandList, Client);
 	FNaniteVisualizationMenuCommands::Get().BindCommands(OutCommandList, Client);
@@ -3939,11 +3862,13 @@ void SLevelViewport::SetViewportTypeWithinLayout(FName InLayoutType)
 	{
 		// Important - RefreshViewportConfiguration does not save config values. We save its state first, to ensure that .TypeWithinLayout (below) doesn't get overwritten
 		TSharedPtr<FEditorViewportTabContent> ViewportTabPinned = LayoutPinned->GetParentTabContent().Pin();
-		if (ViewportTabPinned.IsValid())
+		if (!ViewportTabPinned)
 		{
-			ViewportTabPinned->SaveConfig();
+			return;
 		}
 
+		ViewportTabPinned->SaveConfig();
+
 		const FString& IniSection = FLayoutSaveRestore::GetAdditionalLayoutConfigIni();
 		GConfig->SetString( *IniSection, *( ConfigKey.ToString() + TEXT(".TypeWithinLayout") ), *InLayoutType.ToString(), GEditorPerProjectIni );
 
diff --git a/Engine/Source/Editor/LevelEditor/Public/LevelViewportTabContent.h b/Engine/Source/Editor/LevelEditor/Public/LevelViewportTabContent.h
--- a/Engine/Source/Editor/LevelEditor/Public/LevelViewportTabContent.h
+++ b/Engine/Source/Editor/LevelEditor/Public/LevelViewportTabContent.h
@@ -21,10 +21,19 @@ public:
 	/** Starts the tab content object and creates the initial layout based on the layout string */
 	virtual void Initialize(AssetEditorViewportFactoryFunction Func, TSharedPtr<SDockTab> InParentTab, const FString& InLayoutString) override;
 
+	virtual void BindViewportLayoutCommands(FUICommandList& InOutCommandList, FName ViewportConfigKey) override;
+
 protected:
 	virtual TSharedPtr<FEditorViewportLayout> FactoryViewportLayout(bool bIsSwitchingLayouts) override;
 	virtual FName GetLayoutTypeNameFromLayoutString() const override;
 
 	void OnLayoutStartChange(bool bSwitchingLayouts);
 	void OnLayoutChanged();
+
+private:
+	void OnUIActionSetViewportConfiguration(FName InConfigurationName);
+	FName GetViewportTypeWithinLayout(FName InConfigKey) const;
+	void OnUIActionSetViewportTypeWithinLayout(FName InConfigKey, FName InLayoutType);
+	bool IsViewportTypeWithinLayoutEqual(FName InConfigName, FName InLayoutType) const;
+	bool IsViewportConfigurationChecked(FName InLayoutType) const;
 };
diff --git a/Engine/Source/Editor/LevelEditor/Public/SLevelViewport.h b/Engine/Source/Editor/LevelEditor/Public/SLevelViewport.h
--- a/Engine/Source/Editor/LevelEditor/Public/SLevelViewport.h
+++ b/Engine/Source/Editor/LevelEditor/Public/SLevelViewport.h
@@ -342,6 +342,7 @@ public:
 	 * 
 	 * @param ConfigurationName		The name of the layout (for the names in namespace LevelViewportConfigurationNames)
 	 */
+	UE_DEPRECATED(5.1, "Moved to internal handling by FViewportTabContent. See FViewportTabContent::BindCommonViewportCommands")
 	void OnSetViewportConfiguration(FName ConfigurationName);
 
 	/**
@@ -350,18 +351,23 @@ public:
 	 * @param ConfigurationName		The name of the layout (for the names in namespace LevelViewportConfigurationNames)
 	 * @return						True, if the named layout is currently active
 	 */
+	UE_DEPRECATED(5.1, "Moved to internal handling by FViewportTabContent. See FViewportTabContent::BindCommonViewportCommands")
 	bool IsViewportConfigurationSet(FName ConfigurationName) const;
 
 	/** Get this level viewport widget's type within its parent layout */
+	UE_DEPRECATED(5.1, "Moved to internal handling by FViewportTabContent. See FViewportTabContent::BindCommonViewportCommands")
 	FName GetViewportTypeWithinLayout() const;
 
 	/** Set this level viewport widget's type within its parent layout */
+	UE_DEPRECATED(5.1, "Moved to internal handling by FViewportTabContent. See FViewportTabContent::BindCommonViewportCommands")
 	void SetViewportTypeWithinLayout(FName InLayoutType);
 
 	/** Activates the specified viewport type in the layout, if it's not already, or reverts to default if it is. */
+	UE_DEPRECATED(5.1, "Moved to internal handling by FViewportTabContent. See FViewportTabContent::BindCommonViewportCommands")
 	void ToggleViewportTypeActivationWithinLayout(FName InLayoutType);
 
 	/** Checks if the specified layout type matches our current viewport type. */
+	UE_DEPRECATED(5.1, "Moved to internal handling by FViewportTabContent. See FViewportTabContent::BindCommonViewportCommands")
 	bool IsViewportTypeWithinLayoutEqual(FName InLayoutType);
 
 	/** For the specified actor, See if we're forcing a preview */
