Adding missing deprecation pragma guards from CL 19853509

#rnx
#pf skip

diff --git a/Engine/Source/Editor/LevelEditor/Private/SLevelViewport.cpp b/Engine/Source/Editor/LevelEditor/Private/SLevelViewport.cpp
--- a/Engine/Source/Editor/LevelEditor/Private/SLevelViewport.cpp
+++ b/Engine/Source/Editor/LevelEditor/Private/SLevelViewport.cpp
@@ -3880,15 +3880,19 @@ void SLevelViewport::SetViewportTypeWithinLayout(FName InLayoutType)
 
 void SLevelViewport::ToggleViewportTypeActivationWithinLayout(FName InLayoutType)
 {
+PRAGMA_DISABLE_DEPRECATION_WARNINGS
 	if (GetViewportTypeWithinLayout() != InLayoutType)
 	{
 		SetViewportTypeWithinLayout(InLayoutType);
 	}
+PRAGMA_ENABLE_DEPRECATION_WARNINGS
 }
 
 bool SLevelViewport::IsViewportTypeWithinLayoutEqual(FName InLayoutType)
 {
+PRAGMA_DISABLE_DEPRECATION_WARNINGS
 	return GetViewportTypeWithinLayout() == InLayoutType;
+PRAGMA_ENABLE_DEPRECATION_WARNINGS
 }
 
 void SLevelViewport::StartPlayInEditorSession(UGameViewportClient* PlayClient, const bool bInSimulateInEditor)
