Fix a crash in level editor viewports when the slate window is destroyed and recreated (like when reloading the default viewport).

diff --git a/Engine/Source/Editor/LevelEditor/Private/SLevelEditor.cpp b/Engine/Source/Editor/LevelEditor/Private/SLevelEditor.cpp
--- a/Engine/Source/Editor/LevelEditor/Private/SLevelEditor.cpp
+++ b/Engine/Source/Editor/LevelEditor/Private/SLevelEditor.cpp
@@ -387,6 +387,8 @@ SLevelEditor::~SLevelEditor()
 		{
 			SelectedElements->ClearSelection(FTypedElementSelectionOptions());
 		}
+		SelectedElements->OnPreChange().Clear();
+		SelectedElements->OnChanged().Clear();
 		SelectedElements->RemoveFromRoot();
 		SelectedElements = nullptr;
 
