Build Menu - Display the World Partition section depending if we are in a world partition world or not. If not, display the previous hlod menu entry like before.

#rnx

diff --git a/Engine/Source/Editor/LevelEditor/Private/LevelEditorActions.cpp b/Engine/Source/Editor/LevelEditor/Private/LevelEditorActions.cpp
--- a/Engine/Source/Editor/LevelEditor/Private/LevelEditorActions.cpp
+++ b/Engine/Source/Editor/LevelEditor/Private/LevelEditorActions.cpp
@@ -807,7 +807,7 @@ void FLevelEditorActionCallbacks::BuildHLODs_Execute()
 
 void FLevelEditorActionCallbacks::BuildMinimap_Execute()
 {
-	// Build HLOD
+	// Build Minimap
 	FEditorBuildUtils::EditorBuild(GetWorld(), FBuildOptions::BuildMinimap);
 }
 
diff --git a/Engine/Source/Editor/LevelEditor/Private/LevelEditorMenu.cpp b/Engine/Source/Editor/LevelEditor/Private/LevelEditorMenu.cpp
--- a/Engine/Source/Editor/LevelEditor/Private/LevelEditorMenu.cpp
+++ b/Engine/Source/Editor/LevelEditor/Private/LevelEditorMenu.cpp
@@ -637,9 +637,27 @@ void FLevelEditorMenu::RegisterBuildMenu()
 	}
 
 	{
-		FToolMenuSection& Section = Menu->AddSection("LevelEditorWorldPartition", LOCTEXT("WorldPartitionHeading", "World Partition"));
-		Section.AddMenuEntry(FLevelEditorCommands::Get().BuildHLODs);
-		Section.AddMenuEntry(FLevelEditorCommands::Get().BuildMinimap);
+		Menu->AddDynamicSection(NAME_None, FNewToolMenuDelegate::CreateLambda([](UToolMenu* InMenu)
+		{
+			if (GWorld != nullptr && GWorld->GetWorldPartition() != nullptr)
+			{
+				FToolMenuSection& Section = InMenu->AddSection("LevelEditorWorldPartition", LOCTEXT("WorldPartitionHeading", "World Partition"));
+				Section.AddMenuEntry(FLevelEditorCommands::Get().BuildHLODs);
+				Section.AddMenuEntry(FLevelEditorCommands::Get().BuildMinimap);
+			}
+		}));
+	}
+
+	{
+		// The day we only support World Partitioned worlds, we can remove this section.
+		Menu->AddDynamicSection(NAME_None, FNewToolMenuDelegate::CreateLambda([](UToolMenu* InMenu)
+		{
+			if (GWorld != nullptr && GWorld->GetWorldPartition() == nullptr)
+			{
+				FToolMenuSection& Section = InMenu->AddSection("LevelEditorLOD", LOCTEXT("LODHeading", "Hierarchical LOD"));
+				Section.AddMenuEntry(FLevelEditorCommands::Get().BuildHLODs);
+			}
+		}));
 	}
 
 	{
