Add Extender to middle of LevelViewportToolbar

diff --git a/Engine/Source/Editor/LevelEditor/Private/SLevelViewportToolBar.cpp b/Engine/Source/Editor/LevelEditor/Private/SLevelViewportToolBar.cpp
--- a/Engine/Source/Editor/LevelEditor/Private/SLevelViewportToolBar.cpp
+++ b/Engine/Source/Editor/LevelEditor/Private/SLevelViewportToolBar.cpp
@@ -5,6 +5,7 @@
 #include "Framework/MultiBox/MultiBoxDefs.h"
 #include "Framework/MultiBox/MultiBoxExtender.h"
 #include "Framework/MultiBox/MultiBoxBuilder.h"
+#include "Subsystems/PanelExtensionSubsystem.h"
 #include "ToolMenus.h"
 #include "LevelEditorMenuContext.h"
 #include "Modules/ModuleManager.h"
@@ -280,6 +281,14 @@ void SLevelViewportToolBar::Construct( const FArguments& InArgs )
 				]
 				+ SHorizontalBox::Slot()
 				.Padding(ToolbarSlotPadding)
+				.HAlign(HAlign_Center)
+				.VAlign(VAlign_Fill)
+				[
+					SAssignNew(ExtensionPanel, SExtensionPanel)
+					.ExtensionPanelID("LevelViewportToolBar.MiddleExtension")
+				]
+				+ SHorizontalBox::Slot()
+				.Padding(ToolbarSlotPadding)
 				.HAlign(HAlign_Right)
 				[
 					SNew(STransformViewportToolBar)
diff --git a/Engine/Source/Editor/LevelEditor/Private/SLevelViewportToolBar.h b/Engine/Source/Editor/LevelEditor/Private/SLevelViewportToolBar.h
--- a/Engine/Source/Editor/LevelEditor/Private/SLevelViewportToolBar.h
+++ b/Engine/Source/Editor/LevelEditor/Private/SLevelViewportToolBar.h
@@ -11,6 +11,7 @@
 
 class ACameraActor;
 class FExtender;
+class SExtensionPanel;
 class FMenuBuilder;
 class UToolMenu;
 struct FToolMenuSection;
@@ -279,5 +280,8 @@ private:
 private:
 	/** The viewport that we are in */
 	TWeakPtr<class SLevelViewport> Viewport;
+
+	/** SLevelViewportToolBar UI extension */
+	TSharedPtr<SExtensionPanel> ExtensionPanel;
 };
 
