Expose BuildLightMaps command to python and blueprint utilities

diff --git a/Engine/Source/Editor/LevelEditor/Private/LevelEditorSubsystem.cpp b/Engine/Source/Editor/LevelEditor/Private/LevelEditorSubsystem.cpp
--- a/Engine/Source/Editor/LevelEditor/Private/LevelEditorSubsystem.cpp
+++ b/Engine/Source/Editor/LevelEditor/Private/LevelEditorSubsystem.cpp
@@ -23,6 +23,7 @@
 #include "LevelInstance/LevelInstanceSubsystem.h"
 #include "LevelInstance/LevelInstanceInterface.h"
 #include "ClassIconFinder.h"
+#include "LightingBuildOptions.h"
 
 DEFINE_LOG_CATEGORY_STATIC(LevelEditorSubsystem, Log, All);
 
@@ -650,6 +651,56 @@ UTypedElementSelectionSet* ULevelEditorSubsystem::GetSelectionSet()
 	return nullptr;
 }
 
+bool ULevelEditorSubsystem::BuildLightMaps(ELightingBuildQuality Quality, bool bWithReflectionCaptures)
+{
+	FLightingBuildOptions LightingOptions;
+	LightingOptions.QualityLevel = Quality;
+
+	UUnrealEditorSubsystem* UnrealEditorSubsystem = GEditor->GetEditorSubsystem<UUnrealEditorSubsystem>();
+	if (!UnrealEditorSubsystem)
+	{
+		return false;
+	}
+
+	UWorld* World = UnrealEditorSubsystem->GetEditorWorld();
+	if (!World)
+	{
+		UE_LOG(LevelEditorSubsystem, Error, TEXT("BuildLightMaps. Can't build the light maps of the current level because there is no world."));
+		return false;
+	}
+
+	bool Success = false;
+
+	auto BuildFailedDelegate = [&World, &Success]() {
+		UE_LOG(LevelEditorSubsystem, Error, TEXT("BuildLightMaps. Failed building lighting for %s"), *World->GetOutermost()->GetName());
+		Success = false;
+	};
+	FDelegateHandle BuildFailedDelegateHandle = FEditorDelegates::OnLightingBuildFailed.AddLambda(BuildFailedDelegate);
+
+	auto BuildSucceededDelegate = [&World, &Success]() {
+		UE_LOG(LevelEditorSubsystem, Log, TEXT("BuildLightMaps. Successfully built lighting for %s"), *World->GetOutermost()->GetName());
+		Success = true;
+	};
+	FDelegateHandle BuildSucceededDelegateHandle = FEditorDelegates::OnLightingBuildSucceeded.AddLambda(BuildSucceededDelegate);
+
+	UE_LOG(LevelEditorSubsystem, Log, TEXT("BuildLightMaps. Start building lighting for %s"), *World->GetOutermost()->GetName());
+	GEditor->BuildLighting(LightingOptions);
+	while (GEditor->IsLightingBuildCurrentlyRunning())
+	{
+		GEditor->UpdateBuildLighting();
+	}
+
+	if (bWithReflectionCaptures)
+	{
+		GEditor->BuildReflectionCaptures();
+	}
+
+	FEditorDelegates::OnLightingBuildFailed.Remove(BuildFailedDelegateHandle);
+	FEditorDelegates::OnLightingBuildFailed.Remove(BuildSucceededDelegateHandle);
+
+	return Success;
+}
+
 FEditorModeTools* ULevelEditorSubsystem::GetLevelEditorModeManager()
 {
 	FLevelEditorModule& LevelEditorModule = FModuleManager::GetModuleChecked<FLevelEditorModule>("LevelEditor");
diff --git a/Engine/Source/Editor/LevelEditor/Public/LevelEditorSubsystem.h b/Engine/Source/Editor/LevelEditor/Public/LevelEditorSubsystem.h
--- a/Engine/Source/Editor/LevelEditor/Public/LevelEditorSubsystem.h
+++ b/Engine/Source/Editor/LevelEditor/Public/LevelEditorSubsystem.h
@@ -5,6 +5,7 @@
 #include "CoreMinimal.h"
 #include "EditorSubsystem.h"
 #include "IActorEditorContextClient.h"
+#include "Engine/EngineTypes.h"
 
 #include "LevelEditorSubsystem.generated.h"
 
@@ -131,6 +132,15 @@ public:
 	UFUNCTION(BlueprintCallable, Category = "Editor Scripting | Level Utility")
 	UTypedElementSelectionSet* GetSelectionSet();
 
+	/**
+	 * Build Light Maps and optionally the reflection captures.
+	 * @param	Quality	One of the enum LightingBuildQuality value. Default is Quality_Production.
+	 * @param	bWithReflectionCaptures	Build the related reflection captures after building the light maps.
+	 * @return	True if build was successful.
+	 */
+	UFUNCTION(BlueprintCallable, Category = "Editor Scripting | Level Utility")
+	bool BuildLightMaps(ELightingBuildQuality Quality = ELightingBuildQuality::Quality_Production, bool bWithReflectionCaptures = false);
+
 	/**
 	 * Gets the global level editor mode manager, if we have one.
 	 * The mode manager is not created in commandlet environments, because modes inherently imply user interactions.
diff --git a/Engine/Source/Runtime/Engine/Classes/Engine/EngineTypes.h b/Engine/Source/Runtime/Engine/Classes/Engine/EngineTypes.h
--- a/Engine/Source/Runtime/Engine/Classes/Engine/EngineTypes.h
+++ b/Engine/Source/Runtime/Engine/Classes/Engine/EngineTypes.h
@@ -795,14 +795,14 @@ enum EMaterialShadingRate
 
 
 /**	Lighting build quality enumeration */
-UENUM()
+UENUM(BlueprintType)
 enum ELightingBuildQuality
 {
-	Quality_Preview,
-	Quality_Medium,
-	Quality_High,
-	Quality_Production,
-	Quality_MAX,
+	Quality_Preview		UMETA(DisplayName = "Preview"),
+	Quality_Medium		UMETA(DisplayName = "Medium"),
+	Quality_High		UMETA(DisplayName = "High"),
+	Quality_Production	UMETA(DisplayName = "Production"),
+	Quality_MAX			UMETA(Hidden),
 };
 
 /** Movement modes for Characters. */
