- Fixed `nullptr` dereference when the module is being unloaded since the tab's parent is already de-allocated if the engine is exiting.
- Also prevents this sort of bug from occurring in Slate altogether.

diff --git a/Engine/Source/Editor/LevelEditor/Private/LevelEditor.cpp b/Engine/Source/Editor/LevelEditor/Private/LevelEditor.cpp
--- a/Engine/Source/Editor/LevelEditor/Private/LevelEditor.cpp
+++ b/Engine/Source/Editor/LevelEditor/Private/LevelEditor.cpp
@@ -319,12 +319,15 @@ void FLevelEditorModule::ShutdownModule()
 
 	// If the level editor tab is currently open, close it
 	{
-		TSharedPtr<SDockTab> LevelEditorTab = LevelEditorInstanceTabPtr.Pin();
-		if (LevelEditorTab.IsValid())
+		if (!IsEngineExitRequested())
 		{
-			LevelEditorTab->RemoveTabFromParent();
+			TSharedPtr<SDockTab> LevelEditorTab = LevelEditorInstanceTabPtr.Pin();
+			if (LevelEditorTab.IsValid())
+			{
+				LevelEditorTab->RemoveTabFromParent();
+			}
+			LevelEditorInstanceTabPtr.Reset();
 		}
-		LevelEditorInstanceTabPtr.Reset();
 	}
 
 	// Clear out some globals that may be referencing this module
diff --git a/Engine/Source/Runtime/Slate/Private/Widgets/Docking/SDockTab.cpp b/Engine/Source/Runtime/Slate/Private/Widgets/Docking/SDockTab.cpp
--- a/Engine/Source/Runtime/Slate/Private/Widgets/Docking/SDockTab.cpp
+++ b/Engine/Source/Runtime/Slate/Private/Widgets/Docking/SDockTab.cpp
@@ -990,5 +990,8 @@ FVector2D SDockTab::GetAnimatedScale() const
 
 void SDockTab::UpdateActivationTime()
 {
-	LastActivationTime = FSlateApplication::Get().GetCurrentTime();
+	if (FSlateApplication::IsInitialized())
+	{
+		LastActivationTime = FSlateApplication::Get().GetCurrentTime();
+	}
 }
diff --git a/Engine/Source/Runtime/Slate/Public/Framework/Application/SWindowTitleBar.h b/Engine/Source/Runtime/Slate/Public/Framework/Application/SWindowTitleBar.h
--- a/Engine/Source/Runtime/Slate/Public/Framework/Application/SWindowTitleBar.h
+++ b/Engine/Source/Runtime/Slate/Public/Framework/Application/SWindowTitleBar.h
@@ -209,8 +209,12 @@ public:
 	}
 
 
-	virtual void SetAllowMenuBar(bool bInAllowMenuBar)
+	virtual void SetAllowMenuBar(bool bInAllowMenuBar) override
 	{
+		if (!FSlateApplicationBase::IsInitialized())
+		{
+			return;
+		}
 		bAllowMenuBar = bInAllowMenuBar;
 
 		if (AppIconWidget)
