Check for a nullptr for the asset class

diff --git a/Engine/Source/Editor/LevelEditor/Private/LevelEditorCreateActorMenu.cpp b/Engine/Source/Editor/LevelEditor/Private/LevelEditorCreateActorMenu.cpp
--- a/Engine/Source/Editor/LevelEditor/Private/LevelEditorCreateActorMenu.cpp
+++ b/Engine/Source/Editor/LevelEditor/Private/LevelEditorCreateActorMenu.cpp
@@ -97,8 +97,9 @@ static void GetMenuEntryText(const FAssetData& Asset, const TArray<FActorFactory
 	if (AssetMenuOptions.Num() == 1)
 	{
 		const FActorFactoryAssetProxy::FMenuItem& MenuItem = AssetMenuOptions[0];
+		UClass* MenuItemClass = Cast<UClass>(MenuItem.AssetData.GetAsset());
 
-		if (IsClass && Cast<UClass>(MenuItem.AssetData.GetAsset())->IsChildOf(AActor::StaticClass()))
+		if (IsClass && MenuItemClass && MenuItemClass->IsChildOf(AActor::StaticClass()))
 		{
 			AActor* DefaultActor = Cast<AActor>(Cast<UClass>(MenuItem.AssetData.GetAsset())->ClassDefaultObject);
 			OutActorTypeDisplayName = FText::FromString(FName::NameToDisplayString(DefaultActor->GetClass()->GetName(), false));
