Removed function ExtendToolsMenu() because it was a duplicated code that also exists in MainMenu.cpp.

diff --git a/Engine/Source/Editor/LevelEditor/Private/LevelEditorMenu.cpp b/Engine/Source/Editor/LevelEditor/Private/LevelEditorMenu.cpp
--- a/Engine/Source/Editor/LevelEditor/Private/LevelEditorMenu.cpp
+++ b/Engine/Source/Editor/LevelEditor/Private/LevelEditorMenu.cpp
@@ -278,42 +278,6 @@ void FLevelEditorMenu::RegisterLevelEditorMenus()
 			ActionsEntry.InsertPosition = FToolMenuInsert("Help", EToolMenuInsertType::Before);
 		}
 
-
-		static void ExtendToolsMenu()
-		{
-			UToolMenu* Menu = UToolMenus::Get()->RegisterMenu("LevelEditor.MainMenu.Tools", "MainFrame.MainMenu.Tools");
-
-			// Experimental section
-			{
-				// This is a temporary home for the spawners of experimental features that must be explicitly enabled.
-				// When the feature becomes permanent and need not check a flag, register a nomad spawner for it in the proper WorkspaceMenu category
-				const bool bTranslationPicker = GetDefault<UEditorExperimentalSettings>()->bEnableTranslationPicker;
-
-				// Make sure at least one is enabled before creating the section
-				if (bTranslationPicker)
-				{
-					FToolMenuSection& Section = Menu->AddSection("ExperimentalTabSpawners", LOCTEXT("ExperimentalTabSpawnersHeading", "Experimental"));
-					{
-						// Translation Picker
-						if (bTranslationPicker)
-						{
-							Section.AddMenuEntry(
-								"TranslationPicker",
-								LOCTEXT("TranslationPickerMenuItem", "Translation Picker"),
-								LOCTEXT("TranslationPickerMenuItemToolTip", "Launch the Translation Picker to Modify Editor Translations"),
-								FSlateIcon(),
-								FUIAction(FExecuteAction::CreateLambda(
-									[]()
-									{
-										FModuleManager::Get().LoadModuleChecked("TranslationEditor");
-										ITranslationEditor::OpenTranslationPicker();
-									}))
-							);
-						}
-					}
-				}
-			}
-		}
 	};
 
 	UToolMenus* ToolMenus = UToolMenus::Get();
@@ -336,9 +300,6 @@ void FLevelEditorMenu::RegisterLevelEditorMenus()
 	// Extend the Help menu
 	Local::ExtendHelpMenu();
 
-	// Extend the Tools menu
-	Local::ExtendToolsMenu();
-
 	// Extend the Build menu
 	RegisterBuildMenu();
 
diff --git a/Engine/Source/Editor/MainFrame/Private/Menus/MainMenu.cpp b/Engine/Source/Editor/MainFrame/Private/Menus/MainMenu.cpp
--- a/Engine/Source/Editor/MainFrame/Private/Menus/MainMenu.cpp
+++ b/Engine/Source/Editor/MainFrame/Private/Menus/MainMenu.cpp
@@ -505,26 +505,24 @@ void FMainMenu::RegisterToolsMenu()
 		const bool bTranslationPicker = GetDefault<UEditorExperimentalSettings>()->bEnableTranslationPicker;
 
 		// Make sure at least one is enabled before creating the section
-		if ( bTranslationPicker)
+
+		FToolMenuSection& ExperimentalSection = Menu->AddSection("ExperimentalTabSpawners", LOCTEXT("ExperimentalTabSpawnersHeading", "Experimental"));
 		{
-			FToolMenuSection& ExperimentalSection = Menu->AddSection("ExperimentalTabSpawners", LOCTEXT("ExperimentalTabSpawnersHeading", "Experimental"));
+			// Translation Picker
+			if (bTranslationPicker)
 			{
-				// Translation Picker
-				if (bTranslationPicker)
-				{
-					ExperimentalSection.AddMenuEntry(
-						"TranslationPicker",
-						LOCTEXT("TranslationPickerMenuItem", "Translation Picker"),
-						LOCTEXT("TranslationPickerMenuItemToolTip", "Launch the Translation Picker to Modify Editor Translations"),
-						FSlateIcon(),
-						FUIAction(FExecuteAction::CreateLambda(
-							[]()
-							{
-								FModuleManager::Get().LoadModuleChecked("TranslationEditor");
-								ITranslationEditor::OpenTranslationPicker();
-							}))
-					);
-				}
+				ExperimentalSection.AddMenuEntry(
+					"TranslationPicker",
+					LOCTEXT("TranslationPickerMenuItem", "Translation Picker"),
+					LOCTEXT("TranslationPickerMenuItemToolTip", "Launch the Translation Picker to Modify Editor Translations"),
+					FSlateIcon(),
+					FUIAction(FExecuteAction::CreateLambda(
+						[]()
+						{
+							FModuleManager::Get().LoadModuleChecked("TranslationEditor");
+							ITranslationEditor::OpenTranslationPicker();
+						}))
+				);
 			}
 		}
 	}
