PR #9244: fix: ImportScene typo (Contributed by [REDACTED_PERSON])

diff --git a/Engine/Source/Editor/LevelEditor/Private/LevelEditorActions.cpp b/Engine/Source/Editor/LevelEditor/Private/LevelEditorActions.cpp
--- a/Engine/Source/Editor/LevelEditor/Private/LevelEditorActions.cpp
+++ b/Engine/Source/Editor/LevelEditor/Private/LevelEditorActions.cpp
@@ -3372,7 +3372,7 @@ void FLevelEditorCommands::RegisterCommands()
 		OpenFavoriteFileCommands.Add(OpenFavoriteFile);
 	}
 
-	UI_COMMAND( ImportScene, "Import Into Level...", "Imports a scene from a FBX or T3D format into the current level", EUserInterfaceActionType::Button, FInputChord());
+	UI_COMMAND( ImportScene, "Import Into Level...", "Imports a scene from a FBX or OBJ format into the current level", EUserInterfaceActionType::Button, FInputChord());
 	UI_COMMAND( ExportAll, "Export All...", "Exports the entire level to a file on disk (multiple formats are supported.)", EUserInterfaceActionType::Button, FInputChord() );
 	UI_COMMAND( ExportSelected, "Export Selected...", "Exports currently-selected objects to a file on disk (multiple formats are supported.)", EUserInterfaceActionType::Button, FInputChord() );
 
