Add ensure in SLevelEditor::CanSpawnLevelEditorTab

#rnx

diff --git a/Engine/Source/Editor/LevelEditor/Private/SLevelEditor.cpp b/Engine/Source/Editor/LevelEditor/Private/SLevelEditor.cpp
--- a/Engine/Source/Editor/LevelEditor/Private/SLevelEditor.cpp
+++ b/Engine/Source/Editor/LevelEditor/Private/SLevelEditor.cpp
@@ -1084,7 +1084,7 @@ bool SLevelEditor::CanSpawnLevelEditorTab(const FSpawnTabArgs& Args, FName TabId
 	}
 
 	// HLOD Outliner not yet supported with World Partition
-	if (TabIdentifier == LevelEditorTabIds::LevelEditorHierarchicalLODOutliner && World->IsPartitionedWorld())
+	if (TabIdentifier == LevelEditorTabIds::LevelEditorHierarchicalLODOutliner && ensure(World) && World->IsPartitionedWorld())
 	{
 		return false;
 	}
