Create Camera Here now show the DisplayName instead of the top level asset path.
#rnx

diff --git a/Engine/Source/Editor/LevelEditor/Private/SLevelViewport.cpp b/Engine/Source/Editor/LevelEditor/Private/SLevelViewport.cpp
--- a/Engine/Source/Editor/LevelEditor/Private/SLevelViewport.cpp
+++ b/Engine/Source/Editor/LevelEditor/Private/SLevelViewport.cpp
@@ -1377,20 +1377,13 @@ void SLevelViewport::BindOptionCommands( FUICommandList& OutCommandList )
 		{
 			continue;
 		}
-		
-		// Remove _C from display names for blueprint classes
-		FString DisplayName = Name;
-		if (Name.EndsWith("_C"))
-		{
-			DisplayName = DisplayName.LeftChop(2);
-		}
 
 		UClass* CameraClass = FindObject<UClass>(ClassPathName);
-		if (!CameraClass || CameraClass->HasAllClassFlags(CLASS_Abstract))
+		if (!CameraClass || CameraClass->HasAnyClassFlags(CLASS_Abstract|CLASS_NotPlaceable|CLASS_HideDropDown|CLASS_Hidden))
 		{
 			continue;
 		}
-		
+
 		// Look for existing UI Command info so one isn't created for every viewport
 		FName CommandName;
 		{
@@ -1410,7 +1403,7 @@ void SLevelViewport::BindOptionCommands( FUICommandList& OutCommandList )
 		else
 		{
 			// If command info isn't found, create a new one
-			TSharedRef<FUICommandInfo> NewCamera = FUICommandInfoDecl(FLevelViewportCommands::Get().AsShared(), CommandName, FText::FromString(DisplayName), FText::Format(LOCTEXT("SpawnCamerasTooltip", "Spawn Camera here of type {0}"), FText::FromString(Name))).UserInterfaceType(EUserInterfaceActionType::Button).DefaultChord(FInputChord());
+			TSharedRef<FUICommandInfo> NewCamera = FUICommandInfoDecl(FLevelViewportCommands::Get().AsShared(), CommandName, CameraClass->GetDisplayNameText(), FText::Format(LOCTEXT("SpawnCamerasTooltip", "Spawn Camera here of type {0}"), FText::FromString(Name))).UserInterfaceType(EUserInterfaceActionType::Button).DefaultChord(FInputChord());
 			
 			OutCommandList.MapAction(
 				NewCamera,
