NewLevel output param to tell whether a level was created (user can cancel)

diff --git a/Engine/Source/Editor/LevelEditor/Private/LevelEditorActions.cpp b/Engine/Source/Editor/LevelEditor/Private/LevelEditorActions.cpp
--- a/Engine/Source/Editor/LevelEditor/Private/LevelEditorActions.cpp
+++ b/Engine/Source/Editor/LevelEditor/Private/LevelEditorActions.cpp
@@ -123,7 +123,7 @@ DEFINE_LOG_CATEGORY_STATIC(LevelEditorActions, Log, All);
 
 const FName HotReloadModule("HotReload");
 
-FSimpleDelegate FLevelEditorActionCallbacks::NewLevelOverride;
+FLevelEditorActionCallbacks::FNewLevelOverride FLevelEditorActionCallbacks::NewLevelOverride;
 
 namespace LevelEditorActionsHelpers
 {
@@ -198,8 +198,17 @@ void FLevelEditorActionCallbacks::BrowseViewportControls()
 
 void FLevelEditorActionCallbacks::NewLevel()
 {
-	if (NewLevelOverride.ExecuteIfBound())
+	bool bLevelCreated = false;
+	NewLevel(bLevelCreated);
+}
+
+void FLevelEditorActionCallbacks::NewLevel(bool& bOutLevelCreated)
+{
+	bOutLevelCreated = false;
+
+	if (NewLevelOverride.IsBound())
 	{
+		NewLevelOverride.Execute(bOutLevelCreated);
 		return;
 	}
 
@@ -216,45 +225,49 @@ void FLevelEditorActionCallbacks::NewLevel()
 
 	FNewLevelDialogModule& NewLevelDialogModule = FModuleManager::LoadModuleChecked<FNewLevelDialogModule>("NewLevelDialog");
 	
-	if (NewLevelDialogModule.CreateAndShowNewLevelDialog(MainFrameModule.GetParentWindow(), TemplateMapPackageName, bShowPartitionedTemplates, bOutIsPartitionedWorld))
+	if (!NewLevelDialogModule.CreateAndShowNewLevelDialog(MainFrameModule.GetParentWindow(), TemplateMapPackageName, bShowPartitionedTemplates, bOutIsPartitionedWorld))
 	{
-		// The new map screen will return a blank TemplateName if the user has selected to begin a new blank map
-		if (TemplateMapPackageName.IsEmpty())
-		{
-			GEditor->CreateNewMapForEditing(/*bPromptUserToSave=*/true, bOutIsPartitionedWorld);
-		}
-		else
-		{
-			// New map screen returned a non-empty TemplateName, so the user has selected to begin from a template map
-			bool TemplateFound = false;
+		return;
+	}
 
-			// Search all template map folders for a match with TemplateName
-			const bool bIncludeReadOnlyRoots = true;
-			if ( FPackageName::IsValidLongPackageName(TemplateMapPackageName, bIncludeReadOnlyRoots) )
+	// The new map screen will return a blank TemplateName if the user has selected to begin a new blank map
+	if (TemplateMapPackageName.IsEmpty())
+	{
+		GEditor->CreateNewMapForEditing(/*bPromptUserToSave=*/true, bOutIsPartitionedWorld);
+	}
+	else
+	{
+		// New map screen returned a non-empty TemplateName, so the user has selected to begin from a template map
+		bool TemplateFound = false;
+
+		// Search all template map folders for a match with TemplateName
+		const bool bIncludeReadOnlyRoots = true;
+		if ( FPackageName::IsValidLongPackageName(TemplateMapPackageName, bIncludeReadOnlyRoots) )
+		{
+			const FString MapPackageFilename = FPackageName::LongPackageNameToFilename(TemplateMapPackageName, FPackageName::GetMapPackageExtension());
+			if ( FPaths::FileExists(MapPackageFilename) )
 			{
-				const FString MapPackageFilename = FPackageName::LongPackageNameToFilename(TemplateMapPackageName, FPackageName::GetMapPackageExtension());
-				if ( FPaths::FileExists(MapPackageFilename) )
-				{
-					// File found because the size check came back non-zero
-					TemplateFound = true;
+				// File found because the size check came back non-zero
+				TemplateFound = true;
 
-					// If there are any unsaved changes to the current level, see if the user wants to save those first.
-					if ( FEditorFileUtils::SaveDirtyPackages(/*bPromptUserToSave*/true, /*bSaveMapPackages*/true, /*bSaveContentPackages*/false) )
-					{
-						// Load the template map file - passes LoadAsTemplate==true making the
-						// level load into an untitled package that won't save over the template
-						FEditorFileUtils::LoadMap(*MapPackageFilename, /*bLoadAsTemplate=*/true);
-					}
+				// If there are any unsaved changes to the current level, see if the user wants to save those first.
+				if ( FEditorFileUtils::SaveDirtyPackages(/*bPromptUserToSave*/true, /*bSaveMapPackages*/true, /*bSaveContentPackages*/false) )
+				{
+					// Load the template map file - passes LoadAsTemplate==true making the
+					// level load into an untitled package that won't save over the template
+					FEditorFileUtils::LoadMap(*MapPackageFilename, /*bLoadAsTemplate=*/true);
 				}
 			}
+		}
 
-			if (!TemplateFound)
-			{
-				UE_LOG( LevelEditorActions, Warning, TEXT("Couldn't find template map package %s"), *TemplateMapPackageName);
-				GEditor->CreateNewMapForEditing();
-			}
+		if (!TemplateFound)
+		{
+			UE_LOG( LevelEditorActions, Warning, TEXT("Couldn't find template map package %s"), *TemplateMapPackageName);
+			GEditor->CreateNewMapForEditing();
 		}
 	}
+
+	bOutLevelCreated = true;
 }
 
 bool FLevelEditorActionCallbacks::NewLevel_CanExecute()
diff --git a/Engine/Source/Editor/LevelEditor/Public/LevelEditorActions.h b/Engine/Source/Editor/LevelEditor/Public/LevelEditorActions.h
--- a/Engine/Source/Editor/LevelEditor/Public/LevelEditorActions.h
+++ b/Engine/Source/Editor/LevelEditor/Public/LevelEditorActions.h
@@ -645,7 +645,9 @@ public:
 
 	/** Creates a new level */
 	static void NewLevel();
-	static FSimpleDelegate NewLevelOverride;
+	static void NewLevel(bool& bOutLevelCreated);
+	DECLARE_DELEGATE_OneParam(FNewLevelOverride, bool& /*bOutLevelCreated*/);
+	static FNewLevelOverride NewLevelOverride;
 	static bool NewLevel_CanExecute();
 
 	/** Opens an existing level */
