Fixed case actor editor context not refreshing when comitting a level instance

diff --git a/Engine/Source/Editor/LevelEditor/Private/LevelEditorSubsystem.cpp b/Engine/Source/Editor/LevelEditor/Private/LevelEditorSubsystem.cpp
--- a/Engine/Source/Editor/LevelEditor/Private/LevelEditorSubsystem.cpp
+++ b/Engine/Source/Editor/LevelEditor/Private/LevelEditorSubsystem.cpp
@@ -64,6 +64,7 @@ void ULevelEditorSubsystem::Initialize(FSubsystemCollectionBase& Collection)
 	UActorEditorContextSubsystem::Get()->RegisterClient(this);
 	FWorldDelegates::LevelAddedToWorld.AddUObject(this, &ULevelEditorSubsystem::OnLevelAddedOrRemoved);
 	FWorldDelegates::LevelRemovedFromWorld.AddUObject(this, &ULevelEditorSubsystem::OnLevelAddedOrRemoved);
+	FWorldDelegates::OnCurrentLevelChanged.AddUObject(this, &ULevelEditorSubsystem::OnCurrentLevelChanged);
 }
 
 void ULevelEditorSubsystem::Deinitialize()
@@ -71,6 +72,7 @@ void ULevelEditorSubsystem::Deinitialize()
 	UActorEditorContextSubsystem::Get()->UnregisterClient(this);
 	FWorldDelegates::LevelAddedToWorld.RemoveAll(this);
 	FWorldDelegates::LevelRemovedFromWorld.RemoveAll(this);
+	FWorldDelegates::OnCurrentLevelChanged.RemoveAll(this);
 	UToolMenus::UnRegisterStartupCallback(this);
 	UToolMenus::UnregisterOwner(this);
 }
@@ -892,6 +894,11 @@ void ULevelEditorSubsystem::OnLevelAddedOrRemoved(ULevel* InLevel, UWorld* InWor
 	ActorEditorContextClientChanged.Broadcast(this);
 }
 
+void ULevelEditorSubsystem::OnCurrentLevelChanged(ULevel* InNewLevel, ULevel* InOldLevel, UWorld* InWorld)
+{
+	ActorEditorContextClientChanged.Broadcast(this);
+}
+
 TSharedRef<SWidget> ULevelEditorSubsystem::GetActorEditorContextWidget(UWorld* InWorld) const
 {
 	return SNew(SCurrentLevelWidget).World(InWorld);
diff --git a/Engine/Source/Editor/LevelEditor/Public/LevelEditorSubsystem.h b/Engine/Source/Editor/LevelEditor/Public/LevelEditorSubsystem.h
--- a/Engine/Source/Editor/LevelEditor/Public/LevelEditorSubsystem.h
+++ b/Engine/Source/Editor/LevelEditor/Public/LevelEditorSubsystem.h
@@ -160,6 +160,9 @@ private:
 	/** Called when a Level is added to a world or remove from a world */
 	void OnLevelAddedOrRemoved(ULevel* InLevel, UWorld* InWorld);
 
+	/** Called when the current level changes on a world */
+	void OnCurrentLevelChanged(ULevel* InNewLevel, ULevel* InOldLevel, UWorld* InWorld);
+
 	/** Delegate used to notify changes to ActorEditorContextSubsystem */
 	FOnActorEditorContextClientChanged ActorEditorContextClientChanged;
 };
\ No newline at end of file
