Fixing an issue where level instances nested inside a locked sub-level would permit selections that would assert later.

diff --git a/Engine/Source/Editor/LevelEditor/Private/Elements/Actor/ActorElementLevelEditorSelectionCustomization.cpp b/Engine/Source/Editor/LevelEditor/Private/Elements/Actor/ActorElementLevelEditorSelectionCustomization.cpp
--- a/Engine/Source/Editor/LevelEditor/Private/Elements/Actor/ActorElementLevelEditorSelectionCustomization.cpp
+++ b/Engine/Source/Editor/LevelEditor/Private/Elements/Actor/ActorElementLevelEditorSelectionCustomization.cpp
@@ -106,7 +106,16 @@ bool FActorElementLevelEditorSelectionCustomization::CanSelectActorElement(const
 		return false;
 	}
 
-	if (!Actor->IsTemplate() && FLevelUtils::IsLevelLocked(Actor->GetLevel()))
+	// When trying to select a particular actor with a 'selection parent' we'll end 
+	// up selecting the actor's selection root instead (see `SelectActorElement()`), 
+	// Because of this, we need to also check how selectable the selection root is as well.
+	// In the case of Level Instances, an actor in that instance does not belong to 
+	// the same level as the instance... so we should be checking the selection root's level
+	// instead of the actor's direct level.
+	AActor* SelectionRoot = Actor->GetRootSelectionParent();
+	ULevel* SelectionLevel = (SelectionRoot != nullptr) ? SelectionRoot->GetLevel() : Actor->GetLevel();
+
+	if (!Actor->IsTemplate() && FLevelUtils::IsLevelLocked(SelectionLevel))
 	{
 		UE_CLOG(InSelectionOptions.WarnIfLocked(), LogActorLevelEditorSelection, Warning, TEXT("SelectActor: %s (%s)"), TEXT("The requested operation could not be completed because the level is locked."), *Actor->GetActorLabel());
 		return false;
