ViewportToolBar: Added LeftExtension point for the LevelViewportToolBar.
Added ViewportToolBarContext to permit use of FEditorViewportToolBarMenu widgets in LevelViewportToolBar panel extensions.

diff --git a/Engine/Source/Editor/LevelEditor/Private/SLevelViewportToolBar.cpp b/Engine/Source/Editor/LevelEditor/Private/SLevelViewportToolBar.cpp
--- a/Engine/Source/Editor/LevelEditor/Private/SLevelViewportToolBar.cpp
+++ b/Engine/Source/Editor/LevelEditor/Private/SLevelViewportToolBar.cpp
@@ -8,6 +8,7 @@
 #include "Subsystems/PanelExtensionSubsystem.h"
 #include "ToolMenus.h"
 #include "LevelEditorMenuContext.h"
+#include "ViewportToolBarContext.h"
 #include "Modules/ModuleManager.h"
 #include "Widgets/Layout/SBorder.h"
 #include "Widgets/Layout/SBox.h"
@@ -172,6 +173,9 @@ void SLevelViewportToolBar::Construct( const FArguments& InArgs )
 
 	FLevelEditorModule& LevelEditorModule = FModuleManager::GetModuleChecked<FLevelEditorModule>(TEXT("LevelEditor"));
 
+	UViewportToolBarContext* ExtensionContextObject = NewObject<UViewportToolBarContext>();
+	ExtensionContextObject->ViewportToolBar = SharedThis(this);
+
 	const FMargin ToolbarSlotPadding(4.0f, 1.0f);
 
 	ChildSlot
@@ -249,6 +253,15 @@ void SLevelViewportToolBar::Construct( const FArguments& InArgs )
 					.Visibility(EVisibility::Collapsed)
 				]
 				+ SHorizontalBox::Slot()
+				.Padding(ToolbarSlotPadding)
+				.HAlign(HAlign_Left)
+				.VAlign(VAlign_Fill)
+				[
+					SAssignNew(ExtensionPanel, SExtensionPanel)
+					.ExtensionPanelID("LevelViewportToolBar.LeftExtension")
+					.ExtensionContext(ExtensionContextObject)
+				]
+				+ SHorizontalBox::Slot()
 				.AutoWidth()
 				.Padding(ToolbarSlotPadding)
 				[
@@ -286,6 +299,7 @@ void SLevelViewportToolBar::Construct( const FArguments& InArgs )
 				[
 					SAssignNew(ExtensionPanel, SExtensionPanel)
 					.ExtensionPanelID("LevelViewportToolBar.MiddleExtension")
+					.ExtensionContext(ExtensionContextObject)
 				]
 				+ SHorizontalBox::Slot()
 				.Padding(ToolbarSlotPadding)
diff --git a/Engine/Source/Editor/UnrealEd/Public/ViewportToolBarContext.h b/Engine/Source/Editor/UnrealEd/Public/ViewportToolBarContext.h
new file mode 100644
--- /dev/null
+++ b/Engine/Source/Editor/UnrealEd/Public/ViewportToolBarContext.h
@@ -0,0 +1,16 @@
+// Copyright Epic Games, Inc. All Rights Reserved.
+
+#pragma once
+
+#include "CoreMinimal.h"
+#include "ViewportToolBarContext.generated.h"
+
+class SViewportToolBar;
+
+UCLASS()
+class UNREALED_API UViewportToolBarContext : public UObject
+{
+	GENERATED_BODY()
+public:
+	TWeakPtr<SViewportToolBar> ViewportToolBar;
+};
\ No newline at end of file
