Add Naming logic to BuildSingleAssetAddReplaceActorMenu to allow ToolsMenu filtering operations on AddActor Entries

diff --git a/Engine/Source/Editor/LevelEditor/Private/LevelEditorCreateActorMenu.cpp b/Engine/Source/Editor/LevelEditor/Private/LevelEditorCreateActorMenu.cpp
--- a/Engine/Source/Editor/LevelEditor/Private/LevelEditorCreateActorMenu.cpp
+++ b/Engine/Source/Editor/LevelEditor/Private/LevelEditorCreateActorMenu.cpp
@@ -368,7 +368,9 @@ static void BuildSingleAssetAddReplaceActorMenu(FToolMenuSection& Section, const
 		// For a single option that doesn't open a submenu, we have an option of the custom tile widget (used for recents, selection) and a regular menu entry.
 		if (bUseAssetTile)
 		{
-			FToolMenuEntry Entry = FToolMenuEntry::InitMenuEntry(NAME_None, Action, SNew(SAssetMenuEntry, Asset, AssetMenuOptions).LabelOverride(LabelOverride));
+			FString EntryName = LabelOverride.BuildSourceString();
+			EntryName.RemoveSpacesInline();
+			FToolMenuEntry Entry = FToolMenuEntry::InitMenuEntry(*EntryName, Action, SNew(SAssetMenuEntry, Asset, AssetMenuOptions).LabelOverride(LabelOverride));
 			Section.AddEntry(Entry);
 		}
 		else
@@ -384,8 +386,10 @@ static void BuildSingleAssetAddReplaceActorMenu(FToolMenuSection& Section, const
 				AssetDisplayName = LabelOverride;
 			}
 
+			FString EntryName = AssetDisplayName.BuildSourceString();
+			EntryName.RemoveSpacesInline();
 			FSlateIcon Icon = FSlateIconFinder::FindIconForClass(FClassIconFinder::GetIconClassForAssetData(Asset));
-			Section.AddMenuEntry(NAME_None, AssetDisplayName, TAttribute<FText>(), Icon, Action, EUserInterfaceActionType::Button);
+			Section.AddMenuEntry(*EntryName, AssetDisplayName, TAttribute<FText>(), Icon, Action, EUserInterfaceActionType::Button);
 		}
 	}
 	else
@@ -402,9 +406,11 @@ static void BuildSingleAssetAddReplaceActorMenu(FToolMenuSection& Section, const
 			AssetDisplayName = LabelOverride;
 		}
 
+		FString SubMenuName = AssetDisplayName.BuildSourceString();
+		SubMenuName.RemoveSpacesInline();
 		FSlateIcon Icon = FSlateIconFinder::FindIconForClass(FClassIconFinder::GetIconClassForAssetData(Asset));
 		Section.AddSubMenu(
-			NAME_None, AssetDisplayName, TAttribute<FText>(), FNewToolMenuDelegate::CreateStatic(&FillAssetAddReplaceActorMenu, Asset, AssetMenuOptions, CreateMode),
+			*SubMenuName, AssetDisplayName, TAttribute<FText>(), FNewToolMenuDelegate::CreateStatic(&FillAssetAddReplaceActorMenu, Asset, AssetMenuOptions, CreateMode),
 			FToolUIActionChoice(), EUserInterfaceActionType::Button, /*bInOpenSubMenuOnClick*/ false, Icon);
 	}
 }
