[Backout] - CL21114894
Original CL Desc
-----------------------------------------------------------------
Move ObjectMixer to Engine Module
Publish ObjectMixer

[snapshot CL20876622 - 7/9/2022 5:12:01 PM]

diff --git a/Engine/Source/Editor/LevelEditor/LevelEditor.Build.cs b/Engine/Source/Editor/LevelEditor/LevelEditor.Build.cs
--- a/Engine/Source/Editor/LevelEditor/LevelEditor.Build.cs
+++ b/Engine/Source/Editor/LevelEditor/LevelEditor.Build.cs
@@ -61,7 +61,6 @@ public class LevelEditor : ModuleRules
 				"TypedElementFramework",
 				"TypedElementRuntime",
 				"EngineSettings",
-				"ObjectMixerEditor",
 				"PropertyEditor",
 				"Kismet",
 				"KismetWidgets",
@@ -108,7 +107,8 @@ public class LevelEditor : ModuleRules
 				"VREditor",
                 "Persona",
 				"LevelAssetEditor",
-				"MergeActors"
+				"MergeActors",
+
 			}
 		);
 
diff --git a/Engine/Source/Editor/LevelEditor/Private/SLevelEditor.cpp b/Engine/Source/Editor/LevelEditor/Private/SLevelEditor.cpp
--- a/Engine/Source/Editor/LevelEditor/Private/SLevelEditor.cpp
+++ b/Engine/Source/Editor/LevelEditor/Private/SLevelEditor.cpp
@@ -72,7 +72,6 @@
 #include "EditorViewportLayout.h"
 #include "LevelViewportTabContent.h"
 #include "SLevelViewport.h"
-#include "ObjectMixerEditorModule.h"
 
 #define LOCTEXT_NAMESPACE "SLevelEditor"
 
@@ -1482,11 +1481,6 @@ TSharedRef<SWidget> SLevelEditor::RestoreContentArea( const TSharedRef<SDockTab>
 				.SetIcon(EnvironmentLightingViewerIcon);
 		}
 
-		{
-			// The object mixer module will take care of its tab spawner once it's loaded.
-			FModuleManager::Get().LoadModuleChecked<FObjectMixerEditorModule>( "ObjectMixerEditor" );
-		}
-
 		FTabSpawnerEntry& BuildAndSubmitEntry = LevelEditorTabManager->RegisterTabSpawner(LevelEditorTabIds::LevelEditorBuildAndSubmit, FOnSpawnTab::CreateSP<SLevelEditor, FName, FString>(this, &SLevelEditor::SpawnLevelEditorTab, LevelEditorTabIds::LevelEditorBuildAndSubmit, FString()));
 		BuildAndSubmitEntry.SetAutoGenerateMenuEntry(false);
 
@@ -2178,4 +2172,4 @@ FVector2D SLevelEditor::GetActiveViewportSize()
 	return FVector2D(MaxX, MaxY);
 }
 
-#undef LOCTEXT_NAMESPACE
+#undef LOCTEXT_NAMESPACE
\ No newline at end of file
diff --git a/Engine/Source/Editor/LightMixer/LightMixer.build.cs b/Engine/Source/Editor/LightMixer/LightMixer.build.cs
deleted file mode 100644
diff --git a/Engine/Source/Editor/LightMixer/Private/LightMixerModule.cpp b/Engine/Source/Editor/LightMixer/Private/LightMixerModule.cpp
deleted file mode 100644
diff --git a/Engine/Source/Editor/LightMixer/Private/LightMixerStyle.cpp b/Engine/Source/Editor/LightMixer/Private/LightMixerStyle.cpp
deleted file mode 100644
diff --git a/Engine/Source/Editor/LightMixer/Private/LightMixerStyle.h b/Engine/Source/Editor/LightMixer/Private/LightMixerStyle.h
deleted file mode 100644
diff --git a/Engine/Source/Editor/LightMixer/Public/LightMixerModule.h b/Engine/Source/Editor/LightMixer/Public/LightMixerModule.h
deleted file mode 100644
diff --git a/Engine/Source/Editor/LightMixer/Public/LightMixerObjectFilter.h b/Engine/Source/Editor/LightMixer/Public/LightMixerObjectFilter.h
deleted file mode 100644
diff --git a/Engine/Source/Editor/LightMixer/Public/LightMixerProjectSettings.h b/Engine/Source/Editor/LightMixer/Public/LightMixerProjectSettings.h
deleted file mode 100644
diff --git a/Engine/Source/Editor/ObjectMixer/ObjectMixerEditor.build.cs b/Engine/Source/Editor/ObjectMixer/ObjectMixerEditor.build.cs
deleted file mode 100644
diff --git a/Engine/Source/Editor/ObjectMixer/Private/ObjectFilter/ObjectMixerEditorObjectFilter.cpp b/Engine/Source/Editor/ObjectMixer/Private/ObjectFilter/ObjectMixerEditorObjectFilter.cpp
deleted file mode 100644
diff --git a/Engine/Source/Editor/ObjectMixer/Private/ObjectMixerEditorModule.cpp b/Engine/Source/Editor/ObjectMixer/Private/ObjectMixerEditorModule.cpp
deleted file mode 100644
diff --git a/Engine/Source/Editor/ObjectMixer/Private/ObjectMixerEditorStyle.cpp b/Engine/Source/Editor/ObjectMixer/Private/ObjectMixerEditorStyle.cpp
deleted file mode 100644
diff --git a/Engine/Source/Editor/ObjectMixer/Private/ObjectMixerFilterFactory.cpp b/Engine/Source/Editor/ObjectMixer/Private/ObjectMixerFilterFactory.cpp
deleted file mode 100644
diff --git a/Engine/Source/Editor/ObjectMixer/Private/Views/List/ObjectMixerEditorList.cpp b/Engine/Source/Editor/ObjectMixer/Private/Views/List/ObjectMixerEditorList.cpp
deleted file mode 100644
diff --git a/Engine/Source/Editor/ObjectMixer/Private/Views/List/ObjectMixerEditorListRow.cpp b/Engine/Source/Editor/ObjectMixer/Private/Views/List/ObjectMixerEditorListRow.cpp
deleted file mode 100644
diff --git a/Engine/Source/Editor/ObjectMixer/Private/Views/List/SObjectMixerEditorList.cpp b/Engine/Source/Editor/ObjectMixer/Private/Views/List/SObjectMixerEditorList.cpp
deleted file mode 100644
diff --git a/Engine/Source/Editor/ObjectMixer/Private/Views/List/SObjectMixerEditorListRow.cpp b/Engine/Source/Editor/ObjectMixer/Private/Views/List/SObjectMixerEditorListRow.cpp
deleted file mode 100644
diff --git a/Engine/Source/Editor/ObjectMixer/Private/Views/MainPanel/ObjectMixerEditorMainPanel.cpp b/Engine/Source/Editor/ObjectMixer/Private/Views/MainPanel/ObjectMixerEditorMainPanel.cpp
deleted file mode 100644
diff --git a/Engine/Source/Editor/ObjectMixer/Private/Views/MainPanel/SObjectMixerEditorMainPanel.cpp b/Engine/Source/Editor/ObjectMixer/Private/Views/MainPanel/SObjectMixerEditorMainPanel.cpp
deleted file mode 100644
diff --git a/Engine/Source/Editor/ObjectMixer/Public/ObjectFilter/ObjectMixerEditorObjectFilter.h b/Engine/Source/Editor/ObjectMixer/Public/ObjectFilter/ObjectMixerEditorObjectFilter.h
deleted file mode 100644
diff --git a/Engine/Source/Editor/ObjectMixer/Public/ObjectMixerEditorLog.h b/Engine/Source/Editor/ObjectMixer/Public/ObjectMixerEditorLog.h
deleted file mode 100644
diff --git a/Engine/Source/Editor/ObjectMixer/Public/ObjectMixerEditorModule.h b/Engine/Source/Editor/ObjectMixer/Public/ObjectMixerEditorModule.h
deleted file mode 100644
diff --git a/Engine/Source/Editor/ObjectMixer/Public/ObjectMixerEditorProjectSettings.h b/Engine/Source/Editor/ObjectMixer/Public/ObjectMixerEditorProjectSettings.h
deleted file mode 100644
diff --git a/Engine/Source/Editor/ObjectMixer/Public/ObjectMixerEditorStyle.h b/Engine/Source/Editor/ObjectMixer/Public/ObjectMixerEditorStyle.h
deleted file mode 100644
diff --git a/Engine/Source/Editor/ObjectMixer/Public/ObjectMixerFilterFactory.h b/Engine/Source/Editor/ObjectMixer/Public/ObjectMixerFilterFactory.h
deleted file mode 100644
diff --git a/Engine/Source/Editor/ObjectMixer/Public/Views/List/ObjectMixerEditorList.h b/Engine/Source/Editor/ObjectMixer/Public/Views/List/ObjectMixerEditorList.h
deleted file mode 100644
diff --git a/Engine/Source/Editor/ObjectMixer/Public/Views/List/ObjectMixerEditorListFilters/IObjectMixerEditorListFilter.h b/Engine/Source/Editor/ObjectMixer/Public/Views/List/ObjectMixerEditorListFilters/IObjectMixerEditorListFilter.h
deleted file mode 100644
diff --git a/Engine/Source/Editor/ObjectMixer/Public/Views/List/ObjectMixerEditorListFilters/ObjectMixerEditorListFilter_Source.h b/Engine/Source/Editor/ObjectMixer/Public/Views/List/ObjectMixerEditorListFilters/ObjectMixerEditorListFilter_Source.h
deleted file mode 100644
diff --git a/Engine/Source/Editor/ObjectMixer/Public/Views/List/ObjectMixerEditorListRow.h b/Engine/Source/Editor/ObjectMixer/Public/Views/List/ObjectMixerEditorListRow.h
deleted file mode 100644
diff --git a/Engine/Source/Editor/ObjectMixer/Public/Views/List/SObjectMixerEditorList.h b/Engine/Source/Editor/ObjectMixer/Public/Views/List/SObjectMixerEditorList.h
deleted file mode 100644
diff --git a/Engine/Source/Editor/ObjectMixer/Public/Views/List/SObjectMixerEditorListRow.h b/Engine/Source/Editor/ObjectMixer/Public/Views/List/SObjectMixerEditorListRow.h
deleted file mode 100644
diff --git a/Engine/Source/Editor/ObjectMixer/Public/Views/MainPanel/ObjectMixerEditorMainPanel.h b/Engine/Source/Editor/ObjectMixer/Public/Views/MainPanel/ObjectMixerEditorMainPanel.h
deleted file mode 100644
diff --git a/Engine/Source/Editor/ObjectMixer/Public/Views/MainPanel/SObjectMixerEditorMainPanel.h b/Engine/Source/Editor/ObjectMixer/Public/Views/MainPanel/SObjectMixerEditorMainPanel.h
deleted file mode 100644
diff --git a/Engine/Source/Editor/ObjectMixer/Public/Views/Widgets/SObjectMixerEditorTooltipWidget.h b/Engine/Source/Editor/ObjectMixer/Public/Views/Widgets/SObjectMixerEditorTooltipWidget.h
deleted file mode 100644
