WorldPartition: Improved UI when No regions are loaded

diff --git a/Engine/Source/Editor/LevelEditor/Private/SLevelViewport.cpp b/Engine/Source/Editor/LevelEditor/Private/SLevelViewport.cpp
--- a/Engine/Source/Editor/LevelEditor/Private/SLevelViewport.cpp
+++ b/Engine/Source/Editor/LevelEditor/Private/SLevelViewport.cpp
@@ -92,6 +92,7 @@
 #include "SInViewportDetails.h"
 #include "Viewports/InViewportUIDragOperation.h"
 #include "SActorEditorContext.h"
+#include "SWorldPartitionViewportWidget.h"
 #include "LevelViewportLayout.h"
 #include "EditorViewportTabContent.h"
 
@@ -286,79 +287,110 @@ void SLevelViewport::ConstructViewportOverlayContent()
 		SAssignNew( ActorPreviewHorizontalBox, SHorizontalBox )
 	];
 
+	auto GetCombinedVisibility = [this]()
+	{
+		if (GetCurrentScreenPercentageVisibility() == EVisibility::Collapsed &&
+			GetCurrentFeatureLevelPreviewTextVisibility() == EVisibility::Collapsed &&
+			GetSelectedActorsCurrentLevelTextVisibility() == EVisibility::Collapsed &&
+			!IsActorEditorContextVisible())
+		{
+			return EVisibility::Collapsed;
+		}
+
+		return EVisibility::Visible;
+	};
+
 	ViewportOverlay->AddSlot( SlotIndex )
 	.VAlign( VAlign_Bottom )
 	.HAlign( HAlign_Right )
 	.Padding( 5.0f )
 	[
-		SNew(SBorder)
-		.BorderImage(FAppStyle::Get().GetBrush("FloatingBorder"))
+		SNew(SVerticalBox)
+		+ SVerticalBox::Slot()
+		.AutoHeight()
 		[
-			SNew(SVerticalBox)
-			+ SVerticalBox::Slot()
-			.AutoHeight()
-			.Padding(2.0f, 1.0f, 2.0f, 1.0f)
+			SNew(SBorder)
+			.BorderImage(FAppStyle::Get().GetBrush("FloatingBorder"))
 			[
-				SNew(SHorizontalBox)
-				.Visibility(this, &SLevelViewport::GetCurrentScreenPercentageVisibility)
-				// Current screen percentage label
-				+ SHorizontalBox::Slot()
-				.AutoWidth()
+				SNew(SVerticalBox)
+				.Visibility_Lambda(GetCombinedVisibility)
+				+ SVerticalBox::Slot()
+				.AutoHeight()
 				.Padding(2.0f, 1.0f, 2.0f, 1.0f)
 				[
-					SNew(STextBlock)
-					.Text(LOCTEXT("ScreenPercentageLabel", "Screen Percentage"))
-					.ShadowOffset(FVector2D(1, 1))
-				]
+					SNew(SHorizontalBox)
+					.Visibility(this, &SLevelViewport::GetCurrentScreenPercentageVisibility)
+					// Current screen percentage label
+					+ SHorizontalBox::Slot()
+					.AutoWidth()
+					.Padding(2.0f, 1.0f, 2.0f, 1.0f)
+					[
+						SNew(STextBlock)
+						.Text(LOCTEXT("ScreenPercentageLabel", "Screen Percentage"))
+						.ShadowOffset(FVector2D(1, 1))
+					]
 
-				// Current screen percentage
-				+ SHorizontalBox::Slot()
-				.AutoWidth()
-				.Padding(4.0f, 1.0f, 2.0f, 1.0f)
+					// Current screen percentage
+					+ SHorizontalBox::Slot()
+					.AutoWidth()
+					.Padding(4.0f, 1.0f, 2.0f, 1.0f)
+					[
+						SNew(STextBlock)
+						.Text(this, &SLevelViewport::GetCurrentScreenPercentageText)
+						.ShadowOffset(FVector2D(1, 1))
+					]
+				]
+				// add feature level widget
+				+ SVerticalBox::Slot()
+				.AutoHeight()
+				.Padding(2.0f, 1.0f, 2.0f, 1.0f)
 				[
-					SNew(STextBlock)
-					.Text(this, &SLevelViewport::GetCurrentScreenPercentageText)
-					.ShadowOffset(FVector2D(1, 1))
+					BuildFeatureLevelWidget()
 				]
-			]
-			// add feature level widget
-			+ SVerticalBox::Slot()
-			.AutoHeight()
-			.Padding(2.0f, 1.0f, 2.0f, 1.0f)
-			[
-				BuildFeatureLevelWidget()
-			]
-			+ SVerticalBox::Slot()
-			.AutoHeight()
-			.Padding(2.0f, 1.0f, 2.0f, 1.0f)
-			[
-				SNew(SVerticalBox)
-				.Visibility(this, &SLevelViewport::GetSelectedActorsCurrentLevelTextVisibility)
-				// Current level label
 				+ SVerticalBox::Slot()
 				.AutoHeight()
 				.Padding(2.0f, 1.0f, 2.0f, 1.0f)
 				[
-					SNew(STextBlock)
-					.Text(this, &SLevelViewport::GetSelectedActorsCurrentLevelText, true)
-					.ShadowOffset(FVector2D(1, 1))
+					SNew(SVerticalBox)
+					.Visibility(this, &SLevelViewport::GetSelectedActorsCurrentLevelTextVisibility)
+					// Current level label
+					+ SVerticalBox::Slot()
+					.AutoHeight()
+					.Padding(2.0f, 1.0f, 2.0f, 1.0f)
+					[
+						SNew(STextBlock)
+						.Text(this, &SLevelViewport::GetSelectedActorsCurrentLevelText, true)
+						.ShadowOffset(FVector2D(1, 1))
+					]
+					// Current level
+					+ SVerticalBox::Slot()
+					.AutoHeight()
+					.Padding(4.0f, 1.0f, 2.0f, 1.0f)
+					[
+						SNew(STextBlock)
+						.Text(this, &SLevelViewport::GetSelectedActorsCurrentLevelText, false)
+						.ShadowOffset(FVector2D(1, 1))
+					]
 				]
-				// Current level
 				+ SVerticalBox::Slot()
 				.AutoHeight()
-				.Padding(4.0f, 1.0f, 2.0f, 1.0f)
 				[
-					SNew(STextBlock)
-					.Text(this, &SLevelViewport::GetSelectedActorsCurrentLevelText, false)
-					.ShadowOffset(FVector2D(1, 1))
+					SNew(SActorEditorContext)
+					.World(GetWorld())
+					.Visibility_Lambda([this]() { return IsActorEditorContextVisible() ? OnGetViewportContentVisibility() : EVisibility::Collapsed; })
 				]
 			]
-			+ SVerticalBox::Slot()
-			.AutoHeight()
+		]
+		+ SVerticalBox::Slot()
+		.AutoHeight()
+		.Padding(0.0f, 4.0f, 0.0f, 0.0f)
+		[
+			SNew(SBorder)
+			.BorderImage(FAppStyle::Get().GetBrush("FloatingBorder"))
 			[
-				SNew(SActorEditorContext)
-				.World(GetWorld())
-				.Visibility_Lambda([this]() { return IsActorEditorContextVisible() ? OnGetViewportContentVisibility() : EVisibility::Collapsed; })
+				SNew(SWorldPartitionViewportWidget)
+				.Clickable(true)
+				.Visibility_Lambda([this]() { return SWorldPartitionViewportWidget::GetVisibility(GetWorld());})
 			]
 		]
 	];
@@ -380,7 +412,8 @@ bool SLevelViewport::IsActorEditorContextVisible() const
 		GetWorld()->GetCurrentLevel() &&
 		(&GetLevelViewportClient() == GCurrentLevelEditingViewportClient) &&
 		ActiveViewport.IsValid() &&
-		(ActiveViewport->GetPlayInEditorIsSimulate() || !ActiveViewport->GetClient()->GetWorld()->IsGameWorld());
+		(ActiveViewport->GetPlayInEditorIsSimulate() || !ActiveViewport->GetClient()->GetWorld()->IsGameWorld()) &&
+		SActorEditorContext::IsVisible(GetWorld());
 }
 
 void SLevelViewport::ConstructLevelEditorViewportClient(FLevelEditorViewportInstanceSettings& ViewportInstanceSettings)
diff --git a/Engine/Source/Editor/UnrealEd/Private/SActorEditorContext.cpp b/Engine/Source/Editor/UnrealEd/Private/SActorEditorContext.cpp
--- a/Engine/Source/Editor/UnrealEd/Private/SActorEditorContext.cpp
+++ b/Engine/Source/Editor/UnrealEd/Private/SActorEditorContext.cpp
@@ -198,4 +198,10 @@ void SActorEditorContext::Rebuild()
 	}
 }
 
+bool SActorEditorContext::IsVisible(UWorld* InWorld)
+{
+	TArray<IActorEditorContextClient*> Clients = UActorEditorContextSubsystem::Get()->GetDisplayableClients();
+	return InWorld && Clients.Num() > 0;
+}
+
 #undef LOCTEXT_NAMESPACE
\ No newline at end of file
diff --git a/Engine/Source/Editor/UnrealEd/Private/SWorldPartitionViewportWidget.cpp b/Engine/Source/Editor/UnrealEd/Private/SWorldPartitionViewportWidget.cpp
new file mode 100644
--- /dev/null
+++ b/Engine/Source/Editor/UnrealEd/Private/SWorldPartitionViewportWidget.cpp
@@ -0,0 +1,86 @@
+// Copyright Epic Games, Inc. All Rights Reserved.
+
+#include "SWorldPartitionViewportWidget.h"
+#include "Widgets/Input/SButton.h"
+#include "Widgets/Images/SImage.h"
+#include "Widgets/Text/STextBlock.h"
+#include "Widgets/SBoxPanel.h"
+#include "Framework/Docking/TabManager.h"
+#include "LevelEditor.h"
+#include "Editor.h"
+#include "Engine/World.h"
+#include "WorldPartition/WorldPartition.h"
+#include "Styling/AppStyle.h"
+
+#define LOCTEXT_NAMESPACE "WorldPartitionViewportWidget"
+
+void SWorldPartitionViewportWidget::Construct(const FArguments& InArgs)
+{
+	bClickable = InArgs._Clickable;
+	
+	ChildSlot.Padding(2, 2, 2, 2)
+	[
+		SNew(SButton)
+		.ButtonStyle(bClickable? &FCoreStyle::Get().GetWidgetStyle<FButtonStyle>("HoverHintOnly") : &FCoreStyle::Get().GetWidgetStyle<FButtonStyle>("NoBorder"))
+		.Cursor(EMouseCursor::Default)
+		.IsEnabled_Lambda([this]() { return bClickable; })
+		.OnClicked_Lambda([]
+		{
+			FLevelEditorModule& LevelEditorModule = FModuleManager::GetModuleChecked<FLevelEditorModule>(TEXT("LevelEditor"));
+			TSharedPtr<FTabManager> LevelEditorTabManager = LevelEditorModule.GetLevelEditorTabManager();
+			LevelEditorTabManager->TryInvokeTab(FName("WorldBrowserPartitionEditor"));
+			return FReply::Handled();
+		})
+		.ToolTipText(LOCTEXT("NoRegionsLoadedTooltip", "To load a region, drag select an area in the World Partition map and choose 'Load Region From Selection' from the context menu."))
+		.ContentPadding(0.f)
+		[
+			SNew(SVerticalBox)
+			+SVerticalBox::Slot()
+			.AutoHeight()
+			[
+				SNew(SHorizontalBox)
+				+ SHorizontalBox::Slot()
+				.Padding(FMargin(2, 1, 0, 1))
+				.AutoWidth()
+				.VAlign(VAlign_Center)
+				[
+					SNew(SImage)
+					.Image(FAppStyle::Get().GetBrush("Icons.WarningWithColor"))
+					.DesiredSizeOverride(FVector2D(16, 16))
+				]
+				+ SHorizontalBox::Slot()
+				.Padding(FMargin(6, 1, 2, 1))
+				.VAlign(VAlign_Center)
+				.HAlign(HAlign_Left)
+				.FillWidth(1)
+				[
+					SNew(STextBlock)
+					.Text(LOCTEXT("NoRegionsLoadedText","No regions loaded"))
+					.ColorAndOpacity(FAppStyle::Get().GetSlateColor("Colors.Warning"))
+				]
+			]
+		]
+	];
+
+
+}
+
+SWorldPartitionViewportWidget::~SWorldPartitionViewportWidget()
+{
+}
+
+EVisibility SWorldPartitionViewportWidget::GetVisibility(UWorld* InWorld)
+{
+	if (InWorld && InWorld->IsPartitionedWorld())
+	{
+		UWorldPartition* WorldPartition = InWorld->GetWorldPartition();
+		if (WorldPartition->IsStreamingEnabled() && !WorldPartition->HasLoadedUserCreatedRegions())
+		{
+			return EVisibility::Visible;
+		}
+	}
+
+	return EVisibility::Collapsed;
+}
+
+#undef LOCTEXT_NAMESPACE
\ No newline at end of file
diff --git a/Engine/Source/Editor/UnrealEd/Public/SActorEditorContext.h b/Engine/Source/Editor/UnrealEd/Public/SActorEditorContext.h
--- a/Engine/Source/Editor/UnrealEd/Public/SActorEditorContext.h
+++ b/Engine/Source/Editor/UnrealEd/Public/SActorEditorContext.h
@@ -14,6 +14,7 @@ public:
 	void Construct(const FArguments& InArgs);
 	~SActorEditorContext();
 
+	static bool IsVisible(UWorld* InWorld);
 private:
 
 	void OnEditorMapChange(uint32 MapChangeFlags = 0) { Rebuild(); }
diff --git a/Engine/Source/Editor/UnrealEd/Public/SWorldPartitionViewportWidget.h b/Engine/Source/Editor/UnrealEd/Public/SWorldPartitionViewportWidget.h
new file mode 100644
--- /dev/null
+++ b/Engine/Source/Editor/UnrealEd/Public/SWorldPartitionViewportWidget.h
@@ -0,0 +1,21 @@
+// Copyright Epic Games, Inc. All Rights Reserved.
+
+#pragma once
+
+#include "Widgets/SCompoundWidget.h"
+
+class UNREALED_API SWorldPartitionViewportWidget : public SCompoundWidget
+{
+public:
+	SLATE_BEGIN_ARGS(SWorldPartitionViewportWidget) {}
+	SLATE_ARGUMENT(bool, Clickable)
+	SLATE_END_ARGS()
+
+	void Construct(const FArguments& InArgs);
+	~SWorldPartitionViewportWidget();
+
+	static EVisibility GetVisibility(UWorld* InWorld);
+private:
+
+	bool bClickable;
+};
\ No newline at end of file
diff --git a/Engine/Source/Editor/WorldPartitionEditor/Private/WorldPartition/SWorldPartitionEditorGrid2D.cpp b/Engine/Source/Editor/WorldPartitionEditor/Private/WorldPartition/SWorldPartitionEditorGrid2D.cpp
--- a/Engine/Source/Editor/WorldPartitionEditor/Private/WorldPartition/SWorldPartitionEditorGrid2D.cpp
+++ b/Engine/Source/Editor/WorldPartitionEditor/Private/WorldPartition/SWorldPartitionEditorGrid2D.cpp
@@ -26,6 +26,7 @@
 #include "LevelEditorViewport.h"
 #include "LocationVolume.h"
 #include "Algo/Transform.h"
+#include "SWorldPartitionViewportWidget.h"
 
 #define LOCTEXT_NAMESPACE "WorldPartitionEditor"
 
@@ -236,6 +237,15 @@ void SWorldPartitionEditorGrid2D::Construct(const FArguments& InArgs)
 				]
 			]
 		]
+		+SOverlay::Slot()
+		.VAlign(VAlign_Bottom)
+		.HAlign(HAlign_Left)
+		.Padding(10.f, 0.f, 0.f, 10.f)
+		[
+			SNew(SWorldPartitionViewportWidget)
+			.Clickable(false)
+			.Visibility_Lambda([this]() { return SWorldPartitionViewportWidget::GetVisibility(World); })
+		]
 	];
 
 	SmallLayoutFont = FCoreStyle::GetDefaultFontStyle("Regular", 10);
diff --git a/Engine/Source/Runtime/Engine/Private/WorldPartition/WorldPartition.cpp b/Engine/Source/Runtime/Engine/Private/WorldPartition/WorldPartition.cpp
--- a/Engine/Source/Runtime/Engine/Private/WorldPartition/WorldPartition.cpp
+++ b/Engine/Source/Runtime/Engine/Private/WorldPartition/WorldPartition.cpp
@@ -220,6 +220,7 @@ UWorldPartition::UWorldPartition(const FObjectInitializer& ObjectInitializer)
 	, bForceGarbageCollection(false)
 	, bForceGarbageCollectionPurge(false)
 	, bIsPIE(false)
+	, NumUserCreatedLoadedRegions(0)
 #endif
 	, InitState(EWorldPartitionInitState::Uninitialized)
 	, StreamingPolicy(nullptr)
diff --git a/Engine/Source/Runtime/Engine/Private/WorldPartition/WorldPartitionActorLoaderInterface.cpp b/Engine/Source/Runtime/Engine/Private/WorldPartition/WorldPartitionActorLoaderInterface.cpp
--- a/Engine/Source/Runtime/Engine/Private/WorldPartition/WorldPartitionActorLoaderInterface.cpp
+++ b/Engine/Source/Runtime/Engine/Private/WorldPartition/WorldPartitionActorLoaderInterface.cpp
@@ -42,6 +42,14 @@ void IWorldPartitionActorLoaderInterface::ILoaderAdapter::Load()
 		bLoaded = true;
 		RefreshLoadedState();
 		RegisterDelegates();
+
+		if (bUserCreated)
+		{
+			if (UWorldPartition* WorldPartition = World->GetWorldPartition())
+			{
+				WorldPartition->OnUserCreatedRegionLoaded();
+			}
+		}
 	}
 }
 
@@ -70,6 +78,11 @@ void IWorldPartitionActorLoaderInterface::ILoaderAdapter::Unload()
 
 			bLoaded = false;
 
+			if (bUserCreated)
+			{
+				WorldPartition->OnUserCreatedRegionUnloaded();
+			}
+
 			PostLoadedStateChanged(true);
 		}
 	}
diff --git a/Engine/Source/Runtime/Engine/Public/WorldPartition/WorldPartition.h b/Engine/Source/Runtime/Engine/Public/WorldPartition/WorldPartition.h
--- a/Engine/Source/Runtime/Engine/Public/WorldPartition/WorldPartition.h
+++ b/Engine/Source/Runtime/Engine/Public/WorldPartition/WorldPartition.h
@@ -181,6 +181,10 @@ public:
 
 	void LoadLastLoadedRegions(const TArray<FBox>& EditorLastLoadedRegions);
 	void LoadLastLoadedRegions();
+
+	bool HasLoadedUserCreatedRegions() const { return !!NumUserCreatedLoadedRegions; }
+	void OnUserCreatedRegionLoaded() { NumUserCreatedLoadedRegions++; }
+	void OnUserCreatedRegionUnloaded() { check(HasLoadedUserCreatedRegions()); NumUserCreatedLoadedRegions--; }
 #endif
 
 public:
@@ -257,6 +261,7 @@ private:
 	bool bForceGarbageCollection;
 	bool bForceGarbageCollectionPurge;
 	bool bIsPIE;
+	int32 NumUserCreatedLoadedRegions;
 #endif
 
 #if WITH_EDITORONLY_DATA
