Hide SWorldPartitionViewportWidget from viewport in PIE

#rnx

diff --git a/Engine/Source/Editor/LevelEditor/Private/SLevelViewport.cpp b/Engine/Source/Editor/LevelEditor/Private/SLevelViewport.cpp
--- a/Engine/Source/Editor/LevelEditor/Private/SLevelViewport.cpp
+++ b/Engine/Source/Editor/LevelEditor/Private/SLevelViewport.cpp
@@ -390,7 +390,10 @@ void SLevelViewport::ConstructViewportOverlayContent()
 			[
 				SNew(SWorldPartitionViewportWidget)
 				.Clickable(true)
-				.Visibility_Lambda([this]() { return SWorldPartitionViewportWidget::GetVisibility(GetWorld());})
+				.Visibility_Lambda([this]() 
+				{ 
+					return (!ActiveViewport.IsValid() || IsPlayInEditorViewportActive()) ? EVisibility::Collapsed : SWorldPartitionViewportWidget::GetVisibility(GetWorld());
+				})
 			]
 		]
 	];
