Fix typo

diff --git a/Engine/Source/Editor/LevelEditor/Private/SLevelEditor.cpp b/Engine/Source/Editor/LevelEditor/Private/SLevelEditor.cpp
--- a/Engine/Source/Editor/LevelEditor/Private/SLevelEditor.cpp
+++ b/Engine/Source/Editor/LevelEditor/Private/SLevelEditor.cpp
@@ -1485,13 +1485,13 @@ TSharedRef<SWidget> SLevelEditor::RestoreContentArea( const TSharedRef<SDockTab>
 
 		{
 			// These modules will take care of their own tab spawners once loaded.
-			FObjectMixerEditorModule ObjectMixer = FObjectMixerEditorModule::Get();
+			FObjectMixerEditorModule& ObjectMixer = FObjectMixerEditorModule::Get();
 			{
 				FDelegateHandle ObjectMixerComponentEditedHandle =
 				   LevelEditorModule.OnComponentsEdited().AddRaw(&ObjectMixer, &FObjectMixerEditorModule::RequestRebuildList);
 				ObjectMixer.AddOnComponentEditedDelegate(ObjectMixerComponentEditedHandle);
 			}
-			FLightMixerModule LightMixer = FLightMixerModule::Get();
+			FLightMixerModule& LightMixer = FLightMixerModule::Get();
 			{
 				FDelegateHandle LightMixerComponentEditedHandle =
 				   LevelEditorModule.OnComponentsEdited().AddRaw(&LightMixer, &FLightMixerModule::RequestRebuildList);
