Level Instance: add an option to sync the world asset with the content browser.

#rnx

diff --git a/Engine/Source/Editor/LevelEditor/Private/LevelEditorContextMenu.cpp b/Engine/Source/Editor/LevelEditor/Private/LevelEditorContextMenu.cpp
--- a/Engine/Source/Editor/LevelEditor/Private/LevelEditorContextMenu.cpp
+++ b/Engine/Source/Editor/LevelEditor/Private/LevelEditorContextMenu.cpp
@@ -279,11 +279,14 @@ void FLevelEditorContextMenu::RegisterActorContextMenu()
 			TArray< UObject* > ReferencedAssets;
 			GEditor->GetReferencedAssetsForEditorSelection(ReferencedAssets);
 
+			TArray< FSoftObjectPath> SoftReferencedAssets;
+			GEditor->GetSoftReferencedAssetsForEditorSelection(SoftReferencedAssets);
+
 			// Asset type icon is used in multiple places below
 			FSlateIcon AssetIcon = ReferencedAssets.Num() == 1 ? FSlateIconFinder::FindIconForClass(ReferencedAssets[0]->GetClass()) : FSlateIcon(FAppStyle::GetAppStyleSetName(), "ClassIcon.Default");
 
 			// Edit and Find entries (a) always appear in main menu, and (b) appear in right-click menu if referenced asset is available
-			if (LevelEditorContext->ContextType == ELevelEditorMenuContext::MainMenu || ReferencedAssets.Num() > 0)
+			if (LevelEditorContext->ContextType == ELevelEditorMenuContext::MainMenu || ReferencedAssets.Num() > 0 || SoftReferencedAssets.Num() > 0)
 			{
 				Section.AddMenuEntry(FGlobalEditorCommonCommands::Get().FindInContentBrowser);
 
@@ -298,11 +301,12 @@ void FLevelEditorContextMenu::RegisterActorContextMenu()
 				}
 				else if (ReferencedAssets.Num() == 1)
 				{
-					auto Asset = ReferencedAssets[0];
+					UObject*  Asset = ReferencedAssets[0];
+					const FString AssetLabel = Cast<AActor>(Asset) ? Cast<AActor>(Asset)->GetActorNameOrLabel() : Asset->GetName();
 
 					Section.AddMenuEntry(
 						FLevelEditorCommands::Get().EditAsset,
-						FText::Format(LOCTEXT("EditAssociatedAsset", "Edit {0}"), FText::FromString(Asset->GetName())),
+						FText::Format(LOCTEXT("EditAssociatedAsset", "Edit {0}"), FText::FromString(AssetLabel)),
 						TAttribute<FText>(), // use command's tooltip
 						AssetIcon
 					);
diff --git a/Engine/Source/Editor/UnrealEd/Classes/Editor/EditorEngine.h b/Engine/Source/Editor/UnrealEd/Classes/Editor/EditorEngine.h
--- a/Engine/Source/Editor/UnrealEd/Classes/Editor/EditorEngine.h
+++ b/Engine/Source/Editor/UnrealEd/Classes/Editor/EditorEngine.h
@@ -2441,6 +2441,13 @@ public:
 	*/
 	void GetReferencedAssetsForEditorSelection(TArray<UObject*>& Objects, const bool bIgnoreOtherAssetsIfBPReferenced = false);
 
+	/**
+	* Queries for a list of assets that are soft referenced by the current editor selection (actors, surfaces, etc.)
+	*
+	* @param	SoftObjects							Array to be filled with asset objects referenced soft by the current editor selection
+	*/
+	void GetSoftReferencedAssetsForEditorSelection(TArray<FSoftObjectPath>& SoftObjects);
+
 	/** Returns a filter to restruct what assets show up in asset pickers based on what the selection is used for (i.e. what will reference the assets) */
 	DECLARE_DELEGATE_RetVal_OneParam(TSharedPtr<IAssetReferenceFilter>, FOnMakeAssetReferenceFilter, const FAssetReferenceFilterContext& /*Context*/);
 	FOnMakeAssetReferenceFilter& OnMakeAssetReferenceFilter() { return OnMakeAssetReferenceFilterDelegate; }
diff --git a/Engine/Source/Editor/UnrealEd/Private/EditorEngine.cpp b/Engine/Source/Editor/UnrealEd/Private/EditorEngine.cpp
--- a/Engine/Source/Editor/UnrealEd/Private/EditorEngine.cpp
+++ b/Engine/Source/Editor/UnrealEd/Private/EditorEngine.cpp
@@ -2986,6 +2986,25 @@ void UEditorEngine::GetAssetsToSyncToContentBrowser(TArray<FAssetData>& Assets,
 					{
 						Assets.Add(FAssetData(Object));
 					}
+
+					TArray<FSoftObjectPath> SoftObjects;
+					Actor->GetSoftReferencedContentObjects(SoftObjects);
+
+					if (SoftObjects.Num())
+					{
+						FAssetRegistryModule& AssetRegistryModule = FModuleManager::LoadModuleChecked<FAssetRegistryModule>("AssetRegistry");
+						IAssetRegistry& AssetRegistry = AssetRegistryModule.Get();
+
+						for (const FSoftObjectPath& SoftObject : SoftObjects)
+						{
+							FAssetData AssetData = AssetRegistry.GetAssetByObjectPath(*SoftObject.ToString());
+
+							if (AssetData.IsValid())
+							{
+								Assets.Add(AssetData);
+							}
+						}
+					}
 				}
 			}
 		}
@@ -3072,6 +3091,16 @@ void UEditorEngine::GetReferencedAssetsForEditorSelection(TArray<UObject*>& Obje
 	}
 }
 
+void UEditorEngine::GetSoftReferencedAssetsForEditorSelection(TArray<FSoftObjectPath>& SoftObjects)
+{
+	for ( FSelectionIterator It( GetSelectedActorIterator() ) ; It ; ++It )
+	{
+		AActor* Actor = static_cast<AActor*>( *It );
+		checkSlow( Actor->IsA(AActor::StaticClass()) );
+
+		Actor->GetSoftReferencedContentObjects(SoftObjects);
+	}
+}
 
 void UEditorEngine::ToggleSelectedActorMovementLock()
 {
diff --git a/Engine/Source/Runtime/Engine/Classes/GameFramework/Actor.h b/Engine/Source/Runtime/Engine/Classes/GameFramework/Actor.h
--- a/Engine/Source/Runtime/Engine/Classes/GameFramework/Actor.h
+++ b/Engine/Source/Runtime/Engine/Classes/GameFramework/Actor.h
@@ -2457,6 +2457,9 @@ public:
 	 */
 	virtual bool GetReferencedContentObjects( TArray<UObject*>& Objects ) const;
 
+	/** Similar to GetReferencedContentObjects, but for soft referenced objects */
+	virtual bool GetSoftReferencedContentObjects(TArray<FSoftObjectPath>& SoftObjects) const;
+
 	/** Returns NumUncachedStaticLightingInteractions for this actor */
 	const int32 GetNumUncachedStaticLightingInteractions() const;
 
diff --git a/Engine/Source/Runtime/Engine/Private/ActorEditor.cpp b/Engine/Source/Runtime/Engine/Private/ActorEditor.cpp
--- a/Engine/Source/Runtime/Engine/Private/ActorEditor.cpp
+++ b/Engine/Source/Runtime/Engine/Private/ActorEditor.cpp
@@ -1070,14 +1070,14 @@ UActorFolder* AActor::GetActorFolder(bool bSkipDeleted) const
 	UActorFolder* ActorFolder = nullptr;
 	if (ULevel* Level = GetLevel())
 	{
-		if (FolderGuid.IsValid())
-		{
+	if (FolderGuid.IsValid())
+	{
 			ActorFolder = Level->GetActorFolder(FolderGuid, bSkipDeleted);
-		}
-		else if (!FolderPath.IsNone())
-		{
+	}
+	else if (!FolderPath.IsNone())
+	{
 			ActorFolder = Level->GetActorFolder(FolderPath);
-		}
+	}
 	}
 	return ActorFolder;
 }
@@ -1325,6 +1325,11 @@ bool AActor::GetReferencedContentObjects( TArray<UObject*>& Objects ) const
 	return true;
 }
 
+bool AActor::GetSoftReferencedContentObjects(TArray<FSoftObjectPath>& SoftObjects) const
+{
+	return false;
+}
+
 EDataValidationResult AActor::IsDataValid(TArray<FText>& ValidationErrors)
 {
 	// Do not run asset validation on external actors, validation will be caught through map check
diff --git a/Engine/Source/Runtime/Engine/Private/LevelInstance/LevelInstanceActor.cpp b/Engine/Source/Runtime/Engine/Private/LevelInstance/LevelInstanceActor.cpp
--- a/Engine/Source/Runtime/Engine/Private/LevelInstance/LevelInstanceActor.cpp
+++ b/Engine/Source/Runtime/Engine/Private/LevelInstance/LevelInstanceActor.cpp
@@ -250,6 +250,16 @@ bool ALevelInstance::IsLockLocation() const
 	return Super::IsLockLocation() || LevelInstanceActorImpl.IsLockLocation();
 }
 
+bool ALevelInstance::GetSoftReferencedContentObjects(TArray<FSoftObjectPath>& SoftObjects) const
+{
+	if (WorldAsset.ToSoftObjectPath().IsValid())
+	{
+		SoftObjects.Add(WorldAsset.ToSoftObjectPath());
+		return true;
+	}
+	return false;
+}
+
 FBox ALevelInstance::GetComponentsBoundingBox(bool bNonColliding, bool bIncludeFromChildActors) const
 {
 	FBox Box = Super::GetComponentsBoundingBox(bNonColliding, bIncludeFromChildActors);
diff --git a/Engine/Source/Runtime/Engine/Public/LevelInstance/LevelInstanceActor.h b/Engine/Source/Runtime/Engine/Public/LevelInstance/LevelInstanceActor.h
--- a/Engine/Source/Runtime/Engine/Public/LevelInstance/LevelInstanceActor.h
+++ b/Engine/Source/Runtime/Engine/Public/LevelInstance/LevelInstanceActor.h
@@ -94,6 +94,8 @@ public:
 	virtual FBox GetComponentsBoundingBox(bool bNonColliding = false, bool bIncludeFromChildActors = false) const override;
 	virtual FBox GetStreamingBounds() const override;
 	virtual bool IsLockLocation() const override;
+	virtual bool GetSoftReferencedContentObjects(TArray<FSoftObjectPath>& SoftObjects) const override;
+	// End of AActor interface
 
 	DECLARE_MULTICAST_DELEGATE_OneParam(FOnLevelInstanceActorPostLoad, ALevelInstance*);
 	static FOnLevelInstanceActorPostLoad OnLevelInstanceActorPostLoad;
