Fix tools menu for level editor not being populated when substituted

diff --git a/Engine/Source/Editor/LevelEditor/Private/LevelEditorMenu.cpp b/Engine/Source/Editor/LevelEditor/Private/LevelEditorMenu.cpp
--- a/Engine/Source/Editor/LevelEditor/Private/LevelEditorMenu.cpp
+++ b/Engine/Source/Editor/LevelEditor/Private/LevelEditorMenu.cpp
@@ -289,6 +289,7 @@ void FLevelEditorMenu::RegisterLevelEditorMenus()
 	ToolMenus->RegisterMenu("LevelEditor.MainMenu", "MainFrame.MainMenu", EMultiBoxType::MenuBar);
 	ToolMenus->RegisterMenu("LevelEditor.MainMenu.File", "MainFrame.MainTabMenu.File");
 	ToolMenus->RegisterMenu("LevelEditor.MainMenu.Window", "MainFrame.MainMenu.Window");
+	ToolMenus->RegisterMenu("LevelEditor.MainMenu.Tools", "MainFrame.MainMenu.Tools");
 
 	// Add other top level menus
 	Local::ExtendMenuBar();
