Allow saving selected actors from the actor context menu.

#rnx

diff --git a/Engine/Source/Editor/LevelEditor/Private/LevelEditor.cpp b/Engine/Source/Editor/LevelEditor/Private/LevelEditor.cpp
--- a/Engine/Source/Editor/LevelEditor/Private/LevelEditor.cpp
+++ b/Engine/Source/Editor/LevelEditor/Private/LevelEditor.cpp
@@ -751,6 +751,12 @@ void FLevelEditorModule::BindGlobalLevelEditorCommands()
 		FExecuteAction::CreateStatic(&FLevelEditorActionCallbacks::CopyActorFilePathtoClipboard_Clicked)
 		);
 
+	ActionList.MapAction(
+		Commands.SaveActor,
+		FExecuteAction::CreateStatic(&FLevelEditorActionCallbacks::SaveActor_Clicked),
+		FCanExecuteAction::CreateStatic(&FLevelEditorActionCallbacks::SaveActor_CanExecute)
+		);
+
 	ActionList.MapAction(
 		Commands.ShowActorHistory,
 		FExecuteAction::CreateStatic(&FLevelEditorActionCallbacks::ShowActorHistory_Clicked),
diff --git a/Engine/Source/Editor/LevelEditor/Private/LevelEditorActions.cpp b/Engine/Source/Editor/LevelEditor/Private/LevelEditorActions.cpp
--- a/Engine/Source/Editor/LevelEditor/Private/LevelEditorActions.cpp
+++ b/Engine/Source/Editor/LevelEditor/Private/LevelEditorActions.cpp
@@ -2734,6 +2734,38 @@ void FLevelEditorActionCallbacks::CopyActorFilePathtoClipboard_Clicked()
 	}
 }
 
+void FLevelEditorActionCallbacks::SaveActor_Clicked()
+{
+	TSet<UPackage*> PackagesToSave;
+	for (FSelectionIterator It(GEditor->GetSelectedActorIterator()); It; ++It)
+	{
+		AActor* Actor = Cast<AActor>(*It);
+		if (Actor->IsPackageExternal())
+		{
+			PackagesToSave.Add(Actor->GetPackage());
+		}
+	}
+
+	if (PackagesToSave.Num())
+	{
+		FEditorFileUtils::PromptForCheckoutAndSave(PackagesToSave.Array(), /*bCheckDirty*/false, /*bPromptToSave*/false);
+	}
+}
+
+bool FLevelEditorActionCallbacks::SaveActor_CanExecute()
+{
+	TSet<UPackage*> PackagesToSave;
+	for (FSelectionIterator It(GEditor->GetSelectedActorIterator()); It; ++It)
+	{
+		AActor* Actor = Cast<AActor>(*It);
+		if (Actor->IsPackageExternal())
+		{
+			return true;
+		}
+	}
+	return false;
+}
+
 static TArray<FString> GetSelectedActorsPackageFullpath()
 {
 	TArray<FString> PackageFullpaths;
@@ -3436,8 +3468,9 @@ void FLevelEditorCommands::RegisterCommands()
 	UI_COMMAND( SnapCameraToObject, "Snap View to Object", "Snaps the view to the selected object", EUserInterfaceActionType::Button, FInputChord() );
 	UI_COMMAND( SnapObjectToCamera, "Snap Object to View", "Snaps the selected object to the view", EUserInterfaceActionType::Button, FInputChord() );
 
-	UI_COMMAND(CopyActorFilePathtoClipboard, "Copy Actor File Path", "Copy the file path where this actor is saved", EUserInterfaceActionType::Button, FInputChord());
-	UI_COMMAND(ShowActorHistory, "Show Actor History", "Shows the history of the file containing the actor.", EUserInterfaceActionType::Button, FInputChord());
+	UI_COMMAND(CopyActorFilePathtoClipboard, "Copy Selected Actor(s) File Path", "Copy the file path of the seledcted actors", EUserInterfaceActionType::Button, FInputChord());
+	UI_COMMAND(SaveActor, "Save Selected Actor(s)", "Save the selected actors", EUserInterfaceActionType::Button, FInputChord());
+	UI_COMMAND(ShowActorHistory, "Show Actor History", "Shows the history of the file containing the actor.", EUserInterfaceActionType::Button, FInputChord());	
 
 	UI_COMMAND( GoToCodeForActor, "Go to C++ Code for Actor", "Opens a code editing IDE and navigates to the source file associated with the seleced actor", EUserInterfaceActionType::Button, FInputChord() );
 	UI_COMMAND( GoToDocsForActor, "Go to Documentation for Actor", "Opens documentation for the Actor in the default web browser", EUserInterfaceActionType::Button, FInputChord() );
diff --git a/Engine/Source/Editor/LevelEditor/Private/LevelEditorContextMenu.cpp b/Engine/Source/Editor/LevelEditor/Private/LevelEditorContextMenu.cpp
--- a/Engine/Source/Editor/LevelEditor/Private/LevelEditorContextMenu.cpp
+++ b/Engine/Source/Editor/LevelEditor/Private/LevelEditorContextMenu.cpp
@@ -345,6 +345,13 @@ void FLevelEditorContextMenu::RegisterActorContextMenu()
 				FSlateIcon(FAppStyle::GetAppStyleSetName(), "GenericCommands.Copy")
 			);
 
+			Section.AddMenuEntry(
+				FLevelEditorCommands::Get().SaveActor,
+				TAttribute<FText>(), // use command's label
+				TAttribute<FText>(), // use command's tooltip
+				FSlateIcon(FAppStyle::GetAppStyleSetName(), "Icons.Save")
+			);
+
 			LevelEditorCreateActorMenu::FillAddReplaceContextMenuSections(Section, LevelEditorContext);
 		}
 
diff --git a/Engine/Source/Editor/LevelEditor/Public/LevelEditorActions.h b/Engine/Source/Editor/LevelEditor/Public/LevelEditorActions.h
--- a/Engine/Source/Editor/LevelEditor/Public/LevelEditorActions.h
+++ b/Engine/Source/Editor/LevelEditor/Public/LevelEditorActions.h
@@ -132,6 +132,9 @@ public:
 	/** Copy the file path where the actor is saved. */
 	TSharedPtr< FUICommandInfo > CopyActorFilePathtoClipboard;
 
+	/** Save the selected actor. */
+	TSharedPtr< FUICommandInfo > SaveActor;
+
 	/** Opens the source control panel. */
 	TSharedPtr< FUICommandInfo > OpenSourceControl;
 
@@ -1252,6 +1255,17 @@ public:
 	*/
 	static void CopyActorFilePathtoClipboard_Clicked();
 
+	/**
+	* Save the actor.
+	*/
+	static void SaveActor_Clicked();
+
+	/**
+	 * Checks whether SaveActor can be executed on the selected actors.
+	 * @return 	True if SaveActor can be executed on the selected actors.
+	 */
+	static bool SaveActor_CanExecute();
+
 	/**
 	* Shows the history of the file containing the actor.
 	*/
@@ -1259,7 +1273,7 @@ public:
 
 	/**
 	 * Checks whether ShowActorHistory can be executed on the selected actors.
-	 * @return 	True if ShowActorHistory can be executer on the selected actors.
+	 * @return 	True if ShowActorHistory can be executed on the selected actors.
 	 */
 	static bool ShowActorHistory_CanExecute();
 
