LevelViewportToolbar: Make LeftExtension AutoWidth to collapse unused space.

diff --git a/Engine/Source/Editor/LevelEditor/Private/SLevelViewportToolBar.cpp b/Engine/Source/Editor/LevelEditor/Private/SLevelViewportToolBar.cpp
--- a/Engine/Source/Editor/LevelEditor/Private/SLevelViewportToolBar.cpp
+++ b/Engine/Source/Editor/LevelEditor/Private/SLevelViewportToolBar.cpp
@@ -255,6 +255,7 @@ void SLevelViewportToolBar::Construct( const FArguments& InArgs )
 				]
 				+ SHorizontalBox::Slot()
 				.Padding(ToolbarSlotPadding)
+				.AutoWidth()
 				.HAlign(HAlign_Left)
 				.VAlign(VAlign_Fill)
 				[
