fix tooltip typo

#rnx

diff --git a/Engine/Source/Editor/LevelEditor/Private/LevelEditorActions.cpp b/Engine/Source/Editor/LevelEditor/Private/LevelEditorActions.cpp
--- a/Engine/Source/Editor/LevelEditor/Private/LevelEditorActions.cpp
+++ b/Engine/Source/Editor/LevelEditor/Private/LevelEditorActions.cpp
@@ -3468,7 +3468,7 @@ void FLevelEditorCommands::RegisterCommands()
 	UI_COMMAND( SnapCameraToObject, "Snap View to Object", "Snaps the view to the selected object", EUserInterfaceActionType::Button, FInputChord() );
 	UI_COMMAND( SnapObjectToCamera, "Snap Object to View", "Snaps the selected object to the view", EUserInterfaceActionType::Button, FInputChord() );
 
-	UI_COMMAND(CopyActorFilePathtoClipboard, "Copy Selected Actor(s) File Path", "Copy the file path of the seledcted actors", EUserInterfaceActionType::Button, FInputChord());
+	UI_COMMAND(CopyActorFilePathtoClipboard, "Copy Selected Actor(s) File Path", "Copy the file path of the selected actors", EUserInterfaceActionType::Button, FInputChord());
 	UI_COMMAND(SaveActor, "Save Selected Actor(s)", "Save the selected actors", EUserInterfaceActionType::Button, FInputChord());
 	UI_COMMAND(ShowActorHistory, "Show Actor History", "Shows the history of the file containing the actor.", EUserInterfaceActionType::Button, FInputChord());	
 
