When saving a template, make sure to go throuih the save as path as it properly handles partitioned worlds.

#rnx

diff --git a/Engine/Source/Editor/LevelEditor/Private/LevelEditorActions.cpp b/Engine/Source/Editor/LevelEditor/Private/LevelEditorActions.cpp
--- a/Engine/Source/Editor/LevelEditor/Private/LevelEditorActions.cpp
+++ b/Engine/Source/Editor/LevelEditor/Private/LevelEditorActions.cpp
@@ -439,7 +439,16 @@ bool FLevelEditorActionCallbacks::CanSaveUnpartitionedWorld()
 
 void FLevelEditorActionCallbacks::Save()
 {
-	FEditorFileUtils::SaveCurrentLevel();
+	// If the world is a template, go through the save current as path as it handles loading all external actors properly
+	UWorld* World = GetWorld();
+	if (FPackageName::IsTempPackage(World->GetPackage()->GetName()))
+	{
+		SaveCurrentAs();
+	}
+	else
+	{
+		FEditorFileUtils::SaveCurrentLevel();
+	}
 }
 
 bool FLevelEditorActionCallbacks::CanSaveCurrentAs()
@@ -458,7 +467,7 @@ void FLevelEditorActionCallbacks::SaveCurrentAs()
 	{
 		CurrentStreamingLevelClass = StreamingLevel->GetClass();
 	}
-
+	
 	
 	const bool bSavedPersistentLevelAs = CurrentLevel == World->PersistentLevel;
 	FString SavedFilename;
