[Level Viewport] Speculatively fix a crash in FLevelEditorModule::OnElementSelectionChanged delegate handler

diff --git a/Engine/Source/Editor/LevelEditor/Private/SLevelViewport.cpp b/Engine/Source/Editor/LevelEditor/Private/SLevelViewport.cpp
--- a/Engine/Source/Editor/LevelEditor/Private/SLevelViewport.cpp
+++ b/Engine/Source/Editor/LevelEditor/Private/SLevelViewport.cpp
@@ -151,6 +151,7 @@ SLevelViewport::~SLevelViewport()
 	FLevelEditorModule& LevelEditor = FModuleManager::GetModuleChecked<FLevelEditorModule>( LevelEditorName );
 	LevelEditor.OnRedrawLevelEditingViewports().RemoveAll( this );
 	LevelEditor.OnActorSelectionChanged().RemoveAll( this );
+	LevelEditor.OnElementSelectionChanged().RemoveAll( this );
 	LevelEditor.OnMapChanged().RemoveAll( this );
 
 	if(UObjectInitialized())
