Fix regression assert when saving newly created level instance.

#rnx

diff --git a/Engine/Source/Editor/LevelEditor/Private/LevelEditorActions.cpp b/Engine/Source/Editor/LevelEditor/Private/LevelEditorActions.cpp
--- a/Engine/Source/Editor/LevelEditor/Private/LevelEditorActions.cpp
+++ b/Engine/Source/Editor/LevelEditor/Private/LevelEditorActions.cpp
@@ -441,7 +441,7 @@ void FLevelEditorActionCallbacks::Save()
 {
 	// If the world is a template, go through the save current as path as it handles loading all external actors properly
 	UWorld* World = GetWorld();
-	if (FPackageName::IsTempPackage(World->GetPackage()->GetName()))
+	if (FPackageName::IsTempPackage(World->GetPackage()->GetName()) && CanSaveCurrentAs())
 	{
 		SaveCurrentAs();
 	}
@@ -3707,6 +3707,4 @@ void FLevelEditorCommands::RegisterCommands()
 	UI_COMMAND(OpenMergeActor, "Merge Actors", "Opens the Merge Actor panel", EUserInterfaceActionType::Button, FInputChord());
 }
 
-PRAGMA_ENABLE_OPTIMIZATION
-
 #undef LOCTEXT_NAMESPACE
