Fix removing needed PRAGMA_ENABLE_OPTIMIZATION.

#rnx

diff --git a/Engine/Source/Editor/LevelEditor/Private/LevelEditorActions.cpp b/Engine/Source/Editor/LevelEditor/Private/LevelEditorActions.cpp
--- a/Engine/Source/Editor/LevelEditor/Private/LevelEditorActions.cpp
+++ b/Engine/Source/Editor/LevelEditor/Private/LevelEditorActions.cpp
@@ -3707,4 +3707,6 @@ void FLevelEditorCommands::RegisterCommands()
 	UI_COMMAND(OpenMergeActor, "Merge Actors", "Opens the Merge Actor panel", EUserInterfaceActionType::Button, FInputChord());
 }
 
+PRAGMA_ENABLE_OPTIMIZATION
+
 #undef LOCTEXT_NAMESPACE
