Updating pilot the actor editor viewport toolbar section to UE5 styles.

#p4v-preflight-copy 21734448

diff --git a/Engine/Source/Editor/EditorStyle/Private/SlateEditorStyle.cpp b/Engine/Source/Editor/EditorStyle/Private/SlateEditorStyle.cpp
--- a/Engine/Source/Editor/EditorStyle/Private/SlateEditorStyle.cpp
+++ b/Engine/Source/Editor/EditorStyle/Private/SlateEditorStyle.cpp
@@ -5067,20 +5067,6 @@ void FSlateEditorStyle::FStyle::SetupLevelEditorStyle()
 		Set( "LevelViewport.ViewportConfig_FourPanesBottom.Small", new IMAGE_BRUSH( "Icons/ViewportLayout_FourPanesBottom", IconLayoutSizeSmall ) );
 		Set( "LevelViewport.ViewportConfig_FourPanes2x2", new IMAGE_BRUSH( "Icons/ViewportLayout_FourPanes2x2", IconLayoutSize ) );
 		Set( "LevelViewport.ViewportConfig_FourPanes2x2.Small", new IMAGE_BRUSH( "Icons/ViewportLayout_FourPanes2x2", IconLayoutSizeSmall ) );
-
-		Set( "LevelViewport.EjectActorPilot", new IMAGE_BRUSH( "Icons/icon_EjectActorPilot_16x", Icon16x16 ) );
-		Set( "LevelViewport.EjectActorPilot.Small", new IMAGE_BRUSH( "Icons/icon_EjectActorPilot_16x", Icon16x16 ) );
-		Set( "LevelViewport.PilotSelectedActor", new IMAGE_BRUSH( "Icons/icon_PilotSelectedActor_16x", Icon16x16 ) );
-		Set( "LevelViewport.PilotSelectedActor.Small", new IMAGE_BRUSH( "Icons/icon_PilotSelectedActor_16x", Icon16x16 ) );
-		Set( "LevelViewport.ToggleActorPilotCameraView", new IMAGE_BRUSH( "Icons/icon_ToggleActorPilotCameraView_16x", Icon16x16 ) );
-		Set( "LevelViewport.ToggleActorPilotCameraView.Small", new IMAGE_BRUSH( "Icons/icon_ToggleActorPilotCameraView_16x", Icon16x16 ) );
-
-		Set( "LevelViewport.ActorPilotText", FTextBlockStyle()
-			.SetFont( DEFAULT_FONT( "BoldCondensed", 12 ) )
-			.SetColorAndOpacity( FLinearColor(0.9f, 0.9f, 0.9f, 1.f) )
-			.SetShadowColorAndOpacity( FLinearColor(0.f, 0.f, 0.f, 0.4f) )
-			.SetShadowOffset( FVector2D(1.f, 1.f) )
-		);
 	}
 
 	// Level editor status bar
diff --git a/Engine/Source/Editor/EditorStyle/Private/StarshipStyle.cpp b/Engine/Source/Editor/EditorStyle/Private/StarshipStyle.cpp
--- a/Engine/Source/Editor/EditorStyle/Private/StarshipStyle.cpp
+++ b/Engine/Source/Editor/EditorStyle/Private/StarshipStyle.cpp
@@ -2276,6 +2276,8 @@ void FStarshipEditorStyle::FStyle::SetupViewportStyles()
 			.SetPressedPadding(FMargin(4.0f, 4.0f, 3.0f, 4.0f));
 		Set("EditorViewportToolBar.Button", ViewportMenuButton);
 
+		Set("EditorViewportToolBar.StartToolbarImage", new BOX_BRUSH("Starship/EditorViewport/ToolBarLeftGroup", 12.f/25.f, FStyleColors::Dropdown));
+
 		const FCheckBoxStyle ViewportMenuToggleLeftButtonStyle = FCheckBoxStyle(ViewportToggleButton)
 			.SetCheckBoxType(ESlateCheckBoxType::ToggleButton)
 			.SetUncheckedImage(		  BOX_BRUSH("Starship/EditorViewport/ToolBarLeftGroup", 12.f/25.f, FStyleColors::Dropdown))
@@ -2309,6 +2311,12 @@ void FStarshipEditorStyle::FStyle::SetupViewportStyles()
 			.SetPadding(ViewportMarginRight);
 		Set("EditorViewportToolBar.ToggleButton.End", ViewportMenuToggleRightButtonStyle);
 
+		Set("EditorViewportToolBar.Button.Middle", FButtonStyle(ViewportMenuButton)
+			.SetNormal(BOX_BRUSH("Starship/EditorViewport/ToolBarMiddleGroup", 12.f/25.f, FStyleColors::Dropdown))
+			.SetPressed(BOX_BRUSH("Starship/EditorViewport/ToolBarMiddleGroup", 12.f/25.f, FStyleColors::Recessed))
+			.SetHovered(BOX_BRUSH("Starship/EditorViewport/ToolBarMiddleGroup", 12.f/25.f, FStyleColors::Hover))
+		);
+
 		// We want a background-less version as the ComboMenu has its own unified background
 		const FToolBarStyle& SlimCoreToolBarStyle = FStarshipCoreStyle::GetCoreStyle().GetWidgetStyle<FToolBarStyle>("SlimToolBar");
 
@@ -4440,21 +4448,14 @@ void FStarshipEditorStyle::FStyle::SetupLevelEditorStyle()
 		Set( "LevelViewport.ViewportConfig_FourPanes2x2", new IMAGE_BRUSH( "Icons/ViewportLayout_FourPanes2x2", IconLayoutSize ) );
 		Set( "LevelViewport.ViewportConfig_FourPanes2x2.Small", new IMAGE_BRUSH( "Icons/ViewportLayout_FourPanes2x2", IconLayoutSizeSmall ) );
 
-		Set( "LevelViewport.EjectActorPilot", new IMAGE_BRUSH( "Icons/icon_EjectActorPilot_16x", Icon16x16 ) );
-		Set( "LevelViewport.EjectActorPilot.Small", new IMAGE_BRUSH( "Icons/icon_EjectActorPilot_16x", Icon16x16 ) );
+		Set( "LevelViewport.EjectActorPilot", new IMAGE_BRUSH_SVG( "Starship/EditorViewport/eject", Icon16x16 ) );
+		Set( "LevelViewport.EjectActorPilot.Small", new IMAGE_BRUSH_SVG( "Starship/EditorViewport/eject", Icon16x16 ) );
 		Set( "LevelViewport.PilotSelectedActor", new IMAGE_BRUSH_SVG( "Starship/EditorViewport/pilot", Icon16x16 ) );
 		
-		Set( "LevelViewport.ToggleActorPilotCameraView",       new IMAGE_BRUSH_SVG( "Starship/Common/CameraShake", Icon16x16 ) );
-		Set( "LevelViewport.ToggleActorPilotCameraView.Small", new IMAGE_BRUSH_SVG( "Starship/Common/CameraShake", Icon16x16 ) );
-
-		Set( "LevelViewport.ActorPilotText", FTextBlockStyle()
-			.SetFont( DEFAULT_FONT( "BoldCondensed", 12 ) )
-			.SetColorAndOpacity( FLinearColor(0.9f, 0.9f, 0.9f, 1.f) )
-			.SetShadowColorAndOpacity( FLinearColor(0.f, 0.f, 0.f, 0.4f) )
-			.SetShadowOffset( FVector2D(1.f, 1.f) )
-		);
+		Set( "LevelViewport.ToggleActorPilotCameraView",       new IMAGE_BRUSH_SVG( "Starship/EditorViewport/actor-pilot-camera", Icon16x16 ) );
+		Set( "LevelViewport.ToggleActorPilotCameraView.Small", new IMAGE_BRUSH_SVG( "Starship/EditorViewport/actor-pilot-camera", Icon16x16 ) );
 	}
-
+	
 	// Level editor status bar
 	{
 		Set( "TransformSettings.RelativeCoordinateSettings", new IMAGE_BRUSH( "Icons/icon_axis_16px", FVector2D( 16, 16 ) ) );
diff --git a/Engine/Source/Editor/LevelEditor/Private/SActorPilotViewportToolbar.h b/Engine/Source/Editor/LevelEditor/Private/SActorPilotViewportToolbar.h
--- a/Engine/Source/Editor/LevelEditor/Private/SActorPilotViewportToolbar.h
+++ b/Engine/Source/Editor/LevelEditor/Private/SActorPilotViewportToolbar.h
@@ -38,7 +38,7 @@ public:
 			Pilot = ViewportPtr->GetLevelViewportClient().GetActiveActorLock().Get();
 		}
 		
-		return Pilot ? FText::Format(LOCTEXT("ActiveText", "[ Pilot Active - {0} ]"), FText::FromString(Pilot->GetActorLabel())) : FText();
+		return Pilot ? FText::Format(LOCTEXT("ActiveText", "Pilot Actor: {0} "), FText::FromString(Pilot->GetActorLabel())) : FText();
 	}
 
 	EVisibility GetLockedTextVisibility() const
@@ -61,38 +61,35 @@ public:
 
 		auto& ViewportCommands = FLevelViewportCommands::Get();
 
-		FToolBarBuilder ToolbarBuilder(InArgs._Viewport->GetCommandList(), FMultiBoxCustomization::None);
+		FSlimHorizontalToolBarBuilder ToolbarBuilder(InArgs._Viewport->GetCommandList(), FMultiBoxCustomization::None);
 
 		// Use a custom style
-		FName ToolBarStyle = "LegacyViewportMenu";
-		ToolbarBuilder.SetStyle(&FAppStyle::Get(), ToolBarStyle);
+		FName ToolBarStyle = "EditorViewportToolBar";
 		ToolbarBuilder.SetStyle(&FAppStyle::Get(), ToolBarStyle);
 		ToolbarBuilder.SetLabelVisibility(EVisibility::Collapsed);
-
+		
 		ToolbarBuilder.BeginSection("ActorPilot");
-		ToolbarBuilder.BeginBlockGroup();
 		{
-			static FName EjectActorPilotName = FName(TEXT("EjectActorPilot"));
-			ToolbarBuilder.AddToolBarButton(ViewportCommands.EjectActorPilot, NAME_None, TAttribute<FText>(), TAttribute<FText>(), TAttribute<FSlateIcon>(), EjectActorPilotName);
+			ToolbarBuilder.BeginBlockGroup();
 
-			static FName ToggleActorPilotCameraViewName = FName(TEXT("ToggleActorPilotCameraView"));
-			ToolbarBuilder.AddToolBarButton(ViewportCommands.ToggleActorPilotCameraView, NAME_None, TAttribute<FText>(), TAttribute<FText>(), TAttribute<FSlateIcon>(), ToggleActorPilotCameraViewName);
-		}
-		ToolbarBuilder.EndBlockGroup();
-		ToolbarBuilder.EndSection();
-
-		ToolbarBuilder.BeginSection("ActorPilot_Label");
-		ToolbarBuilder.AddWidget(
-			SNew(SBox)
-			// Nasty hack to make this VAlign_center properly. The parent Box is set to VAlign_Bottom, so we can't fill properly.
-			.HeightOverride(24.f)
+			ToolbarBuilder.AddWidget(
+			SNew(SBorder)
+			.BorderImage(FAppStyle::Get().GetBrush("EditorViewportToolBar.StartToolbarImage"))
 			[
 				SNew(SHorizontalBox)
 				+ SHorizontalBox::Slot()
 				.VAlign(VAlign_Center)
+				.AutoWidth()
+				.Padding(5.0f, 0.0f)
+				[
+						SNew(SImage)
+						.Image(FAppStyle::Get()
+						.GetBrush("LevelViewport.PilotSelectedActor"))
+				]
+				+ SHorizontalBox::Slot()
+				.VAlign(VAlign_Center)
 				[
 					SNew(STextBlock)
-					.TextStyle(FAppStyle::Get(), "LevelViewport.ActorPilotText")
 					.Text(this, &SActorPilotViewportToolbar::GetActiveText)
 				]
 				+ SHorizontalBox::Slot()
@@ -101,14 +98,20 @@ public:
 				.AutoWidth()
 				[
 					SNew(STextBlock)
-					.TextStyle(FAppStyle::Get(), "LevelViewport.ActorPilotText")
-					.ColorAndOpacity(FLinearColor::Red)
 					.Text(LOCTEXT("ActorLockedText", "(Locked)"))
 					.ToolTipText(LOCTEXT("ActorLockedToolTipText", "This actor has locked movement so it will not be updated based on camera position"))
 					.Visibility(this, &SActorPilotViewportToolbar::GetLockedTextVisibility)
 				]
 			]
 		);
+			
+			static const FName EjectActorPilotName = "EjectActorPilot";
+			ToolbarBuilder.AddToolBarButton(ViewportCommands.EjectActorPilot, NAME_None, TAttribute<FText>(), TAttribute<FText>(), TAttribute<FSlateIcon>(), EjectActorPilotName);
+			
+			static FName ToggleActorPilotCameraViewName = "ToggleActorPilotCameraView";
+			ToolbarBuilder.AddToolBarButton(ViewportCommands.ToggleActorPilotCameraView, NAME_None, TAttribute<FText>(), TAttribute<FText>(), TAttribute<FSlateIcon>(), ToggleActorPilotCameraViewName);
+			ToolbarBuilder.EndBlockGroup();
+		}
 		ToolbarBuilder.EndSection();
 
 		ChildSlot
