World Partition: Disable build menu entries if they are not relevant
* Mininmap - World must be partitioned
* HLOD/Spline meshes - World must be partitioned & streaming must be enabled

diff --git a/Engine/Source/Editor/LevelEditor/Private/LevelEditor.cpp b/Engine/Source/Editor/LevelEditor/Private/LevelEditor.cpp
--- a/Engine/Source/Editor/LevelEditor/Private/LevelEditor.cpp
+++ b/Engine/Source/Editor/LevelEditor/Private/LevelEditor.cpp
@@ -1482,13 +1482,16 @@ void FLevelEditorModule::BindGlobalLevelEditorCommands()
 		FExecuteAction::CreateStatic( &FLevelEditorActionCallbacks::BuildPathsOnly_Execute ) );
 
 	ActionList.MapAction(Commands.BuildHLODs,
-		FExecuteAction::CreateStatic(&FLevelEditorActionCallbacks::BuildHLODs_Execute));
+		FExecuteAction::CreateStatic(&FLevelEditorActionCallbacks::BuildHLODs_Execute),
+		FCanExecuteAction::CreateStatic(&FLevelEditorActionCallbacks::IsWorldPartitionStreamingEnabled));
 	
 	ActionList.MapAction(Commands.BuildMinimap,
-		FExecuteAction::CreateStatic(&FLevelEditorActionCallbacks::BuildMinimap_Execute));
+		FExecuteAction::CreateStatic(&FLevelEditorActionCallbacks::BuildMinimap_Execute),
+		FCanExecuteAction::CreateStatic(&FLevelEditorActionCallbacks::IsWorldPartitionEnabled));
 
 	ActionList.MapAction(Commands.BuildLandscapeSplineMeshes,
-		FExecuteAction::CreateStatic(&FLevelEditorActionCallbacks::BuildLandscapeSplineMeshes_Execute));
+		FExecuteAction::CreateStatic(&FLevelEditorActionCallbacks::BuildLandscapeSplineMeshes_Execute),
+		FCanExecuteAction::CreateStatic(&FLevelEditorActionCallbacks::IsWorldPartitionStreamingEnabled));
 
 	ActionList.MapAction(Commands.BuildTextureStreamingOnly,
 		FExecuteAction::CreateStatic(&FLevelEditorActionCallbacks::BuildTextureStreamingOnly_Execute));
diff --git a/Engine/Source/Editor/LevelEditor/Private/LevelEditorActions.cpp b/Engine/Source/Editor/LevelEditor/Private/LevelEditorActions.cpp
--- a/Engine/Source/Editor/LevelEditor/Private/LevelEditorActions.cpp
+++ b/Engine/Source/Editor/LevelEditor/Private/LevelEditorActions.cpp
@@ -820,6 +820,21 @@ void FLevelEditorActionCallbacks::BuildPathsOnly_Execute()
 	FEditorBuildUtils::EditorBuild( GetWorld(), FBuildOptions::BuildAIPaths );
 }
 
+bool FLevelEditorActionCallbacks::IsWorldPartitionEnabled()
+{
+	return UWorld::IsPartitionedWorld(GetWorld());
+}
+
+bool FLevelEditorActionCallbacks::IsWorldPartitionStreamingEnabled()
+{
+	if (!IsWorldPartitionEnabled())
+	{
+		return false;
+	}
+
+	return GetWorld()->GetWorldPartition()->IsStreamingEnabled();
+}
+
 void FLevelEditorActionCallbacks::BuildHLODs_Execute()
 {
 	// Build HLOD
diff --git a/Engine/Source/Editor/LevelEditor/Public/LevelEditorActions.h b/Engine/Source/Editor/LevelEditor/Public/LevelEditorActions.h
--- a/Engine/Source/Editor/LevelEditor/Public/LevelEditorActions.h
+++ b/Engine/Source/Editor/LevelEditor/Public/LevelEditorActions.h
@@ -751,6 +751,8 @@ public:
 	static void BuildGeometryOnly_Execute();
 	static void BuildGeometryOnly_OnlyCurrentLevel_Execute();
 	static void BuildPathsOnly_Execute();
+	static bool IsWorldPartitionEnabled();
+	static bool IsWorldPartitionStreamingEnabled();
 	static void BuildHLODs_Execute();
 	static void BuildMinimap_Execute();
 	static void BuildLandscapeSplineMeshes_Execute();
