Level Editor: ActiveViewport::GetSize can crash if there is no active viewport,  for example if the tab was closed.

diff --git a/Engine/Source/Editor/LevelEditor/Private/SLevelEditor.cpp b/Engine/Source/Editor/LevelEditor/Private/SLevelEditor.cpp
--- a/Engine/Source/Editor/LevelEditor/Private/SLevelEditor.cpp
+++ b/Engine/Source/Editor/LevelEditor/Private/SLevelEditor.cpp
@@ -2166,19 +2166,21 @@ void SLevelEditor::RemoveViewportOverlayWidget(TSharedRef<SWidget> InWidget, TSh
 	
 FVector2D SLevelEditor::GetActiveViewportSize()
 {
-	TSharedPtr<SLevelViewport> ActiveViewport = GetActiveViewport();
-
-	FSceneViewFamilyContext ViewFamily(FSceneViewFamily::ConstructionValues(
-		ActiveViewport->GetActiveViewport(),
-		ActiveViewport->GetViewportClient()->GetScene(),
-		ActiveViewport->GetViewportClient()->EngineShowFlags)
-		.SetRealtimeUpdate(ActiveViewport->IsRealtime()));
-	// SceneView is deleted with the ViewFamily
-	FSceneView* SceneView = ActiveViewport->GetViewportClient()->CalcSceneView(&ViewFamily);
-	const float InvDpiScale = 1.0f / ActiveViewport->GetViewportClient()->GetDPIScale();
-	const float MaxX = SceneView->UnscaledViewRect.Width() * InvDpiScale;
-	const float MaxY = SceneView->UnscaledViewRect.Height() * InvDpiScale;
-	return FVector2D(MaxX, MaxY);
+	if (TSharedPtr<SLevelViewport> ActiveViewport = GetActiveViewport())
+	{
+		FSceneViewFamilyContext ViewFamily(FSceneViewFamily::ConstructionValues(
+			ActiveViewport->GetActiveViewport(),
+			ActiveViewport->GetViewportClient()->GetScene(),
+			ActiveViewport->GetViewportClient()->EngineShowFlags)
+			.SetRealtimeUpdate(ActiveViewport->IsRealtime()));
+		// SceneView is deleted with the ViewFamily
+		FSceneView* SceneView = ActiveViewport->GetViewportClient()->CalcSceneView(&ViewFamily);
+		const float InvDpiScale = 1.0f / ActiveViewport->GetViewportClient()->GetDPIScale();
+		const float MaxX = SceneView->UnscaledViewRect.Width() * InvDpiScale;
+		const float MaxY = SceneView->UnscaledViewRect.Height() * InvDpiScale;
+		return FVector2D(MaxX, MaxY);
+	}
+	return FVector2D(0, 0);
 }
 
 #undef LOCTEXT_NAMESPACE
