Sequencer: Show an error dialog when using the Cinematics Toolbar menu and trying to open a Level Sequence (Actor) that has no Level Sequence (property) assigned to it.

diff --git a/Engine/Source/Editor/LevelEditor/Private/LevelEditorToolBar.cpp b/Engine/Source/Editor/LevelEditor/Private/LevelEditorToolBar.cpp
--- a/Engine/Source/Editor/LevelEditor/Private/LevelEditorToolBar.cpp
+++ b/Engine/Source/Editor/LevelEditor/Private/LevelEditorToolBar.cpp
@@ -2089,14 +2089,21 @@ void FLevelEditorToolBar::OnCinematicsActorPicked( AActor* Actor )
 	// Make sure we dismiss the menus before we open this
 	if (ALevelSequenceActor* LevelSequenceActor = Cast<ALevelSequenceActor>(Actor))
 	{
-		FScopedSlowTask SlowTask(1.f, NSLOCTEXT("LevelToolBarCinematicsMenu", "LoadSequenceSlowTask", "Loading Level Sequence..."));
-		SlowTask.MakeDialog();
-		SlowTask.EnterProgressFrame();
-		UObject* Asset = LevelSequenceActor->GetSequence();
+		if (LevelSequenceActor->GetSequence())
+		{
+			FScopedSlowTask SlowTask(1.f, NSLOCTEXT("LevelToolBarCinematicsMenu", "LoadSequenceSlowTask", "Loading Level Sequence..."));
+			SlowTask.MakeDialog();
+			SlowTask.EnterProgressFrame();
+			UObject* Asset = LevelSequenceActor->GetSequence();
 
-		if (Asset != nullptr)
+			if (Asset != nullptr)
+			{
+				GEditor->GetEditorSubsystem<UAssetEditorSubsystem>()->OpenEditorForAsset(Asset);
+			}
+		}
+		else
 		{
-			GEditor->GetEditorSubsystem<UAssetEditorSubsystem>()->OpenEditorForAsset(Asset);
+			FMessageDialog::Open(EAppMsgType::Type::Ok, NSLOCTEXT("LevelToolBarCinematicsMenu", "LoadSequenceNotValid", "This Level Sequence actor does not have a Sequence asset assigned in the Details panel and cannot be opened."));
 		}
 	}
 }
