Fix CIS for Commandlet

diff --git a/Engine/Source/Editor/LevelEditor/Private/LevelEditorActions.cpp b/Engine/Source/Editor/LevelEditor/Private/LevelEditorActions.cpp
--- a/Engine/Source/Editor/LevelEditor/Private/LevelEditorActions.cpp
+++ b/Engine/Source/Editor/LevelEditor/Private/LevelEditorActions.cpp
@@ -3715,8 +3715,12 @@ void FLevelEditorCommands::RegisterCommands()
 	// Add preview platforms
 	for (const FPreviewPlatformMenuItem& Item : FDataDrivenPlatformInfoRegistry::GetAllPreviewPlatformMenuItems())
 	{
-		EShaderPlatform ShaderPlatform = FDataDrivenShaderPlatformInfo::GetShaderPlatformFromName(Item.ShaderPlatformPreview);
-		FText FriendlyName = FDataDrivenShaderPlatformInfo::GetFriendlyName(ShaderPlatform);
+		FText FriendlyName;
+		if (!IsRunningCommandlet())
+		{
+			EShaderPlatform ShaderPlatform = FDataDrivenShaderPlatformInfo::GetShaderPlatformFromName(Item.ShaderPlatformPreview);
+			FriendlyName = FDataDrivenShaderPlatformInfo::GetFriendlyName(ShaderPlatform);
+		}
 
 		PreviewPlatformOverrides.Add(
 			FUICommandInfoDecl(
diff --git a/Engine/Source/Runtime/RHI/Public/RHIDefinitions.h b/Engine/Source/Runtime/RHI/Public/RHIDefinitions.h
--- a/Engine/Source/Runtime/RHI/Public/RHIDefinitions.h
+++ b/Engine/Source/Runtime/RHI/Public/RHIDefinitions.h
@@ -1043,6 +1043,10 @@ public:
 #if WITH_EDITOR
 	static FORCEINLINE_DEBUGGABLE FText GetFriendlyName(const FStaticShaderPlatform Platform)
 	{
+		if (IsRunningCommandlet())
+		{
+			return FText();
+		}
 		check(IsValid(Platform));
 		return Infos[Platform].FriendlyName;
 	}
