Fix NullRHI LevelEditor

diff --git a/Engine/Source/Editor/LevelEditor/Private/LevelEditorActions.cpp b/Engine/Source/Editor/LevelEditor/Private/LevelEditorActions.cpp
--- a/Engine/Source/Editor/LevelEditor/Private/LevelEditorActions.cpp
+++ b/Engine/Source/Editor/LevelEditor/Private/LevelEditorActions.cpp
@@ -3716,7 +3716,7 @@ void FLevelEditorCommands::RegisterCommands()
 	for (const FPreviewPlatformMenuItem& Item : FDataDrivenPlatformInfoRegistry::GetAllPreviewPlatformMenuItems())
 	{
 		FText FriendlyName;
-		if (!IsRunningCommandlet())
+		if (!IsRunningCommandlet() && !GUsingNullRHI)
 		{
 			EShaderPlatform ShaderPlatform = FDataDrivenShaderPlatformInfo::GetShaderPlatformFromName(Item.ShaderPlatformPreview);
 			FriendlyName = FDataDrivenShaderPlatformInfo::GetFriendlyName(ShaderPlatform);
diff --git a/Engine/Source/Runtime/RHI/Private/RHI.cpp b/Engine/Source/Runtime/RHI/Private/RHI.cpp
--- a/Engine/Source/Runtime/RHI/Private/RHI.cpp
+++ b/Engine/Source/Runtime/RHI/Private/RHI.cpp
@@ -2668,6 +2668,18 @@ void FGenericDataDrivenShaderPlatformInfo::UpdatePreviewPlatforms()
 	}
 }
 
+#if WITH_EDITOR
+FText FGenericDataDrivenShaderPlatformInfo::GetFriendlyName(const FStaticShaderPlatform Platform)
+{
+	if (IsRunningCommandlet() || GUsingNullRHI)
+	{
+		return FText();
+	}
+	check(IsValid(Platform));
+	return Infos[Platform].FriendlyName;
+}
+#endif
+
 const EShaderPlatform FGenericDataDrivenShaderPlatformInfo::GetShaderPlatformFromName(const FName ShaderPlatformName)
 {
 	for (int32 i = 0; i < SP_NumPlatforms; ++i)
diff --git a/Engine/Source/Runtime/RHI/Public/RHIDefinitions.h b/Engine/Source/Runtime/RHI/Public/RHIDefinitions.h
--- a/Engine/Source/Runtime/RHI/Public/RHIDefinitions.h
+++ b/Engine/Source/Runtime/RHI/Public/RHIDefinitions.h
@@ -1041,15 +1041,7 @@ public:
 	}
 
 #if WITH_EDITOR
-	static FORCEINLINE_DEBUGGABLE FText GetFriendlyName(const FStaticShaderPlatform Platform)
-	{
-		if (IsRunningCommandlet())
-		{
-			return FText();
-		}
-		check(IsValid(Platform));
-		return Infos[Platform].FriendlyName;
-	}
+	static FText GetFriendlyName(const FStaticShaderPlatform Platform);
 #endif
 
 private:
