Fix assert when trying to convert an unknown shader platform to a feature level while creating some editor menus

diff --git a/Engine/Source/Editor/LevelEditor/Private/LevelEditor.cpp b/Engine/Source/Editor/LevelEditor/Private/LevelEditor.cpp
--- a/Engine/Source/Editor/LevelEditor/Private/LevelEditor.cpp
+++ b/Engine/Source/Editor/LevelEditor/Private/LevelEditor.cpp
@@ -1794,26 +1794,31 @@ void FLevelEditorModule::BindGlobalLevelEditorCommands()
 	{
 		const FPreviewPlatformMenuItem& Item = MenuItems[Index];
 		EShaderPlatform ShaderPlatform = ShaderFormatToLegacyShaderPlatform(Item.ShaderFormat);
-		ERHIFeatureLevel::Type FeatureLevel = GetMaxSupportedFeatureLevel(ShaderPlatform);
-
-		const bool IsDefaultActive = ShaderPlatform == GMaxRHIShaderPlatform;
-		const bool AllowPreview = !IsDefaultActive;
-
-		EShaderPlatform InShaderPlatform = EShaderPlatform::SP_NumPlatforms;
-		for (uint32 j = 0; j < SP_NumPlatforms; ++j)
+		
+		if (ShaderPlatform < SP_NumPlatforms)
 		{
-			InShaderPlatform = (EShaderPlatform)j;
-			if (!FDataDrivenShaderPlatformInfo::IsValid(InShaderPlatform))
-			{
-				continue;
-			}
+			ERHIFeatureLevel::Type FeatureLevel = GetMaxSupportedFeatureLevel(ShaderPlatform);
+
+			const bool IsDefaultActive = ShaderPlatform == GMaxRHIShaderPlatform;
+			const bool AllowPreview = !IsDefaultActive;
 
-			FName ShaderPlatformTest = FDataDrivenShaderPlatformInfo::GetName(InShaderPlatform);
-			if (Item.ShaderPlatformPreview == ShaderPlatformTest)
+			EShaderPlatform InShaderPlatform = EShaderPlatform::SP_NumPlatforms;
+			for (uint32 j = 0; j < SP_NumPlatforms; ++j)
 			{
-				break;
+				InShaderPlatform = (EShaderPlatform)j;
+				if (!FDataDrivenShaderPlatformInfo::IsValid(InShaderPlatform))
+				{
+					continue;
+				}
+
+				FName ShaderPlatformTest = FDataDrivenShaderPlatformInfo::GetName(InShaderPlatform);
+				if (Item.ShaderPlatformPreview == ShaderPlatformTest)
+				{
+					break;
+				}
 			}
 		}
+
 		FPreviewPlatformInfo PreviewFeatureLevelInfo(FeatureLevel, (EShaderPlatform)InShaderPlatform, IsDefaultActive ? NAME_None : Item.PlatformName, IsDefaultActive ? NAME_None : Item.ShaderFormat, IsDefaultActive ? NAME_None : Item.DeviceProfileName, AllowPreview, Item.ShaderPlatformPreview);
 
 		ActionList.MapAction(
