Fix compile error

diff --git a/Engine/Source/Editor/LevelEditor/Private/LevelEditor.cpp b/Engine/Source/Editor/LevelEditor/Private/LevelEditor.cpp
--- a/Engine/Source/Editor/LevelEditor/Private/LevelEditor.cpp
+++ b/Engine/Source/Editor/LevelEditor/Private/LevelEditor.cpp
@@ -1817,15 +1817,15 @@ void FLevelEditorModule::BindGlobalLevelEditorCommands()
 					break;
 				}
 			}
-		}
 
-		FPreviewPlatformInfo PreviewFeatureLevelInfo(FeatureLevel, (EShaderPlatform)InShaderPlatform, IsDefaultActive ? NAME_None : Item.PlatformName, IsDefaultActive ? NAME_None : Item.ShaderFormat, IsDefaultActive ? NAME_None : Item.DeviceProfileName, AllowPreview, Item.ShaderPlatformPreview);
+			FPreviewPlatformInfo PreviewFeatureLevelInfo(FeatureLevel, (EShaderPlatform)InShaderPlatform, IsDefaultActive ? NAME_None : Item.PlatformName, IsDefaultActive ? NAME_None : Item.ShaderFormat, IsDefaultActive ? NAME_None : Item.DeviceProfileName, AllowPreview, Item.ShaderPlatformPreview);
 
-		ActionList.MapAction(
-			Commands.PreviewPlatformOverrides[Index],
-			FExecuteAction::CreateStatic(&FLevelEditorActionCallbacks::SetPreviewPlatform, PreviewFeatureLevelInfo),
-			FCanExecuteAction::CreateStatic(&FLevelEditorActionCallbacks::CanExecutePreviewPlatform, PreviewFeatureLevelInfo),
-			FIsActionChecked::CreateStatic(&FLevelEditorActionCallbacks::IsPreviewPlatformChecked, PreviewFeatureLevelInfo));
+			ActionList.MapAction(
+				Commands.PreviewPlatformOverrides[Index],
+				FExecuteAction::CreateStatic(&FLevelEditorActionCallbacks::SetPreviewPlatform, PreviewFeatureLevelInfo),
+				FCanExecuteAction::CreateStatic(&FLevelEditorActionCallbacks::CanExecutePreviewPlatform, PreviewFeatureLevelInfo),
+				FIsActionChecked::CreateStatic(&FLevelEditorActionCallbacks::IsPreviewPlatformChecked, PreviewFeatureLevelInfo));
+		}
 	}
 
 	// Add the Disable Preview Menu Item
