Group selection: don't select the group actor itself

Makes it possible to edit properties on the grouped actors

diff --git a/Engine/Source/Editor/LevelEditor/Private/Elements/Actor/ActorElementLevelEditorSelectionCustomization.cpp b/Engine/Source/Editor/LevelEditor/Private/Elements/Actor/ActorElementLevelEditorSelectionCustomization.cpp
--- a/Engine/Source/Editor/LevelEditor/Private/Elements/Actor/ActorElementLevelEditorSelectionCustomization.cpp
+++ b/Engine/Source/Editor/LevelEditor/Private/Elements/Actor/ActorElementLevelEditorSelectionCustomization.cpp
@@ -318,8 +318,6 @@ bool FActorElementLevelEditorSelectionCustomization::SelectActorGroup(AGroupActo
 	{
 		FTypedElementList::FLegacySyncScopedBatch LegacySyncBatch(*InSelectionSet, InSelectionOptions.AllowLegacyNotifications());
 
-		bSelectionChanged |= SelectActorElement(GroupSelectionHandle, InSelectionSet, GroupSelectionOptions);
-
 		TArray<AActor*> GroupActors;
 		InGroupActor->GetGroupActors(GroupActors);
 		for (AActor* Actor : GroupActors)
