Outliner: Update the list of default type filters
Add the ability for other modules to add custom filters to the Level Editor Outliner

diff --git a/Engine/Plugins/Editor/ModelingToolsEditorMode/Source/ModelingToolsEditorMode/ModelingToolsEditorMode.Build.cs b/Engine/Plugins/Editor/ModelingToolsEditorMode/Source/ModelingToolsEditorMode/ModelingToolsEditorMode.Build.cs
--- a/Engine/Plugins/Editor/ModelingToolsEditorMode/Source/ModelingToolsEditorMode/ModelingToolsEditorMode.Build.cs
+++ b/Engine/Plugins/Editor/ModelingToolsEditorMode/Source/ModelingToolsEditorMode/ModelingToolsEditorMode.Build.cs
@@ -49,6 +49,7 @@ public class ModelingToolsEditorMode : ModuleRules
 				"StylusInput",
 				"DeveloperSettings",
 				"PropertyEditor",
+				"GeometryFramework"
 				// ... add private dependencies that you statically link with here ...	
 			}
 		);
diff --git a/Engine/Plugins/Editor/ModelingToolsEditorMode/Source/ModelingToolsEditorMode/Private/ModelingToolsEditorModeModule.cpp b/Engine/Plugins/Editor/ModelingToolsEditorMode/Source/ModelingToolsEditorMode/Private/ModelingToolsEditorModeModule.cpp
--- a/Engine/Plugins/Editor/ModelingToolsEditorMode/Source/ModelingToolsEditorMode/Private/ModelingToolsEditorModeModule.cpp
+++ b/Engine/Plugins/Editor/ModelingToolsEditorMode/Source/ModelingToolsEditorMode/Private/ModelingToolsEditorModeModule.cpp
@@ -25,6 +25,10 @@
 #include "BakeMultiMeshAttributeMapsTool.h"
 #include "BakeMeshAttributeVertexTool.h"
 #include "BakeTransformTool.h"
+#include "DynamicMeshActor.h"
+
+#include "LevelEditor.h"
+#include "Filters/CustomClassFilterData.h"
 
 
 #define LOCTEXT_NAMESPACE "FModelingToolsEditorModeModule"
@@ -101,6 +105,15 @@ void FModelingToolsEditorModeModule::OnPostEngineInit()
 
 	PropertyModule.RegisterCustomClassLayout("BakeTransformToolProperties", FOnGetDetailCustomizationInstance::CreateStatic(&FBakeTransformToolDetails::MakeInstance));
 	ClassesToUnregisterOnShutdown.Add(UBakeTransformToolProperties::StaticClass()->GetFName());
+
+	FLevelEditorModule& LevelEditorModule = FModuleManager::LoadModuleChecked<FLevelEditorModule>("LevelEditor");
+
+	// Register Level Editor Outliner filter
+	if(TSharedPtr<FFilterCategory> GeometryFilterCategory = LevelEditorModule.GetOutlinerFilterCategory(FLevelEditorOutlinerBuiltInCategories::Geometry()))
+	{
+		TSharedRef<FCustomClassFilterData> DynamicMeshActorClassData = MakeShared<FCustomClassFilterData>(ADynamicMeshActor::StaticClass(), GeometryFilterCategory, FLinearColor::White);
+		LevelEditorModule.AddCustomClassFilterToOutliner(DynamicMeshActorClassData);
+	}
 }
 
 #undef LOCTEXT_NAMESPACE
diff --git a/Engine/Plugins/Experimental/VirtualCamera/Source/VirtualCameraEditor/Private/VirtualCameraEditorModule.cpp b/Engine/Plugins/Experimental/VirtualCamera/Source/VirtualCameraEditor/Private/VirtualCameraEditorModule.cpp
--- a/Engine/Plugins/Experimental/VirtualCamera/Source/VirtualCameraEditor/Private/VirtualCameraEditorModule.cpp
+++ b/Engine/Plugins/Experimental/VirtualCamera/Source/VirtualCameraEditor/Private/VirtualCameraEditorModule.cpp
@@ -10,6 +10,11 @@
 #include "ActorFactories/ActorFactoryBlueprint.h"
 #include "IVPUtilitiesEditorModule.h"
 #include "SimpleVirtualCamera.h"
+#include "VirtualCameraActor.h"
+#include "LevelEditor.h"
+#include "LevelEditorOutlinerSettings.h"
+#include "Filters/CustomClassFilterData.h"
+
 
 
 #define LOCTEXT_NAMESPACE "FVirtualCameraEditorModule"
@@ -21,6 +26,7 @@ public:
 	{
 		RegisterSettings();
 		RegisterPlacementModeItems();
+		RegisterOutlinerFilters();
 	}
 
 	virtual void ShutdownModule() override
@@ -75,6 +81,17 @@ public:
 			}
 		}
 	}
+
+	void RegisterOutlinerFilters()
+	{
+		FLevelEditorModule& LevelEditorModule = FModuleManager::LoadModuleChecked<FLevelEditorModule>(TEXT("LevelEditor"));
+		
+		if(TSharedPtr<FFilterCategory> VPFilterCategory = LevelEditorModule.GetOutlinerFilterCategory(FLevelEditorOutlinerBuiltInCategories::VirtualProduction()))
+		{
+			TSharedRef<FCustomClassFilterData> CineCameraActorClassData = MakeShared<FCustomClassFilterData>(ACineCameraActor::StaticClass(), VPFilterCategory, FLinearColor::White);
+			LevelEditorModule.AddCustomClassFilterToOutliner(CineCameraActorClassData);
+		}
+	}
 };
 #undef LOCTEXT_NAMESPACE
 
diff --git a/Engine/Plugins/Experimental/VirtualCamera/Source/VirtualCameraEditor/VirtualCameraEditor.Build.cs b/Engine/Plugins/Experimental/VirtualCamera/Source/VirtualCameraEditor/VirtualCameraEditor.Build.cs
--- a/Engine/Plugins/Experimental/VirtualCamera/Source/VirtualCameraEditor/VirtualCameraEditor.Build.cs
+++ b/Engine/Plugins/Experimental/VirtualCamera/Source/VirtualCameraEditor/VirtualCameraEditor.Build.cs
@@ -20,7 +20,16 @@ public class VirtualCameraEditor : ModuleRules
 				"VPUtilities",
 				"VPUtilitiesEditor",
 				"VCamCore",
-				"VirtualCamera"
+				"VirtualCamera",
+				"EditorWidgets",
+				"CinematicCamera"
+			}
+		);
+			
+		PrivateIncludePathModuleNames.AddRange(
+			new string[]
+			{
+				"LevelEditor",
 			}
 		);
 	}
diff --git a/Engine/Plugins/FX/Niagara/Source/NiagaraEditor/Private/NiagaraEditorModule.cpp b/Engine/Plugins/FX/Niagara/Source/NiagaraEditor/Private/NiagaraEditorModule.cpp
--- a/Engine/Plugins/FX/Niagara/Source/NiagaraEditor/Private/NiagaraEditorModule.cpp
+++ b/Engine/Plugins/FX/Niagara/Source/NiagaraEditor/Private/NiagaraEditorModule.cpp
@@ -145,6 +145,10 @@
 #include "NiagaraEffectType.h"
 #include "SNiagaraSystemViewport.h"
 
+#include "LevelEditor.h"
+#include "LevelEditorOutlinerSettings.h"
+#include "Filters/CustomClassFilterData.h"
+
 #include "SNiagaraDebugger.h"
 
 #include "NiagaraDebugVis.h"
@@ -1409,6 +1413,20 @@ void FNiagaraEditorModule::OnPostEngineInit()
 	ParameterCollectionAssetCache.RefreshCache(true /*bAllowLoading*/);
 	ParameterDefinitionsAssetCache.RefreshCache(true /*bAllowLoading*/);
 
+	if (FModuleManager::Get().IsModuleLoaded("LevelEditor"))
+	{
+		FLevelEditorModule& LevelEditorModule = FModuleManager::GetModuleChecked<FLevelEditorModule>(TEXT("LevelEditor"));
+		
+		TSharedPtr<FFilterCategory> FXFilterCategory = MakeShared<FFilterCategory>(LOCTEXT("FXFilterCategory", "FX"), FText::GetEmpty());
+		TSharedRef<FCustomClassFilterData> NewClassData = MakeShared<FCustomClassFilterData>(ANiagaraActor::StaticClass(), FXFilterCategory, FLinearColor::White);
+
+		if(TSharedPtr<FFilterCategory> EssentialCategory = LevelEditorModule.GetOutlinerFilterCategory(FLevelEditorOutlinerBuiltInCategories::Common()))
+		{
+			NewClassData->AddCategory(EssentialCategory);
+		}
+		LevelEditorModule.AddCustomClassFilterToOutliner(NewClassData);
+	}
+
 #if WITH_NIAGARA_DEBUGGER
 	Debugger = MakeShared<FNiagaraDebugger>();
 	Debugger->Init();
diff --git a/Engine/Source/Developer/ToolWidgets/Public/Filters/SBasicFilterBar.h b/Engine/Source/Developer/ToolWidgets/Public/Filters/SBasicFilterBar.h
--- a/Engine/Source/Developer/ToolWidgets/Public/Filters/SBasicFilterBar.h
+++ b/Engine/Source/Developer/ToolWidgets/Public/Filters/SBasicFilterBar.h
@@ -1524,7 +1524,7 @@ protected:
 	/** Helper function to add all custom filters to the Add Filter Menu */
 	void PopulateCustomFilters(UToolMenu* Menu)
 	{
-		FToolMenuSection& Section = Menu->AddSection("BasicFilterBarFiltersMenu");
+		FToolMenuSection& Section = Menu->AddSection("BasicFilterBarFiltersMenu", LOCTEXT("FilterBarOtherFiltersSection", "Other Filters"));
 		
 		// Add all the filters
 		for (const TSharedPtr<FFilterCategory>& Category : AllFilterCategories)
diff --git a/Engine/Source/Editor/EditorWidgets/Public/Filters/AssetFilter.h b/Engine/Source/Editor/EditorWidgets/Public/Filters/AssetFilter.h
--- a/Engine/Source/Editor/EditorWidgets/Public/Filters/AssetFilter.h
+++ b/Engine/Source/Editor/EditorWidgets/Public/Filters/AssetFilter.h
@@ -113,6 +113,12 @@ public:
 		// If not, try to use the comparison function
 		else if(OnCompareItem.IsBound())
 		{
+			// If there are no filters active just return true
+			if(CompiledBackendFilter.IsEmpty())
+			{
+				return true;
+			}
+			
 			return OnCompareItem.Execute(InItem, CompiledBackendFilter.ClassPaths);
 		}
 
diff --git a/Engine/Source/Editor/EditorWidgets/Public/Filters/SAssetFilterBar.h b/Engine/Source/Editor/EditorWidgets/Public/Filters/SAssetFilterBar.h
--- a/Engine/Source/Editor/EditorWidgets/Public/Filters/SAssetFilterBar.h
+++ b/Engine/Source/Editor/EditorWidgets/Public/Filters/SAssetFilterBar.h
@@ -804,8 +804,12 @@ protected:
 		FAssetToolsModule& AssetToolsModule = FModuleManager::LoadModuleChecked<FAssetToolsModule>(TEXT("AssetTools"));
 
 		// Add the Basic category
-		AssetFilterCategories.Add(EAssetTypeCategories::Basic, MakeShareable(new FFilterCategory(LOCTEXT("BasicFilter", "Basic"), LOCTEXT("BasicFilterTooltip", "Filter by basic assets."))));
+		TSharedPtr<FFilterCategory> BasicCategory = MakeShareable(new FFilterCategory(LOCTEXT("BasicFilter", "Common"), LOCTEXT("BasicFilterTooltip", "Filter by Common assets.")));
+		AssetFilterCategories.Add(EAssetTypeCategories::Basic, BasicCategory);
 
+		// We expand the basic category by default if we are creating it
+		CategoryToExpand = BasicCategory;
+		
 		// Add the advanced categories
 		TArray<FAdvancedAssetCategory> AdvancedAssetCategories;
 		AssetToolsModule.Get().GetAllAdvancedAssetCategories(/*out*/ AdvancedAssetCategories);
@@ -929,9 +933,10 @@ protected:
 		FilterBarContext->PopulateFilterMenu = FOnPopulateAddAssetFilterMenu::CreateSP(this, &SAssetFilterBar<FilterType>::PopulateAddFilterMenu);
 
 		/** Auto expand the Basic Category if it is present */
-		if(TSharedPtr<FFilterCategory>* BasicCategory = AssetFilterCategories.Find(EAssetTypeCategories::Basic))
+
+		if(CategoryToExpand)
 		{
-			FilterBarContext->MenuExpansion = *BasicCategory;
+			FilterBarContext->MenuExpansion = CategoryToExpand;
 		}
 	
 		FToolMenuContext ToolMenuContext(FilterBarContext);
@@ -1032,22 +1037,22 @@ protected:
 			if(ExpandedCategory)
 			{
 				FToolMenuSection& Section = Menu->AddSection(ExpandedCategory->SectionExtensionHook, ExpandedCategory->SectionHeading);
-				if(MenuExpansion == AssetFilterCategories.FindChecked(EAssetTypeCategories::Basic))
-				{
-					// If we are doing a full menu (i.e expanding basic) we add a menu entry which toggles all other categories
-					Section.AddMenuEntry(
-						NAME_None,
-						MenuExpansion->Title,
-						MenuExpansion->Tooltip,
-						FSlateIcon(FAppStyle::Get().GetStyleSetName(), "PlacementBrowser.Icons.Basic"),
-						FUIAction(
-						FExecuteAction::CreateSP( this, &SAssetFilterBar<FilterType>::FilterByTypeCategoryClicked, MenuExpansion, ExpandedCategory->Classes ),
-						FCanExecuteAction(),
-						FGetActionCheckState::CreateSP(this, &SAssetFilterBar<FilterType>::IsTypeCategoryChecked, MenuExpansion, ExpandedCategory->Classes ) ),
-						EUserInterfaceActionType::ToggleButton
-						);
-				}
-
+				
+				// If we are doing a full menu (i.e expanding basic) we add a menu entry which toggles all other categories
+                Section.AddMenuEntry(
+                	NAME_None,
+                	ExpandedCategory->SectionHeading,
+                	MenuExpansion->Tooltip,
+                	FSlateIcon(FAppStyle::Get().GetStyleSetName(), "PlacementBrowser.Icons.Basic"),
+                	FUIAction(
+                	FExecuteAction::CreateSP( this, &SAssetFilterBar<FilterType>::FilterByTypeCategoryClicked, MenuExpansion, ExpandedCategory->Classes ),
+                	FCanExecuteAction(),
+                	FGetActionCheckState::CreateSP(this, &SAssetFilterBar<FilterType>::IsTypeCategoryChecked, MenuExpansion, ExpandedCategory->Classes ) ),
+                	EUserInterfaceActionType::ToggleButton
+                	);
+
+				Section.AddSeparator("ExpandedCategorySeparator");
+				
 				// Now populate with all the assets from the expanded category
 				SAssetFilterBar<FilterType>::CreateFiltersMenuCategory( Section, ExpandedCategory->Classes);
 				
@@ -1056,36 +1061,28 @@ protected:
 			}
 		}
 
-		TSharedPtr<FFilterCategory>* BasicCategory = AssetFilterCategories.Find(EAssetTypeCategories::Basic);
-
-		// We are in Full Menu Mode if there is no menu expansion, or the menu expansion is EAssetTypeCategories::Basic
-		bool bInFullMenuMode = !MenuExpansion || (BasicCategory && MenuExpansion == *BasicCategory);
-
-		// If we are in full menu mode, add all the other categories as submenus
-		if(bInFullMenuMode)
-		{
-			FToolMenuSection& Section = Menu->AddSection("AssetFilterBarFilterAdvancedAsset", LOCTEXT("AdvancedAssetsMenuHeading", "Other Assets"));
+		// Add all the other categories as submenus
+		FToolMenuSection& Section = Menu->AddSection("AssetFilterBarFilterAdvancedAsset", LOCTEXT("AdvancedAssetsMenuHeading", "Type Filters"));
 			
-			// Sort by category name so that we add the submenus in alphabetical order
-			CategoryToMenuMap.KeySort([](const TSharedPtr<FFilterCategory>& A, const TSharedPtr<FFilterCategory>& B) {
-				return A->Title.CompareTo(B->Title) < 0;
-			});
+		// Sort by category name so that we add the submenus in alphabetical order
+		CategoryToMenuMap.KeySort([](const TSharedPtr<FFilterCategory>& A, const TSharedPtr<FFilterCategory>& B) {
+			return A->Title.CompareTo(B->Title) < 0;
+		});
 
-			// For all the remaining categories, add them as submenus
-			for (const TPair<TSharedPtr<FFilterCategory>, FCategoryMenu>& CategoryMenuPair : CategoryToMenuMap)
-			{
-				Section.AddSubMenu(
-					NAME_None,
-					CategoryMenuPair.Key->Title,
-					CategoryMenuPair.Key->Tooltip,
-					FNewToolMenuDelegate::CreateSP(this, &SAssetFilterBar<FilterType>::CreateFiltersMenuCategory, CategoryMenuPair.Value.Classes),
-					FUIAction(
-					FExecuteAction::CreateSP(this, &SAssetFilterBar<FilterType>::FilterByTypeCategoryClicked, CategoryMenuPair.Key, CategoryMenuPair.Value.Classes),
-					FCanExecuteAction(),
-					FGetActionCheckState::CreateSP(this, &SAssetFilterBar<FilterType>::IsTypeCategoryChecked, CategoryMenuPair.Key, CategoryMenuPair.Value.Classes)),
-					EUserInterfaceActionType::ToggleButton
-					);
-			}
+		// Now actually add them
+		for (const TPair<TSharedPtr<FFilterCategory>, FCategoryMenu>& CategoryMenuPair : CategoryToMenuMap)
+		{
+			Section.AddSubMenu(
+				NAME_None,
+				CategoryMenuPair.Key->Title,
+				CategoryMenuPair.Key->Tooltip,
+				FNewToolMenuDelegate::CreateSP(this, &SAssetFilterBar<FilterType>::CreateFiltersMenuCategory, CategoryMenuPair.Value.Classes),
+				FUIAction(
+				FExecuteAction::CreateSP(this, &SAssetFilterBar<FilterType>::FilterByTypeCategoryClicked, CategoryMenuPair.Key, CategoryMenuPair.Value.Classes),
+				FCanExecuteAction(),
+				FGetActionCheckState::CreateSP(this, &SAssetFilterBar<FilterType>::IsTypeCategoryChecked, CategoryMenuPair.Key, CategoryMenuPair.Value.Classes)),
+				EUserInterfaceActionType::ToggleButton
+				);
 		}
 
 		// Now add all non-asset filters
@@ -1211,6 +1208,9 @@ protected:
 	TArray<TSharedRef<FCustomClassFilterData>> UserCustomClassFilters;
 
 	TMap<EAssetTypeCategories::Type, TSharedPtr<FFilterCategory>> AssetFilterCategories;
+
+	/** The category to expand in the filter menu */
+	TSharedPtr<FFilterCategory> CategoryToExpand;
 	
 	/** Whether the filter bar provides the default Asset Filters */
 	bool bUseDefaultAssetFilters;
diff --git a/Engine/Source/Editor/LandscapeEditor/Private/LandscapeEditorModule.cpp b/Engine/Source/Editor/LandscapeEditor/Private/LandscapeEditorModule.cpp
--- a/Engine/Source/Editor/LandscapeEditor/Private/LandscapeEditorModule.cpp
+++ b/Engine/Source/Editor/LandscapeEditor/Private/LandscapeEditorModule.cpp
@@ -31,6 +31,7 @@
 #include "LandscapeModule.h"
 
 #include "LevelEditor.h"
+#include "Filters/CustomClassFilterData.h"
 #include "ToolMenus.h"
 #include "Editor/EditorEngine.h"
 #include "LandscapeSubsystem.h"
@@ -103,6 +104,13 @@ public:
 		FLevelEditorModule& LevelEditorModule = FModuleManager::LoadModuleChecked<FLevelEditorModule>("LevelEditor");
 		LevelEditorModule.GetMenuExtensibilityManager()->AddExtender(ViewportMenuExtender);
 
+		// Add Level Editor Outliner Filter
+		if (TSharedPtr<FFilterCategory> EnvironmentFilterCategory = LevelEditorModule.GetOutlinerFilterCategory(FLevelEditorOutlinerBuiltInCategories::Environment()))
+		{
+			TSharedRef<FCustomClassFilterData> LandscapeActorClassData = MakeShared<FCustomClassFilterData>(ALandscape::StaticClass(), EnvironmentFilterCategory, FLinearColor::White);
+			LevelEditorModule.AddCustomClassFilterToOutliner(LandscapeActorClassData);
+		}
+		
 		// add actor factories
 		UActorFactoryLandscape* LandscapeActorFactory = NewObject<UActorFactoryLandscape>();
 		LandscapeActorFactory->NewActorClass = ALandscape::StaticClass();
diff --git a/Engine/Source/Editor/LevelEditor/Private/LevelEditor.cpp b/Engine/Source/Editor/LevelEditor/Private/LevelEditor.cpp
--- a/Engine/Source/Editor/LevelEditor/Private/LevelEditor.cpp
+++ b/Engine/Source/Editor/LevelEditor/Private/LevelEditor.cpp
@@ -42,6 +42,7 @@
 #include "Framework/Commands/GenericCommands.h"
 #include "Misc/EngineBuildSettings.h"
 #include "Subsystems/PanelExtensionSubsystem.h"
+#include "LevelEditorOutlinerSettings.h"
 
 #define LOCTEXT_NAMESPACE "LevelEditor"
 
@@ -208,6 +209,8 @@ TSharedRef<SDockTab> FLevelEditorModule::SpawnLevelEditor( const FSpawnTabArgs&
 		OwnerWindow = MainFrameModule.GetParentWindow();
 	}
 
+	OutlinerSettings->CreateDefaultFilters();
+
 	TSharedPtr<SLevelEditor> LevelEditorTmp;
 	if (OwnerWindow.IsValid())
 	{
@@ -294,6 +297,9 @@ void FLevelEditorModule::StartupModule()
 
 	FMessageLogModule& MessageLogModule = FModuleManager::LoadModuleChecked<FMessageLogModule>("MessageLog");
 	MessageLogModule.RegisterLogListing("BuildAndSubmitErrors", LOCTEXT("BuildAndSubmitErrors", "Build and Submit Errors"));
+
+	OutlinerSettings = MakeShared<FLevelEditorOutlinerSettings>();
+	OutlinerSettings->SetupBuiltInCategories();
 }
 
 /**
@@ -1829,5 +1835,24 @@ void FLevelEditorModule::BindGlobalLevelEditorCommands()
 		);
 }
 
+TSharedPtr<FLevelEditorOutlinerSettings> FLevelEditorModule::GetLevelEditorOutlinerSettings() const
+{
+	return OutlinerSettings;
+}
+
+void FLevelEditorModule::AddCustomFilterToOutliner(TSharedRef<FFilterBase<SceneOutliner::FilterBarType>> InCustomFilter)
+{
+	OutlinerSettings->AddCustomFilter(InCustomFilter);
+}
+
+void FLevelEditorModule::AddCustomClassFilterToOutliner(TSharedRef<FCustomClassFilterData> InCustomClassFilterData)
+{
+	OutlinerSettings->AddCustomClassFilter(InCustomClassFilterData);
+}
+
+TSharedPtr<FFilterCategory> FLevelEditorModule::GetOutlinerFilterCategory(const FName& CategoryName) const
+{
+	return OutlinerSettings->GetFilterCategory(CategoryName);
+}
 
 #undef LOCTEXT_NAMESPACE
diff --git a/Engine/Source/Editor/LevelEditor/Private/LevelEditorOutlinerSettings.cpp b/Engine/Source/Editor/LevelEditor/Private/LevelEditorOutlinerSettings.cpp
new file mode 100644
--- /dev/null
+++ b/Engine/Source/Editor/LevelEditor/Private/LevelEditorOutlinerSettings.cpp
@@ -0,0 +1,253 @@
+// Copyright Epic Games, Inc. All Rights Reserved.
+
+#include "LevelEditorOutlinerSettings.h"
+
+#include "Filters/FilterBase.h"
+#include "Filters/CustomClassFilterData.h"
+
+#include "SceneOutlinerPublicTypes.h"
+#include "IPlacementModeModule.h"
+
+#include "Engine/StaticMeshActor.h"
+#include "Engine/Blueprint.h"
+#include "LevelInstance/LevelInstanceActor.h"
+#include "Animation/SkeletalMeshActor.h"
+#include "Engine/PostProcessVolume.h"
+#include "Engine/Brush.h"
+#include "Engine/ExponentialHeightFog.h"
+#include "InstancedFoliageActor.h"
+#include "Landscape.h"
+#include "Engine/Classes/Components/SkyAtmosphereComponent.h"
+#include "Engine/Classes/Components/VolumetricCloudComponent.h"
+#include "Engine/Classes/Sound/AmbientSound.h"
+
+#define LOCTEXT_NAMESPACE "LevelEditorOutlinerSettings"
+
+void FLevelEditorOutlinerSettings::AddCustomFilter(TSharedRef<FFilterBase<SceneOutliner::FilterBarType>> InCustomFilter)
+{
+	CustomFilters.Add(InCustomFilter);
+}
+
+void FLevelEditorOutlinerSettings::AddCustomClassFilter(TSharedRef<FCustomClassFilterData> InCustomClassFilterData)
+{
+	TSharedRef<FCustomClassFilterData>* FoundFilter = CustomClassFilters.FindByPredicate([&InCustomClassFilterData](const TSharedPtr<FCustomClassFilterData> FilterData)
+	{
+		return FilterData->GetClassPathName() == InCustomClassFilterData->GetClassPathName();
+	});
+
+	// If a filter for the class already exists, just add the new categories to it
+	// TODO: Not perfect, what if the user wants to use InCustomClassFilterData after if there is already a dupe?
+	if(FoundFilter)
+	{
+		TArray<TSharedPtr<FFilterCategory>> Categories = InCustomClassFilterData->GetCategories();
+
+		for(const TSharedPtr<FFilterCategory>& Category : Categories)
+		{
+			(*FoundFilter)->AddCategory(Category);
+		}
+	}
+	else
+	{
+		CustomClassFilters.Add(InCustomClassFilterData);
+	}
+}
+
+TSharedPtr<FFilterCategory> FLevelEditorOutlinerSettings::GetFilterCategory(const FName& CategoryName)
+{
+	if(TSharedPtr<FFilterCategory>* FoundCategory = FilterBarCategories.Find(CategoryName))
+	{
+		return *FoundCategory;
+	}
+
+	return nullptr;
+}
+
+void FLevelEditorOutlinerSettings::SetupBuiltInCategories()
+{
+	// First setup our unique built in categories
+	TSharedPtr<FFilterCategory> CommonFiltersCategory = MakeShared<FFilterCategory>(LOCTEXT("CommonFiltersCategory", "Common"), FText::GetEmpty());
+	FilterBarCategories.Add(FLevelEditorOutlinerBuiltInCategories::Common(), CommonFiltersCategory);
+
+	TSharedPtr<FFilterCategory> AudioFiltersCategory = MakeShared<FFilterCategory>(LOCTEXT("AudioFiltersCategory", "Audio"), FText::GetEmpty());
+	FilterBarCategories.Add(FLevelEditorOutlinerBuiltInCategories::Audio(), AudioFiltersCategory);
+
+	TSharedPtr<FFilterCategory> AnimationFiltersCategory = MakeShared<FFilterCategory>(LOCTEXT("AnimationFiltersCategory", "Animation"), FText::GetEmpty());
+	FilterBarCategories.Add(FLevelEditorOutlinerBuiltInCategories::Animation(), AnimationFiltersCategory);
+	
+	TSharedPtr<FFilterCategory> GeometryFiltersCategory = MakeShared<FFilterCategory>(LOCTEXT("GeometryFiltersCategory", "Geometry"), FText::GetEmpty());
+	FilterBarCategories.Add(FLevelEditorOutlinerBuiltInCategories::Geometry(), GeometryFiltersCategory);
+
+	TSharedPtr<FFilterCategory> EnvironmentFiltersCategory = MakeShared<FFilterCategory>(LOCTEXT("EnvironmentFiltersCategory", "Environment"), FText::GetEmpty());
+	FilterBarCategories.Add(FLevelEditorOutlinerBuiltInCategories::Environment(), EnvironmentFiltersCategory);
+	
+	TSharedPtr<FFilterCategory> VPFiltersCategory = MakeShared<FFilterCategory>(LOCTEXT("VPFiltersCategory", "Virtual Production"), FText::GetEmpty());
+	FilterBarCategories.Add(FLevelEditorOutlinerBuiltInCategories::VirtualProduction(), VPFiltersCategory);
+	
+	// Now convert some of the built in placement mode categories we want to filter categories and add them
+
+	PlacementToFilterCategoryMap.Add(FBuiltInPlacementCategories::Basic(), FLevelEditorOutlinerBuiltInCategories::Basic());
+	TSharedPtr<FFilterCategory> BasicFilterCategory = MakeShared<FFilterCategory>(LOCTEXT("BasicFilterCategory", "Basic"), FText::GetEmpty());
+	FilterBarCategories.Add(FLevelEditorOutlinerBuiltInCategories::Basic(), BasicFilterCategory);
+
+	PlacementToFilterCategoryMap.Add(FBuiltInPlacementCategories::Lights(), FLevelEditorOutlinerBuiltInCategories::Lights());
+	TSharedPtr<FFilterCategory> LightsFilterCategory = MakeShared<FFilterCategory>(LOCTEXT("LightsFilterCategory", "Lights"), FText::GetEmpty());
+	FilterBarCategories.Add(FLevelEditorOutlinerBuiltInCategories::Lights(), LightsFilterCategory);
+
+	PlacementToFilterCategoryMap.Add(FBuiltInPlacementCategories::Visual(), FLevelEditorOutlinerBuiltInCategories::Visual());
+	TSharedPtr<FFilterCategory> VisualFilterCategory = MakeShared<FFilterCategory>(LOCTEXT("VisualFilterCategory", "Visual"), FText::GetEmpty());
+	FilterBarCategories.Add(FLevelEditorOutlinerBuiltInCategories::Visual(), VisualFilterCategory);
+
+	PlacementToFilterCategoryMap.Add(FBuiltInPlacementCategories::Volumes(), FLevelEditorOutlinerBuiltInCategories::Volumes());
+	TSharedPtr<FFilterCategory> VolumesFilterCategory = MakeShared<FFilterCategory>(LOCTEXT("VolumnesFilterCategory", "Volumes"), FText::GetEmpty());
+	FilterBarCategories.Add(FLevelEditorOutlinerBuiltInCategories::Volumes(), VolumesFilterCategory);
+}
+
+void FLevelEditorOutlinerSettings::CreateDefaultFilters()
+{
+	// First we will add all items registered to the place actors panel as filters
+	IPlacementModeModule& PlacementModeModule = IPlacementModeModule::Get();
+	
+	// Get all the Categories
+	TArray<FPlacementCategoryInfo> Categories;
+	PlacementModeModule.GetSortedCategories(Categories);
+
+	// Remove the Recently Placed, Shapes (don't make sense) and All Classes (too much bloat) categories from the filters
+	Categories.RemoveAll([](const FPlacementCategoryInfo& Category)
+	{
+		return Category.UniqueHandle == FBuiltInPlacementCategories::RecentlyPlaced() || Category.UniqueHandle == FBuiltInPlacementCategories::AllClasses()
+			|| Category.UniqueHandle == FBuiltInPlacementCategories::Shapes();
+	});
+
+	auto FindOrAddClassFilter = [this](UClass* Class, TSharedPtr<FFilterCategory> FilterCategory)
+	{
+		// If the underlying class already exists, just add this category to it
+		TSharedRef<FCustomClassFilterData>* ExistingClassData = CustomClassFilters.FindByPredicate([Class](const TSharedPtr<FCustomClassFilterData> FilterData)
+		{
+			return FilterData->GetClassPathName() == Class->GetClassPathName();
+		});
+		
+		if(ExistingClassData)
+		{
+			(*ExistingClassData)->AddCategory(FilterCategory);
+		}
+		else
+		{
+			TSharedRef<FCustomClassFilterData> NewClassData = MakeShared<FCustomClassFilterData>(Class, FilterCategory, FLinearColor::White);
+			CustomClassFilters.Add(NewClassData);
+		}
+	};
+	
+	for (const FPlacementCategoryInfo& Category : Categories)
+	{
+		// Make an FFilterCategory using the current PlacementCategory if it doesn't already exist (built in)
+		TSharedPtr<FFilterCategory> FilterCategory;
+		if(TSharedPtr<FFilterCategory>* FoundCategory = FilterBarCategories.Find(Category.UniqueHandle))
+		{
+			FilterCategory = *FoundCategory;
+		}
+		else
+		{
+			FilterCategory = MakeShared<FFilterCategory>(Category.DisplayName, FText::GetEmpty());
+		}
+
+		FName CategoryName = Category.UniqueHandle;
+
+		// If it is a built in placement mode category, convert the name to a built in filter category
+		if(FName* MappedCategoryName = PlacementToFilterCategoryMap.Find(CategoryName))
+		{
+			CategoryName = *MappedCategoryName;
+		}
+
+		// Add the category to our list
+		FilterBarCategories.Add(Category.UniqueHandle, FilterCategory);
+
+		// Get all the items belonging to the current category
+		TArray<TSharedPtr<FPlaceableItem>> Items;
+		PlacementModeModule.RegenerateItemsForCategory(Category.UniqueHandle);
+		PlacementModeModule.GetItemsForCategory(Category.UniqueHandle, Items);
+
+		// Add each item to as a filter
+		for(TSharedPtr<FPlaceableItem>& Item : Items)
+		{
+			// Get the underlying class from the Actor belonging to this item
+			const bool bIsClass = Item->AssetData.GetClass() == UClass::StaticClass();
+			const bool bIsActor = bIsClass ? CastChecked<UClass>(Item->AssetData.GetAsset())->IsChildOf(AActor::StaticClass()) : false;
+			AActor* DefaultActor = nullptr;
+			if (Item->Factory != nullptr)
+			{
+				DefaultActor = Item->Factory->GetDefaultActor(Item->AssetData);
+			}
+			else if (bIsActor)
+			{
+				DefaultActor = CastChecked<AActor>(CastChecked<UClass>(Item->AssetData.GetAsset())->ClassDefaultObject);
+			
+			}
+			if(!DefaultActor)
+			{
+				continue;
+			}
+			UClass* Class = DefaultActor->GetClass();
+
+			FindOrAddClassFilter(Class, FilterCategory);
+		}
+	}
+
+	// Now we add some custom filterable types that are not a part of the place actors panel
+
+	TSharedPtr<FFilterCategory>* CommonFilterCategory = FilterBarCategories.Find(FLevelEditorOutlinerBuiltInCategories::Common());
+
+	if(CommonFilterCategory)
+	{
+		FindOrAddClassFilter(AStaticMeshActor::StaticClass(), *CommonFilterCategory);
+		FindOrAddClassFilter(ALevelInstance::StaticClass(), *CommonFilterCategory);
+		FindOrAddClassFilter(ASkeletalMeshActor::StaticClass(), *CommonFilterCategory);
+		FindOrAddClassFilter(UBlueprint::StaticClass(), *CommonFilterCategory);
+		FindOrAddClassFilter(APostProcessVolume::StaticClass(), *CommonFilterCategory);
+	}
+
+	TSharedPtr<FFilterCategory>* GeometryFilterCategory = FilterBarCategories.Find(FLevelEditorOutlinerBuiltInCategories::Geometry());
+
+	if(GeometryFilterCategory)
+	{
+		FindOrAddClassFilter(AStaticMeshActor::StaticClass(), *GeometryFilterCategory);
+		FindOrAddClassFilter(ABrush::StaticClass(), *GeometryFilterCategory);
+	}
+
+	TSharedPtr<FFilterCategory>* AnimationFilterCategory = FilterBarCategories.Find(FLevelEditorOutlinerBuiltInCategories::Animation());
+
+	if(AnimationFilterCategory)
+	{
+		FindOrAddClassFilter(ASkeletalMeshActor::StaticClass(), *AnimationFilterCategory);
+	}
+
+	TSharedPtr<FFilterCategory>* EnvironmentFilterCategory = FilterBarCategories.Find(FLevelEditorOutlinerBuiltInCategories::Environment());
+
+	if(EnvironmentFilterCategory)
+	{
+		FindOrAddClassFilter(AExponentialHeightFog::StaticClass(), *EnvironmentFilterCategory);
+		FindOrAddClassFilter(AInstancedFoliageActor::StaticClass(), *EnvironmentFilterCategory);
+		FindOrAddClassFilter(ASkyAtmosphere::StaticClass(), *EnvironmentFilterCategory);
+		FindOrAddClassFilter(AVolumetricCloud::StaticClass(), *EnvironmentFilterCategory);
+	}
+	
+	TSharedPtr<FFilterCategory>* AudioFilterCategory = FilterBarCategories.Find(FLevelEditorOutlinerBuiltInCategories::Audio());
+
+	if(AudioFilterCategory)
+	{
+		FindOrAddClassFilter(AAmbientSound::StaticClass(), *AudioFilterCategory);
+	}
+}
+
+void FLevelEditorOutlinerSettings::GetOutlinerFilters(FSceneOutlinerFilterBarOptions& OutFilterBarOptions)
+{
+	// Sort the type filters by name
+	CustomClassFilters.Sort([](const TSharedRef<FCustomClassFilterData>& ClassA, const TSharedRef<FCustomClassFilterData>& ClassB)
+		{
+			return ClassA->GetName().CompareTo(ClassB->GetName()) < 0;
+		});
+	
+	OutFilterBarOptions.CustomClassFilters.Append(CustomClassFilters);
+	OutFilterBarOptions.CustomFilters.Append(CustomFilters);
+}
+
+#undef LOCTEXT_NAMESPACE
\ No newline at end of file
diff --git a/Engine/Source/Editor/LevelEditor/Private/SLevelEditor.cpp b/Engine/Source/Editor/LevelEditor/Private/SLevelEditor.cpp
--- a/Engine/Source/Editor/LevelEditor/Private/SLevelEditor.cpp
+++ b/Engine/Source/Editor/LevelEditor/Private/SLevelEditor.cpp
@@ -72,6 +72,7 @@
 #include "EditorViewportLayout.h"
 #include "LevelViewportTabContent.h"
 #include "SLevelViewport.h"
+#include "LevelEditorOutlinerSettings.h"
 
 #define LOCTEXT_NAMESPACE "SLevelEditor"
 
@@ -799,6 +800,17 @@ TSharedRef<ISceneOutliner> SLevelEditor::CreateSceneOutliner(FName TabIdentifier
 	}
 
 	InitOptions.OutlinerIdentifier = TabIdentifier;
+	InitOptions.FilterBarOptions.bHasFilterBar = true;
+	
+	// All level editor outliners share their custom text filters
+	InitOptions.FilterBarOptions.bUseSharedSettings = true;
+
+	FLevelEditorModule& LevelEditorModule = FModuleManager::GetModuleChecked<FLevelEditorModule>(LevelEditorModuleName);
+	
+	TSharedPtr<FLevelEditorOutlinerSettings> OutlinerSettings = LevelEditorModule.GetLevelEditorOutlinerSettings();
+	OutlinerSettings->GetOutlinerFilters(InitOptions.FilterBarOptions);
+
+	InitOptions.FilterBarOptions.CategoryToExpand = OutlinerSettings->GetFilterCategory(FLevelEditorOutlinerBuiltInCategories::Common());
 
 	FSceneOutlinerModule& SceneOutlinerModule = FModuleManager::Get().LoadModuleChecked<FSceneOutlinerModule>("SceneOutliner");
 	TSharedRef<ISceneOutliner> SceneOutlinerRef = SceneOutlinerModule.CreateActorBrowser(
diff --git a/Engine/Source/Editor/LevelEditor/Public/LevelEditor.h b/Engine/Source/Editor/LevelEditor/Public/LevelEditor.h
--- a/Engine/Source/Editor/LevelEditor/Public/LevelEditor.h
+++ b/Engine/Source/Editor/LevelEditor/Public/LevelEditor.h
@@ -12,6 +12,7 @@
 #include "ILevelEditor.h"
 #include "Toolkits/AssetEditorToolkit.h"
 #include "ViewportTypeDefinition.h"
+#include "LevelEditorOutlinerSettings.h"
 
 class AActor;
 class IAssetViewport;
@@ -336,6 +337,25 @@ public:
 	/** Delegate used to capture skeltal meshes to single-frame animations when 'keeping simulation changes' */
 	DECLARE_DELEGATE_RetVal_OneParam(UAnimSequence*, FCaptureSingleFrameAnimSequence, USkeletalMeshComponent* /*Component*/);
 	virtual FCaptureSingleFrameAnimSequence& OnCaptureSingleFrameAnimSequence() { return CaptureSingleFrameAnimSequenceDelegate; }
+	
+	/**  Add a custom filter to the outliner filter bar. These are all AND'd together
+	 * @see FLevelEditorOutlinerSettings
+	 *   @see FGenericFilter on how to create generic filters
+	 *   NOTE: Currently does not support adding filters dynamically, they must be added before the Level Editor is init
+	 */
+	virtual void AddCustomFilterToOutliner(TSharedRef<FFilterBase<SceneOutliner::FilterBarType>> InCustomFilter);
+
+	/**  Add a custom class filter to the outliner filter bar. These represent asset/actor type filters and are OR'd
+	 *   @see FLevelEditorOutlinerSettings
+	 *   Can be created using IAssetTypeActions or UClass (@see constructor)
+	 *   NOTE: Currently does not support adding filters dynamically, they must be added before the Level Editor is init
+	 */
+	virtual void AddCustomClassFilterToOutliner(TSharedRef<FCustomClassFilterData> InCustomClassFilterData);
+	
+	/** Get the FFilterCategory attached to the given category name. Use this to add filters to the built in categories
+	 *  @see FLevelEditorOutlinerBuiltInCategories
+	 */
+	virtual TSharedPtr<FFilterCategory> GetOutlinerFilterCategory(const FName& CategoryName) const;
 
 public:
 
@@ -365,6 +385,9 @@ public:
 		return true;
 	}
 
+	/** Return the settings (containing the outliner filters) used to create the level editor's outliners  */
+	TSharedPtr<FLevelEditorOutlinerSettings> GetLevelEditorOutlinerSettings() const;
+
 public:
 	
 	/** Register a viewport type for the level editor */
@@ -495,4 +518,7 @@ private:
 
 	/** Array of delegates that are used to check if the specified objects should be editable on the details panel */
 	TArray<FAreObjectsEditable> AreObjectsEditableDelegates;
+
+	/** The settings used to create any Outliners in the Level Editor */
+	TSharedPtr<FLevelEditorOutlinerSettings> OutlinerSettings;
 };
diff --git a/Engine/Source/Editor/LevelEditor/Public/LevelEditorOutlinerSettings.h b/Engine/Source/Editor/LevelEditor/Public/LevelEditorOutlinerSettings.h
new file mode 100644
--- /dev/null
+++ b/Engine/Source/Editor/LevelEditor/Public/LevelEditorOutlinerSettings.h
@@ -0,0 +1,80 @@
+// Copyright Epic Games, Inc. All Rights Reserved.
+
+#pragma once
+
+#include "Filters/GenericFilter.h"
+
+struct ISceneOutlinerTreeItem;
+
+namespace SceneOutliner
+{
+	/** The type of item that the Outliner's Filter Bar operates on */
+	typedef const ISceneOutlinerTreeItem& FilterBarType;
+}
+
+struct FSceneOutlinerFilterBarOptions;
+class FCustomClassFilterData;
+
+/** Structure of built-in filter categories. Defined as functions to enable external use without linkage */
+struct FLevelEditorOutlinerBuiltInCategories
+{
+	static FName Common()			{ static FName Name("Essential");		return Name; }
+	static FName Basic()			 { static FName Name("Basic");			return Name; }
+	static FName Animation()		 { static FName Name("Animation");		return Name; }
+	static FName Audio()			 { static FName Name("Audio");			return Name; }
+	static FName Geometry()			 { static FName Name("Geometry");		return Name; }
+	static FName Lights()			 { static FName Name("Lights");			return Name; }
+	static FName Environment()		 { static FName Name("Environment");	return Name; }
+	static FName Visual()			 { static FName Name("Visual");			return Name; }
+	static FName Volumes()			 { static FName Name("Volumes");		return Name; }
+	static FName VirtualProduction() { static FName Name("VP");				return Name; }
+};
+
+/** Helper class to manage initalization options specific to the Level Editor Outliners
+ * Use AddCustomFilter/AddCustomClass Filter to register new filters in the Outliner.
+ * Make sure that the filters you attach have a category, otherwise they will not show up
+ */
+class LEVELEDITOR_API FLevelEditorOutlinerSettings :  public TSharedFromThis<class FLevelEditorOutlinerSettings>
+{
+public:
+
+	/**  Add a custom filter to the outliner filter bar. These are all AND'd together
+	 *   @see FGenericFilter on how to create generic filters
+	 */
+	void AddCustomFilter(TSharedRef<FFilterBase<SceneOutliner::FilterBarType>> InCustomFilter);
+
+	/**  Add a custom class filter to the outliner filter bar. These represent asset/actor type filters and are OR'd
+	 *   Can be created using IAssetTypeActions or UClass (@see constructor)
+	 */
+	void AddCustomClassFilter(TSharedRef<FCustomClassFilterData> InCustomClassFilterData);
+
+	// Creates the default filters that the level editor outliner has
+	void CreateDefaultFilters();
+
+	// Setup the built in filter categories
+	void SetupBuiltInCategories();
+
+	// Append the init options stored in this class to the given Outliner init options
+	void GetOutlinerFilters(FSceneOutlinerFilterBarOptions& OutFilterBarOptions);
+
+	// Get the FFilterCategory attached to the given category name. Use this to add filters to the built in categories
+	TSharedPtr<FFilterCategory> GetFilterCategory(const FName& CategoryName);
+
+private:
+
+	/** These are the custom filters that the Scene Outliner will have. All active filters will be AND'd together to test
+	 *  against.
+	 */
+	TArray<TSharedRef<FFilterBase<SceneOutliner::FilterBarType>>> CustomFilters;
+
+	/** These are the asset type filters that the Scene Outliner will have. All active filters will be OR'd together to
+	 *  test against.
+	 */
+	TArray<TSharedRef<FCustomClassFilterData>> CustomClassFilters;
+
+	/* A map of the categories the Outliner filter bar will have */
+	TMap<FName, TSharedPtr<FFilterCategory>> FilterBarCategories;
+
+	/* A map to convert placement mode built in categories to filter bar categories */
+	TMap<FName, FName> PlacementToFilterCategoryMap;
+};
diff --git a/Engine/Source/Editor/SceneOutliner/Private/ActorBrowsingMode.cpp b/Engine/Source/Editor/SceneOutliner/Private/ActorBrowsingMode.cpp
--- a/Engine/Source/Editor/SceneOutliner/Private/ActorBrowsingMode.cpp
+++ b/Engine/Source/Editor/SceneOutliner/Private/ActorBrowsingMode.cpp
@@ -42,10 +42,6 @@
 #include "EditorLevelUtils.h"
 #include "EditorViewportCommands.h"
 #include "SceneOutlinerActorSCCColumn.h"
-#include "Filters/GenericFilter.h"
-#include "Engine/StaticMeshActor.h"
-#include "Filters/CustomClassFilterData.h"
-#include "IPlacementModeModule.h"
 #include "Misc/ScopedSlowTask.h"
 
 DEFINE_LOG_CATEGORY_STATIC(LogActorBrowser, Log, All);
@@ -2227,135 +2223,35 @@ void FActorBrowsingMode::SaveConfig()
 	UActorBrowserConfig::Get()->SaveEditorConfig();
 }
 
-void FActorBrowsingMode::CreateOutlinerFilterBarFilters(FSceneOutlinerFilterBarOptions& OutFilterBarOptions)
-{
-	OutFilterBarOptions.bHasFilterBar = true;
-
-	// All level editor outliners share their custom text filters
-	OutFilterBarOptions.bUseSharedSettings = true;
-
-	IPlacementModeModule& PlacementModeModule = IPlacementModeModule::Get();
-
-	// Get all the Categories
-	TArray<FPlacementCategoryInfo> Categories;
-	PlacementModeModule.GetSortedCategories(Categories);
-
-	// Array of unique class types to Add as filters
-	TMap<FTopLevelAssetPath, TSharedRef<FCustomClassFilterData>> FilterClasses;
-
-	// Remove the Recently Placed (doesn't make sense) and All Classes (too much bloat) categories from the filters
-	Categories.RemoveAll([](const FPlacementCategoryInfo& Category)
-	{
-		return Category.UniqueHandle == FBuiltInPlacementCategories::RecentlyPlaced() || Category.UniqueHandle == FBuiltInPlacementCategories::AllClasses();
-	});
-
-	// We will keep a reference to the shapes category if we find it for later
-	TSharedPtr<FFilterCategory> ShapesCategory;
-	
-	auto FindOrAddClassFilter = [&FilterClasses](UClass* Class, TSharedPtr<FFilterCategory> FilterCategory)
-	{
-		// If the underlying class already exists, just add this category to it
-		if(TSharedRef<FCustomClassFilterData>* ExistingClassData = FilterClasses.Find(Class->GetClassPathName()))
-		{
-			(*ExistingClassData)->AddCategory(FilterCategory);
-		}
-		else
-		{
-			TSharedRef<FCustomClassFilterData> NewClassData = MakeShared<FCustomClassFilterData>(Class, FilterCategory, FLinearColor::White);
-
-			FilterClasses.Add(Class->GetClassPathName(), NewClassData);
-		}
-	};
-	
-	for (const FPlacementCategoryInfo& Category : Categories)
-	{
-		// Make an FFilterCategory using the current PlacementCategory
-		TSharedPtr<FFilterCategory> FilterCategory = MakeShared<FFilterCategory>(Category.DisplayName, FText::GetEmpty());
-
-		if(Category.UniqueHandle == FBuiltInPlacementCategories::Shapes())
-		{
-			ShapesCategory = FilterCategory;
-		}
-
-		// Get all the items belonging to the current category
-		TArray<TSharedPtr<FPlaceableItem>> Items;
-		PlacementModeModule.RegenerateItemsForCategory(Category.UniqueHandle);
-		PlacementModeModule.GetItemsForCategory(Category.UniqueHandle, Items);
-
-		for(TSharedPtr<FPlaceableItem>& Item : Items)
-		{
-			// Get the underlying class from the Actor belonging to this item
-			const bool bIsClass = Item->AssetData.GetClass() == UClass::StaticClass();
-			const bool bIsActor = bIsClass ? CastChecked<UClass>(Item->AssetData.GetAsset())->IsChildOf(AActor::StaticClass()) : false;
-			AActor* DefaultActor = nullptr;
-			if (Item->Factory != nullptr)
-			{
-				DefaultActor = Item->Factory->GetDefaultActor(Item->AssetData);
-			}
-			else if (bIsActor)
-			{
-				DefaultActor = CastChecked<AActor>(CastChecked<UClass>(Item->AssetData.GetAsset())->ClassDefaultObject);
-			
-			}
-			if(!DefaultActor)
-			{
-				continue;
-			}
-			UClass* Class = DefaultActor->GetClass();
-
-			FindOrAddClassFilter(Class, FilterCategory);
-		}
-	}
-
-	if(ShapesCategory)
-	{
-		// Special case: Add Static Mesh Actor to the Shapes category
-		// TODO: Add a better way to register exception and extra classes if we want to deviate from the place actors set later on
-		FindOrAddClassFilter(AStaticMeshActor::StaticClass(), ShapesCategory);
-	}
-
-	// Sort the filter menu by name
-	FilterClasses.ValueSort([](const TSharedRef<FCustomClassFilterData>& ClassA, const TSharedRef<FCustomClassFilterData>& ClassB)
-		{
-			return ClassA->GetName().CompareTo(ClassB->GetName()) < 0;
-		});
-
-	// Now add all the classes to the filter bar
-	for(const TPair<FTopLevelAssetPath, TSharedRef<FCustomClassFilterData>>& FilterClass : FilterClasses)
-	{
-		OutFilterBarOptions.CustomClassFilters.Add(FilterClass.Value);
-	}
-	
-}
-
-/** Function called by the Outliner Filter Bar to compare an item with Asset Type Filters*/
-bool FActorBrowsingMode::ConvertItemToAssetData(SceneOutliner::FilterBarType InItem, FAssetData &OutAssetData) const
+bool FActorBrowsingMode::CompareItemWithClassName(SceneOutliner::FilterBarType InItem, const TSet<FTopLevelAssetPath>& AssetClassPaths) const
 {
 	// Type filtering only supported for Actors (and unloaded actors) currently
 	if (const FActorTreeItem* ActorItem = InItem.CastTo<FActorTreeItem>())
 	{
 		AActor* Actor = ActorItem->Actor.Get();
 
-		OutAssetData = FAssetData(Actor);
+		FTopLevelAssetPath AssetClassPath = Actor->GetClass()->GetClassPathName();
 
-		return true;
+		// For Blueprints, we check both the parent type (e.g Pawn/Actor etc) and the Blueprint class itself
+		if (UBlueprint* ClassBP = UBlueprint::GetBlueprintFromClass(Actor->GetClass()))
+		{
+			return AssetClassPaths.Contains(ClassBP->GetClass()->GetClassPathName()) || AssetClassPaths.Contains(AssetClassPath);
+		}
+		
+		return AssetClassPaths.Contains(AssetClassPath);
 	}
+	// TODO: Maybe can be improved to not use TryConvertShortClassNameToPathName to optimize?
 	else if (const FActorDescTreeItem* ActorDescItem = InItem.CastTo<FActorDescTreeItem>())
 	{
 		if (const FWorldPartitionActorDesc* ActorDesc = ActorDescItem->ActorDescHandle.Get())
 		{
-			// For Unloaded Actors, grab the native class 
-			OutAssetData = FAssetData();
+			// For Unloaded Actors, grab the native class and try to convert it to a path name
 			FTopLevelAssetPath ClassPath = ActorDesc->GetNativeClass();
-			if(ClassPath.IsValid())
-			{
-				OutAssetData.AssetClassPath = ClassPath;
-				return true;
-			}
+			return AssetClassPaths.Contains(ClassPath);
 		}
 	}
 
-	return FActorModeInteractive::ConvertItemToAssetData(InItem, OutAssetData);
+	return FActorModeInteractive::CompareItemWithClassName(InItem, AssetClassPaths);
 }
 
 #undef LOCTEXT_NAMESPACE
diff --git a/Engine/Source/Editor/SceneOutliner/Private/SSceneOutliner.cpp b/Engine/Source/Editor/SceneOutliner/Private/SSceneOutliner.cpp
--- a/Engine/Source/Editor/SceneOutliner/Private/SSceneOutliner.cpp
+++ b/Engine/Source/Editor/SceneOutliner/Private/SSceneOutliner.cpp
@@ -2454,9 +2454,9 @@ void SSceneOutliner::OnFilterBarFilterChanged()
 	FullRefresh();
 }
 
-bool SSceneOutliner::ConvertItemToAssetData(SceneOutliner::FilterBarType InItem, FAssetData &OutAssetData) const
+bool SSceneOutliner::CompareItemWithClassName(SceneOutliner::FilterBarType InItem, const TSet<FTopLevelAssetPath>& AssetClassPaths) const
 {
-	return Mode->ConvertItemToAssetData(InItem, OutAssetData);
+	return Mode->CompareItemWithClassName(InItem, AssetClassPaths);
 }
 
 void SSceneOutliner::CreateFilterBar(const FSceneOutlinerFilterBarOptions& FilterBarOptions)
@@ -2469,13 +2469,14 @@ void SSceneOutliner::CreateFilterBar(const FSceneOutlinerFilterBarOptions& Filte
 	FName FilterBarIdentifier = OutlinerIdentifier;
 
 	SAssignNew(FilterBar, SSceneOutlinerFilterBar)
-	.OnConvertItemToAssetData(this, &SSceneOutliner::ConvertItemToAssetData)
+	.OnCompareItemWithClassNames(this, &SSceneOutliner::CompareItemWithClassName)
 	.OnFilterChanged(this, &SSceneOutliner::OnFilterBarFilterChanged)
 	.CustomClassFilters(FilterBarOptions.CustomClassFilters)
 	.CustomFilters(FilterBarOptions.CustomFilters)
 	.FilterSearchBox(FilterTextBoxWidget)
 	.FilterBarIdentifier(FilterBarIdentifier)
 	.UseSharedSettings(FilterBarOptions.bUseSharedSettings)
+	.CategoryToExpand(FilterBarOptions.CategoryToExpand)
 	.CreateTextFilter(SSceneOutlinerFilterBar::FCreateTextFilter::CreateLambda([this]()
 	{
 		TSharedPtr< SceneOutliner::TreeItemTextFilter > Filter = this->CreateTextFilter();
diff --git a/Engine/Source/Editor/SceneOutliner/Private/SceneOutlinerFilterBar.cpp b/Engine/Source/Editor/SceneOutliner/Private/SceneOutlinerFilterBar.cpp
--- a/Engine/Source/Editor/SceneOutliner/Private/SceneOutlinerFilterBar.cpp
+++ b/Engine/Source/Editor/SceneOutliner/Private/SceneOutlinerFilterBar.cpp
@@ -8,6 +8,7 @@ SSceneOutlinerFilterBar::FCustomTextFilterEvent SSceneOutlinerFilterBar::CustomT
 void SSceneOutlinerFilterBar::Construct( const FArguments& InArgs )
 {
 	bUseSharedSettings = InArgs._UseSharedSettings;
+	this->CategoryToExpand = InArgs._CategoryToExpand;
 
 	SFilterBar<SceneOutliner::FilterBarType>::FArguments Args;
 	Args._OnFilterChanged = InArgs._OnFilterChanged;
@@ -17,6 +18,7 @@ void SSceneOutlinerFilterBar::Construct( const FArguments& InArgs )
 	Args._CreateTextFilter = InArgs._CreateTextFilter;
 	Args._FilterSearchBox = InArgs._FilterSearchBox;
 	Args._FilterBarIdentifier = InArgs._FilterBarIdentifier;
+	Args._OnCompareItemWithClassNames = InArgs._OnCompareItemWithClassNames;
 	Args._OnConvertItemToAssetData = InArgs._OnConvertItemToAssetData;
 	Args._FilterPillStyle = EFilterPillStyle::Basic;
 
diff --git a/Engine/Source/Editor/SceneOutliner/Private/SceneOutlinerFilterBar.h b/Engine/Source/Editor/SceneOutliner/Private/SceneOutlinerFilterBar.h
--- a/Engine/Source/Editor/SceneOutliner/Private/SceneOutlinerFilterBar.h
+++ b/Engine/Source/Editor/SceneOutliner/Private/SceneOutlinerFilterBar.h
@@ -16,7 +16,7 @@ public:
 	DECLARE_EVENT_OneParam(SSceneOutlinerFilterBar, FCustomTextFilterEvent, TSharedPtr<SWidget> /* BroadcastingFilterBar */);
 	
 	using FOnFilterChanged = typename SBasicFilterBar<SceneOutliner::FilterBarType>::FOnFilterChanged;
-	using FConvertItemToAssetData = typename FAssetFilter<SceneOutliner::FilterBarType>::FConvertItemToAssetData;
+	using FCompareItemWithClassNames = typename FAssetFilter<SceneOutliner::FilterBarType>::FCompareItemWithClassNames;
 	using FCreateTextFilter = typename SBasicFilterBar<SceneOutliner::FilterBarType>::FCreateTextFilter;
 	
 	SLATE_BEGIN_ARGS( SSceneOutlinerFilterBar )
@@ -32,6 +32,9 @@ public:
 		/** Specify this delegate to use asset comparison through IAssetRegistry, where you specify how to convert your Item into an FAssetData */
 		SLATE_EVENT( FConvertItemToAssetData, OnConvertItemToAssetData )
 
+		/** Specify this delegate to use asset comparison through IAssetRegistry, where you specify how to convert your Item into an FAssetData */
+		SLATE_EVENT( FCompareItemWithClassNames, OnCompareItemWithClassNames )
+
 		/** Initial List of Custom Filters that will be added to the AddFilter Menu */
 		SLATE_ARGUMENT( TArray<TSharedRef<FFilterBase<SceneOutliner::FilterBarType>>>, CustomFilters)
 
@@ -60,6 +63,9 @@ public:
 		 */
 		SLATE_ARGUMENT(bool, UseSharedSettings)
 
+		/** The Category to expand in the Filter Menu */
+		SLATE_ARGUMENT(TSharedPtr<FFilterCategory>, CategoryToExpand)
+
 	SLATE_END_ARGS()
 
 	void Construct( const FArguments& InArgs );
diff --git a/Engine/Source/Editor/SceneOutliner/Private/SceneOutlinerModule.cpp b/Engine/Source/Editor/SceneOutliner/Private/SceneOutlinerModule.cpp
--- a/Engine/Source/Editor/SceneOutliner/Private/SceneOutlinerModule.cpp
+++ b/Engine/Source/Editor/SceneOutliner/Private/SceneOutlinerModule.cpp
@@ -174,9 +174,6 @@ TSharedRef< ISceneOutliner > FSceneOutlinerModule::CreateActorBrowser(const FSce
 		
 	}
 
-	// Setup the FilterBar options for the Actor Browser
-	FActorBrowsingMode::CreateOutlinerFilterBarFilters(InitOptions.FilterBarOptions);
-	
 	return CreateSceneOutliner(InitOptions);
 }
 
diff --git a/Engine/Source/Editor/SceneOutliner/Public/ActorBrowsingMode.h b/Engine/Source/Editor/SceneOutliner/Public/ActorBrowsingMode.h
--- a/Engine/Source/Editor/SceneOutliner/Public/ActorBrowsingMode.h
+++ b/Engine/Source/Editor/SceneOutliner/Public/ActorBrowsingMode.h
@@ -101,12 +101,9 @@ public:
 
 	/** Called by engine when edit delete actors ends */
 	void OnDeleteActorsEnd();
-
-	/* Helper function to create the basic set of filters that the Actor Browser will have */
-	static void CreateOutlinerFilterBarFilters(FSceneOutlinerFilterBarOptions& OutFilterBarOptions);
-
-	/** Function called by the Outliner Filter Bar to compare an item with Asset Type Filters*/
-	virtual bool ConvertItemToAssetData(SceneOutliner::FilterBarType InItem, FAssetData &OutAssetData) const override;
+	
+	/** Function called by the Outliner Filter Bar to compare an item with Type Filters*/
+	virtual bool CompareItemWithClassName(SceneOutliner::FilterBarType InItem, const TSet<FTopLevelAssetPath>&) const override;
 private:
 	/** Build and up the context menu */
 	TSharedPtr<SWidget> BuildContextMenu();
diff --git a/Engine/Source/Editor/SceneOutliner/Public/ISceneOutlinerMode.h b/Engine/Source/Editor/SceneOutliner/Public/ISceneOutlinerMode.h
--- a/Engine/Source/Editor/SceneOutliner/Public/ISceneOutlinerMode.h
+++ b/Engine/Source/Editor/SceneOutliner/Public/ISceneOutlinerMode.h
@@ -164,8 +164,8 @@ public:
 	/** Unpins all selected items */
 	virtual void UnpinSelectedItems() {}
 
-	/** Function called by the Outliner Filter Bar to compare an item with Asset Type Filters*/
-	virtual bool ConvertItemToAssetData(SceneOutliner::FilterBarType InItem, FAssetData &OutAssetData) const { return false; };
+	/** Function called by the Outliner Filter Bar to compare an item with Type Filters*/
+	virtual bool CompareItemWithClassName(SceneOutliner::FilterBarType InItem, const TSet<FTopLevelAssetPath>&) const { return false; };
 public:
 	/* Getters */
 
diff --git a/Engine/Source/Editor/SceneOutliner/Public/SSceneOutliner.h b/Engine/Source/Editor/SceneOutliner/Public/SSceneOutliner.h
--- a/Engine/Source/Editor/SceneOutliner/Public/SSceneOutliner.h
+++ b/Engine/Source/Editor/SceneOutliner/Public/SSceneOutliner.h
@@ -609,10 +609,7 @@ private:
 	/** Creates a TextFilter for ISceneOutlinerTreeItem used to save searches as Text Filters */
 	TSharedPtr< SceneOutliner::TreeItemTextFilter > CreateTextFilter() const;
 
-	/** Convert a SceneOutliner item into an FAssetData struct to compare against type filters *
-	 *  TODO: Move to the ISceneOutlinerMode maybe to allow each mode to specify how to convert items?
-	 */
-	bool ConvertItemToAssetData(SceneOutliner::FilterBarType InItem, FAssetData &OutAssetData) const;
+	bool CompareItemWithClassName(SceneOutliner::FilterBarType InItem, const TSet<FTopLevelAssetPath>&) const;
 
 	/** Delegate for when a filter in the filter bar is changed */
 	void OnFilterBarFilterChanged();
diff --git a/Engine/Source/Editor/SceneOutliner/Public/SceneOutlinerPublicTypes.h b/Engine/Source/Editor/SceneOutliner/Public/SceneOutlinerPublicTypes.h
--- a/Engine/Source/Editor/SceneOutliner/Public/SceneOutlinerPublicTypes.h
+++ b/Engine/Source/Editor/SceneOutliner/Public/SceneOutlinerPublicTypes.h
@@ -14,6 +14,7 @@
 class FExtender;
 struct FToolMenuContext;
 class FCustomClassFilterData;
+class FFilterCategory;
 template<typename FilterType> class FFilterBase;
 
 DECLARE_DELEGATE_TwoParams(FSceneOutlinerModifyContextMenu, FName& /* MenuName */, FToolMenuContext& /* MenuContext */);
@@ -146,6 +147,9 @@ struct FSceneOutlinerFilterBarOptions
 
 	/** If true, share this Outliner filter bar's custom text filters with the level editor outliners */
 	bool bUseSharedSettings = false;
+
+	/** The category to expand in the filter menu (There must be at least one filter attached to the category) */
+	TSharedPtr<FFilterCategory> CategoryToExpand;
 };
 
 /**
diff --git a/Engine/Source/Editor/SceneOutliner/SceneOutliner.Build.cs b/Engine/Source/Editor/SceneOutliner/SceneOutliner.Build.cs
--- a/Engine/Source/Editor/SceneOutliner/SceneOutliner.Build.cs
+++ b/Engine/Source/Editor/SceneOutliner/SceneOutliner.Build.cs
@@ -6,12 +6,6 @@ public class SceneOutliner : ModuleRules
 {
 	public SceneOutliner(ReadOnlyTargetRules Target) : base(Target)
 	{
-		PrivateIncludePathModuleNames.AddRange(
-			new string[] {
-				"PlacementMode"
-			}
-		);
-		
 		PrivateDependencyModuleNames.AddRange(
 			new string[]
 			{
