Fixed STransformViewportToolBar not collapsing when the viewport area got too small

This regressed from CL# 22635797 as STransformViewportToolBar would now auto-width rather than shrink. To workaround this while keeping "LevelViewportToolBar.MiddleExtension" we now calculate a max width for the slot containing STransformViewportToolBar that allows it to shrink when the toolbar is wider than the current viewport.

#rnx

diff --git a/Engine/Source/Editor/LevelEditor/Private/SLevelViewportToolBar.cpp b/Engine/Source/Editor/LevelEditor/Private/SLevelViewportToolBar.cpp
--- a/Engine/Source/Editor/LevelEditor/Private/SLevelViewportToolBar.cpp
+++ b/Engine/Source/Editor/LevelEditor/Private/SLevelViewportToolBar.cpp
@@ -257,7 +257,7 @@ void SLevelViewportToolBar::Construct( const FArguments& InArgs )
 				.HAlign(HAlign_Left)
 				.VAlign(VAlign_Fill)
 				[
-					SAssignNew(ExtensionPanel, SExtensionPanel)
+					SNew(SExtensionPanel)
 					.ExtensionPanelID("LevelViewportToolBar.LeftExtension")
 					.ExtensionContext(ExtensionContextObject)
 				]
@@ -297,16 +297,17 @@ void SLevelViewportToolBar::Construct( const FArguments& InArgs )
 				.HAlign(HAlign_Center)
 				.VAlign(VAlign_Fill)
 				[
-					SAssignNew(ExtensionPanel, SExtensionPanel)
+					SNew(SExtensionPanel)
 					.ExtensionPanelID("LevelViewportToolBar.MiddleExtension")
 					.ExtensionContext(ExtensionContextObject)
 				]
 				+ SHorizontalBox::Slot()
 				.AutoWidth()
+				.MaxWidth(TAttribute<float>::CreateSP(this, &SLevelViewportToolBar::GetTransformToolbarWidth))
 				.Padding(ToolbarSlotPadding)
 				.HAlign(HAlign_Right)
 				[
-					SNew(STransformViewportToolBar)
+					SAssignNew(TransformToolbar, STransformViewportToolBar)
 					.Viewport(ViewportRef)
 					.CommandList(ViewportRef->GetCommandList())
 					.Extenders(LevelEditorModule.GetToolBarExtensibilityManager()->GetAllExtenders())
@@ -352,7 +353,7 @@ bool SLevelViewportToolBar::IsViewModeSupported(EViewModeIndex ViewModeIndex) co
 	return true;
 }
 
-FLevelEditorViewportClient* SLevelViewportToolBar::GetLevelViewportClient()
+FLevelEditorViewportClient* SLevelViewportToolBar::GetLevelViewportClient() const
 {
 	if (Viewport.IsValid())
 	{
@@ -1712,6 +1713,44 @@ FText SLevelViewportToolBar::GetRealtimeOverrideTooltip() const
 	return FText::GetEmpty();
 }
 
+float SLevelViewportToolBar::GetTransformToolbarWidth() const
+{
+	if (TransformToolbar)
+	{
+		const float TransformToolbarWidth = TransformToolbar->GetDesiredSize().X;
+		if (TransformToolbar_CachedMaxWidth == 0.0f)
+		{
+			TransformToolbar_CachedMaxWidth = TransformToolbarWidth;
+		}
+
+		{
+			FLevelEditorViewportClient* LevelEditorViewportClient = GetLevelViewportClient();
+			if (LevelEditorViewportClient && LevelEditorViewportClient->Viewport)
+			{
+				const float ViewportWidth = static_cast<float>(LevelEditorViewportClient->Viewport->GetSizeXY().X);
+				const float ToolbarWidthMinusPreviousTransformToolbar = GetDesiredSize().X - TransformToolbar_CachedMaxWidth;
+				const float ToolbarWidthEstimate = ToolbarWidthMinusPreviousTransformToolbar + TransformToolbarWidth;
+
+				const float OverflowWidth = ToolbarWidthEstimate - ViewportWidth;
+				if (OverflowWidth > 0.0f)
+				{
+					// There isn't enough space in the viewport to show the toolbar!
+					// Try and shrink the transform toolbar (which has an overflow area) to make things fit
+					TransformToolbar_CachedMaxWidth = FMath::Max(FMath::Min(4.0f, TransformToolbarWidth), TransformToolbarWidth - OverflowWidth);
+				}
+				else
+				{
+					TransformToolbar_CachedMaxWidth = TransformToolbarWidth;
+				}
+			}
+		}
+		
+		return TransformToolbar_CachedMaxWidth;
+	}
+
+	return 0.0f;
+}
+
 bool SLevelViewportToolBar::IsLandscapeLODSettingChecked(int32 Value) const
 {
 	return Viewport.Pin()->GetLevelViewportClient().LandscapeLODOverride == Value;
diff --git a/Engine/Source/Editor/LevelEditor/Private/SLevelViewportToolBar.h b/Engine/Source/Editor/LevelEditor/Private/SLevelViewportToolBar.h
--- a/Engine/Source/Editor/LevelEditor/Private/SLevelViewportToolBar.h
+++ b/Engine/Source/Editor/LevelEditor/Private/SLevelViewportToolBar.h
@@ -12,6 +12,7 @@
 class ACameraActor;
 class FExtender;
 class SExtensionPanel;
+class STransformViewportToolBar;
 class FMenuBuilder;
 class UToolMenu;
 struct FToolMenuSection;
@@ -32,7 +33,7 @@ public:
 	virtual bool IsViewModeSupported(EViewModeIndex ViewModeIndex) const;
 
 	/** @return Level editor viewport client. */ 
-	FLevelEditorViewportClient* GetLevelViewportClient();
+	FLevelEditorViewportClient* GetLevelViewportClient() const;
 
 	/** Fills view menu */
 	void FillViewMenu(UToolMenu* InMenu);
@@ -277,11 +278,17 @@ private:
 	void OnDisableRealtimeOverride();
 	bool IsRealtimeOverrideToggleVisible() const;
 	FText GetRealtimeOverrideTooltip() const;
+
+	float GetTransformToolbarWidth() const;
+
 private:
 	/** The viewport that we are in */
 	TWeakPtr<class SLevelViewport> Viewport;
 
-	/** SLevelViewportToolBar UI extension */
-	TSharedPtr<SExtensionPanel> ExtensionPanel;
+	/** STransformViewportToolBar menu */
+	TSharedPtr<STransformViewportToolBar> TransformToolbar;
+
+	/** The previous max STransformViewportToolBar width to allow deterministic size calculations */
+	mutable float TransformToolbar_CachedMaxWidth = 0.0f;
 };
 
