Show Browse button on actors that have a Browse To override

diff --git a/Engine/Source/Editor/LevelEditor/Private/LevelEditorContextMenu.cpp b/Engine/Source/Editor/LevelEditor/Private/LevelEditorContextMenu.cpp
--- a/Engine/Source/Editor/LevelEditor/Private/LevelEditorContextMenu.cpp
+++ b/Engine/Source/Editor/LevelEditor/Private/LevelEditorContextMenu.cpp
@@ -286,7 +286,7 @@ void FLevelEditorContextMenu::RegisterActorContextMenu()
 			FSlateIcon AssetIcon = ReferencedAssets.Num() == 1 ? FSlateIconFinder::FindIconForClass(ReferencedAssets[0]->GetClass()) : FSlateIcon(FAppStyle::GetAppStyleSetName(), "ClassIcon.Default");
 
 			// Edit and Find entries (a) always appear in main menu, and (b) appear in right-click menu if referenced asset is available
-			if (LevelEditorContext->ContextType == ELevelEditorMenuContext::MainMenu || ReferencedAssets.Num() > 0 || SoftReferencedAssets.Num() > 0)
+			if (LevelEditorContext->ContextType == ELevelEditorMenuContext::MainMenu || ReferencedAssets.Num() > 0 || SoftReferencedAssets.Num() > 0 || SelectionInfo.bHaveBrowseOverride)
 			{
 				Section.AddMenuEntry(FGlobalEditorCommonCommands::Get().FindInContentBrowser);
 
diff --git a/Engine/Source/Editor/UnrealEd/Private/AssetSelection.cpp b/Engine/Source/Editor/UnrealEd/Private/AssetSelection.cpp
--- a/Engine/Source/Editor/UnrealEd/Private/AssetSelection.cpp
+++ b/Engine/Source/Editor/UnrealEd/Private/AssetSelection.cpp
@@ -311,6 +311,11 @@ namespace AssetSelectionUtils
 						// Just count the total number of actors with counterparts
 						++ActorInfo.NumSimulationChanges;
 					}
+
+					if (!CurrentActor->GetBrowseToAssetOverride().IsEmpty())
+					{
+						ActorInfo.bHaveBrowseOverride = true;
+					}
 				}
 
 				if( ActorInfo.SelectionClass != NULL )
diff --git a/Engine/Source/Editor/UnrealEd/Public/AssetSelection.h b/Engine/Source/Editor/UnrealEd/Public/AssetSelection.h
--- a/Engine/Source/Editor/UnrealEd/Public/AssetSelection.h
+++ b/Engine/Source/Editor/UnrealEd/Public/AssetSelection.h
@@ -112,6 +112,8 @@ struct FSelectedActorInfo
 	uint32 bHaveExperimentalClass:1;
 	/** true if an early access actor (or actor containing such a component) is selected */
 	uint32 bHaveEarlyAccessClass:1;
+	/** true if an actor in the selection has an override for the asset it browses to in the Content Browser*/
+	uint32 bHaveBrowseOverride:1;
 
 	FSelectedActorInfo()
 		: SelectionClass(NULL)
@@ -145,6 +147,7 @@ struct FSelectedActorInfo
 		, bHaveLandscape(false)
 		, bHaveExperimentalClass(false)
 		, bHaveEarlyAccessClass(false)
+		, bHaveBrowseOverride(false)
 	{
 	}
 
