Outliner:
Add a filter and column to show Unsaved Assets.
Add a filter and modify the source control column to show uncontrolled assets

diff --git a/Engine/Source/Developer/UncontrolledChangelists/Private/UncontrolledChangelistState.cpp b/Engine/Source/Developer/UncontrolledChangelists/Private/UncontrolledChangelistState.cpp
--- a/Engine/Source/Developer/UncontrolledChangelists/Private/UncontrolledChangelistState.cpp
+++ b/Engine/Source/Developer/UncontrolledChangelists/Private/UncontrolledChangelistState.cpp
@@ -50,7 +50,7 @@ const FText& FUncontrolledChangelistState::GetDescriptionText() const
 
 FText FUncontrolledChangelistState::GetDisplayTooltip() const
 {
-	return LOCTEXT("Tooltip", "Tooltip");
+	return LOCTEXT("UncontrolledStateTooltip", "Uncontrolled: Locally modified outside of source control");
 }
 
 const FDateTime& FUncontrolledChangelistState::GetTimeStamp() const
diff --git a/Engine/Source/Developer/UnsavedAssetsTracker/Source/Private/UnsavedAssetsTracker.cpp b/Engine/Source/Developer/UnsavedAssetsTracker/Source/Private/UnsavedAssetsTracker.cpp
--- a/Engine/Source/Developer/UnsavedAssetsTracker/Source/Private/UnsavedAssetsTracker.cpp
+++ b/Engine/Source/Developer/UnsavedAssetsTracker/Source/Private/UnsavedAssetsTracker.cpp
@@ -424,6 +424,11 @@ void FUnsavedAssetsTracker::PrompToSavePackages()
 	}
 }
 
+bool FUnsavedAssetsTracker::IsAssetUnsaved(const FString& FileAbsPathname) const
+{
+	return UnsavedAbsFilePathames.Contains(FileAbsPathname);
+}
+
 void FUnsavedAssetsTracker::OnSourceControlWarningNotification(const ISourceControlState& State, FStatus& InOutStatus)
 {
 	auto UpdateAndShowWarningIfNotAlreadyShown = [this](EWarningTypes WarningType, const FText& Msg, FStatus& InOutStatus)
diff --git a/Engine/Source/Developer/UnsavedAssetsTracker/Source/Private/UnsavedAssetsTracker.h b/Engine/Source/Developer/UnsavedAssetsTracker/Source/Private/UnsavedAssetsTracker.h
--- a/Engine/Source/Developer/UnsavedAssetsTracker/Source/Private/UnsavedAssetsTracker.h
+++ b/Engine/Source/Developer/UnsavedAssetsTracker/Source/Private/UnsavedAssetsTracker.h
@@ -42,6 +42,9 @@ public:
 	/** Displays a dialog prompting the user to save the packages. */
 	void PrompToSavePackages();
 
+	/** Check if the input asset is unsaved. */
+	bool IsAssetUnsaved(const FString& FileAbsPathname) const;
+
 private:
 	enum class EWarningTypes
 	{
diff --git a/Engine/Source/Developer/UnsavedAssetsTracker/Source/Private/UnsavedAssetsTrackerModule.cpp b/Engine/Source/Developer/UnsavedAssetsTracker/Source/Private/UnsavedAssetsTrackerModule.cpp
--- a/Engine/Source/Developer/UnsavedAssetsTracker/Source/Private/UnsavedAssetsTrackerModule.cpp
+++ b/Engine/Source/Developer/UnsavedAssetsTracker/Source/Private/UnsavedAssetsTrackerModule.cpp
@@ -37,6 +37,15 @@ TArray<FString> FUnsavedAssetsTrackerModule::GetUnsavedAssets() const
 	return TArray<FString>();
 }
 
+bool FUnsavedAssetsTrackerModule::IsAssetUnsaved(const FString& FileAbsPathname) const
+{
+	if (UnsavedAssetTracker)
+	{
+		return UnsavedAssetTracker->IsAssetUnsaved(FileAbsPathname);
+	}
+	return false;
+}
+
 TSharedRef<SWidget> FUnsavedAssetsTrackerModule::MakeUnsavedAssetsStatusBarWidget()
 {
 	if (!UnsavedAssetTracker)
diff --git a/Engine/Source/Developer/UnsavedAssetsTracker/Source/Public/UnsavedAssetsTrackerModule.h b/Engine/Source/Developer/UnsavedAssetsTracker/Source/Public/UnsavedAssetsTrackerModule.h
--- a/Engine/Source/Developer/UnsavedAssetsTracker/Source/Public/UnsavedAssetsTrackerModule.h
+++ b/Engine/Source/Developer/UnsavedAssetsTracker/Source/Public/UnsavedAssetsTrackerModule.h
@@ -45,6 +45,9 @@ public:
 	/** Returns the list of unsaved assets. */
 	TArray<FString> GetUnsavedAssets() const;
 
+	/** Check if the input asset is unsaved. */
+	bool IsAssetUnsaved(const FString& FileAbsPathname) const;
+
 	/** Construct the widget for the Editor status bar.*/
 	TSharedRef<SWidget> MakeUnsavedAssetsStatusBarWidget();
 
diff --git a/Engine/Source/Editor/LevelEditor/LevelEditor.Build.cs b/Engine/Source/Editor/LevelEditor/LevelEditor.Build.cs
--- a/Engine/Source/Editor/LevelEditor/LevelEditor.Build.cs
+++ b/Engine/Source/Editor/LevelEditor/LevelEditor.Build.cs
@@ -95,6 +95,8 @@ public class LevelEditor : ModuleRules
 				"EditorWidgets",
 				"ToolWidgets",
 				"VirtualizationEditor",
+				"UnsavedAssetsTracker",
+				"UncontrolledChangelists"
 			}
 		);
 
diff --git a/Engine/Source/Editor/LevelEditor/Private/LevelEditor.cpp b/Engine/Source/Editor/LevelEditor/Private/LevelEditor.cpp
--- a/Engine/Source/Editor/LevelEditor/Private/LevelEditor.cpp
+++ b/Engine/Source/Editor/LevelEditor/Private/LevelEditor.cpp
@@ -1841,9 +1841,17 @@ TSharedPtr<FLevelEditorOutlinerSettings> FLevelEditorModule::GetLevelEditorOutli
 	return OutlinerSettings;
 }
 
-void FLevelEditorModule::AddCustomFilterToOutliner(TSharedRef<FFilterBase<SceneOutliner::FilterBarType>> InCustomFilter)
+void FLevelEditorModule::AddCustomFilterToOutliner(TSharedRef<FFilterBase<const ISceneOutlinerTreeItem&>> InCustomFilter)
 {
+	// Disable deprecation warnings so we can call the deprecated function to support this function (which is also deprecated)
+	PRAGMA_DISABLE_DEPRECATION_WARNINGS
 	OutlinerSettings->AddCustomFilter(InCustomFilter);
+	PRAGMA_ENABLE_DEPRECATION_WARNINGS
+}
+
+void FLevelEditorModule::AddCustomFilterToOutliner(FLevelEditorOutlinerSettings::FOutlinerFilterFactory InCreateCustomFilter)
+{
+	OutlinerSettings->AddCustomFilter(InCreateCustomFilter);
 }
 
 void FLevelEditorModule::AddCustomClassFilterToOutliner(TSharedRef<FCustomClassFilterData> InCustomClassFilterData)
diff --git a/Engine/Source/Editor/LevelEditor/Private/LevelEditorOutlinerSettings.cpp b/Engine/Source/Editor/LevelEditor/Private/LevelEditorOutlinerSettings.cpp
--- a/Engine/Source/Editor/LevelEditor/Private/LevelEditorOutlinerSettings.cpp
+++ b/Engine/Source/Editor/LevelEditor/Private/LevelEditorOutlinerSettings.cpp
@@ -4,10 +4,18 @@
 
 #include "Filters/FilterBase.h"
 #include "Filters/CustomClassFilterData.h"
+#include "Filters/GenericFilter.h"
 
-#include "SceneOutlinerPublicTypes.h"
+#include "LevelEditor.h"
 #include "IPlacementModeModule.h"
 
+#include "UnsavedAssetsTrackerModule.h"
+#include "UncontrolledChangelistsModule.h"
+
+#include "SceneOutlinerPublicTypes.h"
+#include "ISceneOutliner.h"
+#include "SceneOutlinerHelpers.h"
+
 #include "Animation/SkeletalMeshActor.h"
 #include "Components/SkyAtmosphereComponent.h"
 #include "Components/VolumetricCloudComponent.h"
@@ -23,11 +31,21 @@
 
 #define LOCTEXT_NAMESPACE "LevelEditorOutlinerSettings"
 
-void FLevelEditorOutlinerSettings::AddCustomFilter(TSharedRef<FFilterBase<SceneOutliner::FilterBarType>> InCustomFilter)
+const FString FLevelEditorOutlinerSettings::UnsavedAssetsFilterName("UnsavedAssetsFilter");
+const FString FLevelEditorOutlinerSettings::UncontrolledAssetsFilterName("UncontrolledAssetsFilter");
+
+void FLevelEditorOutlinerSettings::AddCustomFilter(TSharedRef<FFilterBase<const ISceneOutlinerTreeItem&>> InCustomFilter)
 {
 	CustomFilters.Add(InCustomFilter);
 }
 
+void FLevelEditorOutlinerSettings::AddCustomFilter(FOutlinerFilterFactory InCreateCustomFilter)
+{
+	check (InCreateCustomFilter.IsBound())
+	
+	CustomFilterDelegates.Add(InCreateCustomFilter);
+}
+
 void FLevelEditorOutlinerSettings::AddCustomClassFilter(TSharedRef<FCustomClassFilterData> InCustomClassFilterData)
 {
 	TSharedRef<FCustomClassFilterData>* FoundFilter = CustomClassFilters.FindByPredicate([&InCustomClassFilterData](const TSharedPtr<FCustomClassFilterData> FilterData)
@@ -37,7 +55,7 @@ void FLevelEditorOutlinerSettings::AddCustomClassFilter(TSharedRef<FCustomClassF
 
 	// If a filter for the class already exists, just add the new categories to it
 	// TODO: Not perfect, what if the user wants to use InCustomClassFilterData after if there is already a dupe?
-	if(FoundFilter)
+	if (FoundFilter)
 	{
 		TArray<TSharedPtr<FFilterCategory>> Categories = InCustomClassFilterData->GetCategories();
 
@@ -54,7 +72,7 @@ void FLevelEditorOutlinerSettings::AddCustomClassFilter(TSharedRef<FCustomClassF
 
 TSharedPtr<FFilterCategory> FLevelEditorOutlinerSettings::GetFilterCategory(const FName& CategoryName)
 {
-	if(TSharedPtr<FFilterCategory>* FoundCategory = FilterBarCategories.Find(CategoryName))
+	if (TSharedPtr<FFilterCategory>* FoundCategory = FilterBarCategories.Find(CategoryName))
 	{
 		return *FoundCategory;
 	}
@@ -102,6 +120,36 @@ void FLevelEditorOutlinerSettings::SetupBuiltInCategories()
 	FilterBarCategories.Add(FLevelEditorOutlinerBuiltInCategories::Volumes(), VolumesFilterCategory);
 }
 
+void FLevelEditorOutlinerSettings::RefreshOutlinersWithActiveFilter(bool bFullRefresh, const FString& InFilterName)
+{
+	TRACE_CPUPROFILER_EVENT_SCOPE(FLevelEditorOutlinerSettings::RefreshOutlinersWithActiveFilter);
+
+	TWeakPtr<ILevelEditor> LevelEditor = FModuleManager::GetModuleChecked<FLevelEditorModule>(TEXT("LevelEditor")).GetLevelEditorInstance();
+	if (LevelEditor.IsValid())
+	{
+		TArray<TWeakPtr<ISceneOutliner>> SceneOutlinerPtrs = LevelEditor.Pin()->GetAllSceneOutliners();
+
+		for (TWeakPtr<ISceneOutliner> SceneOutlinerPtr : SceneOutlinerPtrs)
+		{
+			if (TSharedPtr<ISceneOutliner> SceneOutlinerPin = SceneOutlinerPtr.Pin())
+			{
+				// If this outliner has the input filter active, refresh it
+				if (SceneOutlinerPin->IsFilterActive(InFilterName))
+				{
+					if (bFullRefresh)
+					{
+						SceneOutlinerPin->FullRefresh();
+					}
+					else
+					{
+						SceneOutlinerPin->Refresh();
+					}
+				}
+			}
+		}
+	}
+}
+
 void FLevelEditorOutlinerSettings::CreateDefaultFilters()
 {
 	// First we will add all items registered to the place actors panel as filters
@@ -126,7 +174,7 @@ void FLevelEditorOutlinerSettings::CreateDefaultFilters()
 			return FilterData->GetClassPathName() == Class->GetClassPathName();
 		});
 		
-		if(ExistingClassData)
+		if (ExistingClassData)
 		{
 			(*ExistingClassData)->AddCategory(FilterCategory);
 		}
@@ -141,7 +189,7 @@ void FLevelEditorOutlinerSettings::CreateDefaultFilters()
 	{
 		// Make an FFilterCategory using the current PlacementCategory if it doesn't already exist (built in)
 		TSharedPtr<FFilterCategory> FilterCategory;
-		if(TSharedPtr<FFilterCategory>* FoundCategory = FilterBarCategories.Find(Category.UniqueHandle))
+		if (TSharedPtr<FFilterCategory>* FoundCategory = FilterBarCategories.Find(Category.UniqueHandle))
 		{
 			FilterCategory = *FoundCategory;
 		}
@@ -153,7 +201,7 @@ void FLevelEditorOutlinerSettings::CreateDefaultFilters()
 		FName CategoryName = Category.UniqueHandle;
 
 		// If it is a built in placement mode category, convert the name to a built in filter category
-		if(FName* MappedCategoryName = PlacementToFilterCategoryMap.Find(CategoryName))
+		if (FName* MappedCategoryName = PlacementToFilterCategoryMap.Find(CategoryName))
 		{
 			CategoryName = *MappedCategoryName;
 		}
@@ -182,7 +230,7 @@ void FLevelEditorOutlinerSettings::CreateDefaultFilters()
 				DefaultActor = CastChecked<AActor>(CastChecked<UClass>(Item->AssetData.GetAsset())->ClassDefaultObject);
 			
 			}
-			if(!DefaultActor)
+			if (!DefaultActor)
 			{
 				continue;
 			}
@@ -196,7 +244,7 @@ void FLevelEditorOutlinerSettings::CreateDefaultFilters()
 
 	TSharedPtr<FFilterCategory>* CommonFilterCategory = FilterBarCategories.Find(FLevelEditorOutlinerBuiltInCategories::Common());
 
-	if(CommonFilterCategory)
+	if (CommonFilterCategory)
 	{
 		FindOrAddClassFilter(AStaticMeshActor::StaticClass(), *CommonFilterCategory);
 		FindOrAddClassFilter(ALevelInstance::StaticClass(), *CommonFilterCategory);
@@ -207,7 +255,7 @@ void FLevelEditorOutlinerSettings::CreateDefaultFilters()
 
 	TSharedPtr<FFilterCategory>* GeometryFilterCategory = FilterBarCategories.Find(FLevelEditorOutlinerBuiltInCategories::Geometry());
 
-	if(GeometryFilterCategory)
+	if (GeometryFilterCategory)
 	{
 		FindOrAddClassFilter(AStaticMeshActor::StaticClass(), *GeometryFilterCategory);
 		FindOrAddClassFilter(ABrush::StaticClass(), *GeometryFilterCategory);
@@ -215,14 +263,14 @@ void FLevelEditorOutlinerSettings::CreateDefaultFilters()
 
 	TSharedPtr<FFilterCategory>* AnimationFilterCategory = FilterBarCategories.Find(FLevelEditorOutlinerBuiltInCategories::Animation());
 
-	if(AnimationFilterCategory)
+	if (AnimationFilterCategory)
 	{
 		FindOrAddClassFilter(ASkeletalMeshActor::StaticClass(), *AnimationFilterCategory);
 	}
 
 	TSharedPtr<FFilterCategory>* EnvironmentFilterCategory = FilterBarCategories.Find(FLevelEditorOutlinerBuiltInCategories::Environment());
 
-	if(EnvironmentFilterCategory)
+	if (EnvironmentFilterCategory)
 	{
 		FindOrAddClassFilter(AExponentialHeightFog::StaticClass(), *EnvironmentFilterCategory);
 		FindOrAddClassFilter(AInstancedFoliageActor::StaticClass(), *EnvironmentFilterCategory);
@@ -232,10 +280,84 @@ void FLevelEditorOutlinerSettings::CreateDefaultFilters()
 	
 	TSharedPtr<FFilterCategory>* AudioFilterCategory = FilterBarCategories.Find(FLevelEditorOutlinerBuiltInCategories::Audio());
 
-	if(AudioFilterCategory)
+	if (AudioFilterCategory)
 	{
 		FindOrAddClassFilter(AAmbientSound::StaticClass(), *AudioFilterCategory);
 	}
+	
+	CreateSCCFilters();
+}
+
+bool FLevelEditorOutlinerSettings::DoesActorPassUnsavedFilter(const ISceneOutlinerTreeItem& InItem)
+{
+	return UnsavedAssets.Contains(SceneOutliner::FSceneOutlinerHelpers::GetExternalPackageName(InItem));
+}
+
+bool FLevelEditorOutlinerSettings::DoesActorPassUncontrolledFilter(const ISceneOutlinerTreeItem& InItem)
+{
+	FString ExternalPackageName = SceneOutliner::FSceneOutlinerHelpers::GetExternalPackageName(InItem);
+	
+	for (const TSharedRef<FUncontrolledChangelistState>& UncontrolledChangelistState : UncontrolledChangelistStates)
+	{
+		return UncontrolledChangelistState->GetFilenames().Contains(ExternalPackageName);
+	}
+	
+	return false;
+}
+
+void FLevelEditorOutlinerSettings::OnUnsavedAssetAdded(const FString& InAsset)
+{
+	UnsavedAssets.Add(InAsset);
+
+	// Refresh any outliners that have the unsaved assets filter active to refilter on adding an unsaved asset
+	RefreshOutlinersWithActiveFilter(/*bFullRefresh*/ true, UnsavedAssetsFilterName);
+}
+
+void FLevelEditorOutlinerSettings::OnUnsavedAssetRemoved(const FString& InAsset)
+{
+	UnsavedAssets.Remove(InAsset);
+
+	// Refresh any outliners that have the unsaved assets filter active to refilter on adding an unsaved asset
+	RefreshOutlinersWithActiveFilter(/*bFullRefresh*/ true, UnsavedAssetsFilterName);
+}
+
+void FLevelEditorOutlinerSettings::CreateSCCFilters()
+{
+	// Source Control Category
+	TSharedPtr<FFilterCategory> SCCFiltersCategory = MakeShared<FFilterCategory>(LOCTEXT("SCCFiltersCategory", "Source Control"), FText::GetEmpty());
+	
+	// Uncontrolled Actors Filter
+	FUncontrolledChangelistsModule& UncontrolledChangelistModule = FUncontrolledChangelistsModule::Get();
+	UncontrolledChangelistStates = FUncontrolledChangelistsModule::Get().GetChangelistStates();
+	UncontrolledChangelistModule.OnUncontrolledChangelistModuleChanged.AddSP(this, &FLevelEditorOutlinerSettings::OnUncontrolledChangelistModuleChanged);
+	
+	FGenericFilter<const ISceneOutlinerTreeItem&>::FOnItemFiltered UncontrolledFilterDelegate = FGenericFilter<const ISceneOutlinerTreeItem&>::FOnItemFiltered::CreateSP(this, &FLevelEditorOutlinerSettings::DoesActorPassUncontrolledFilter);
+	TSharedPtr<FGenericFilter<const ISceneOutlinerTreeItem&>> UncontrolledFilter = MakeShared<FGenericFilter<const ISceneOutlinerTreeItem&>>(SCCFiltersCategory, UncontrolledAssetsFilterName, LOCTEXT("UncontrolledFilterName", "Uncontrolled"), UncontrolledFilterDelegate);
+	UncontrolledFilter->SetToolTipText(LOCTEXT("UncontrolledFilterTooltip", "Only show items that are uncontrolled (locally modified outside of source control)"));
+	CustomFilters.Add(UncontrolledFilter.ToSharedRef());
+	
+	// File Management Category
+	TSharedPtr<FFilterCategory> FileManagementFiltersCategory = MakeShared<FFilterCategory>(LOCTEXT("FileManagementFiltersCategory", "File Management"), FText::GetEmpty());
+
+	// Unsaved Actors
+	FUnsavedAssetsTrackerModule& UnsavedAssetsTrackerModule = FUnsavedAssetsTrackerModule::Get();
+	UnsavedAssets = UnsavedAssetsTrackerModule.GetUnsavedAssets();
+	UnsavedAssetsTrackerModule.OnUnsavedAssetAdded.AddSP(this, &FLevelEditorOutlinerSettings::OnUnsavedAssetAdded);
+	UnsavedAssetsTrackerModule.OnUnsavedAssetRemoved.AddSP(this, &FLevelEditorOutlinerSettings::OnUnsavedAssetRemoved);
+	
+	FGenericFilter<const ISceneOutlinerTreeItem&>::FOnItemFiltered UnsavedActorsFilterDelegate = FGenericFilter<const ISceneOutlinerTreeItem&>::FOnItemFiltered::CreateSP(this, & FLevelEditorOutlinerSettings::DoesActorPassUnsavedFilter);
+	TSharedPtr<FGenericFilter<const ISceneOutlinerTreeItem&>> UnsavedAssetsFilter = MakeShared<FGenericFilter<const ISceneOutlinerTreeItem&>>(FileManagementFiltersCategory, UnsavedAssetsFilterName, LOCTEXT("UnsavedFilterName", "Unsaved"), UnsavedActorsFilterDelegate);
+	UnsavedAssetsFilter->SetToolTipText(LOCTEXT("UnsavedAssetsFilterTooltip", "Only show items that are unsaved"));
+	CustomFilters.Add(UnsavedAssetsFilter.ToSharedRef());
+}
+
+void FLevelEditorOutlinerSettings::OnUncontrolledChangelistModuleChanged()
+{
+	// Update our cached uncontrolled CL states
+	UncontrolledChangelistStates = FUncontrolledChangelistsModule::Get().GetChangelistStates();
+
+	// Refresh any Outliners that have the Uncontrolled Filter active
+	RefreshOutlinersWithActiveFilter(/*bFullRefresh*/ true, UncontrolledAssetsFilterName);
 }
 
 void FLevelEditorOutlinerSettings::GetOutlinerFilters(FSceneOutlinerFilterBarOptions& OutFilterBarOptions)
@@ -248,6 +370,12 @@ void FLevelEditorOutlinerSettings::GetOutlinerFilters(FSceneOutlinerFilterBarOpt
 	
 	OutFilterBarOptions.CustomClassFilters.Append(CustomClassFilters);
 	OutFilterBarOptions.CustomFilters.Append(CustomFilters);
+
+	for(FOutlinerFilterFactory& CreateFilterDelegate : CustomFilterDelegates)
+	{
+		TSharedRef<FFilterBase<const ISceneOutlinerTreeItem&>> NewFilter = CreateFilterDelegate.Execute().ToSharedRef();
+		OutFilterBarOptions.CustomFilters.Add(NewFilter);
+	}
 }
 
-#undef LOCTEXT_NAMESPACE
+#undef LOCTEXT_NAMESPACE
\ No newline at end of file
diff --git a/Engine/Source/Editor/LevelEditor/Public/LevelEditor.h b/Engine/Source/Editor/LevelEditor/Public/LevelEditor.h
--- a/Engine/Source/Editor/LevelEditor/Public/LevelEditor.h
+++ b/Engine/Source/Editor/LevelEditor/Public/LevelEditor.h
@@ -343,7 +343,15 @@ public:
 	 *   @see FGenericFilter on how to create generic filters
 	 *   NOTE: Currently does not support adding filters dynamically, they must be added before the Level Editor is init
 	 */
-	virtual void AddCustomFilterToOutliner(TSharedRef<FFilterBase<SceneOutliner::FilterBarType>> InCustomFilter);
+	UE_DEPRECATED(5.2, "Use the AddCustomFilterToOutliner override that takes in a delegate instead so each Outliner can have a unique instance of the filter")
+	virtual void AddCustomFilterToOutliner(TSharedRef<FFilterBase<const ISceneOutlinerTreeItem&>> InCustomFilter);
+
+	/**  Add a custom filter to the outliner filter bar. These are all AND'd together
+	 *	 This function takes in a delegate that will be called to create an instance of the custom filter for each Outliner that
+	 *	 the level editor creates.
+	 *   @see FGenericFilter on how to create generic filters
+	 */
+	void AddCustomFilterToOutliner(FLevelEditorOutlinerSettings::FOutlinerFilterFactory InCreateCustomFilter);
 
 	/**  Add a custom class filter to the outliner filter bar. These represent asset/actor type filters and are OR'd
 	 *   @see FLevelEditorOutlinerSettings
diff --git a/Engine/Source/Editor/LevelEditor/Public/LevelEditorOutlinerSettings.h b/Engine/Source/Editor/LevelEditor/Public/LevelEditorOutlinerSettings.h
--- a/Engine/Source/Editor/LevelEditor/Public/LevelEditorOutlinerSettings.h
+++ b/Engine/Source/Editor/LevelEditor/Public/LevelEditorOutlinerSettings.h
@@ -5,6 +5,7 @@
 #include "Filters/GenericFilter.h"
 
 struct ISceneOutlinerTreeItem;
+class FUncontrolledChangelistState;
 
 namespace SceneOutliner
 {
@@ -31,17 +32,31 @@ struct FLevelEditorOutlinerBuiltInCategories
 };
 
 /** Helper class to manage initalization options specific to the Level Editor Outliners
- * Use AddCustomFilter/AddCustomClass Filter to register new filters in the Outliner.
- * Make sure that the filters you attach have a category, otherwise they will not show up
+ *  Use AddCustomFilter/AddCustomClass Filter to register new filters in the Outliner.
+ *  Make sure that the filters you attach have a category, otherwise they will not show up
  */
 class LEVELEDITOR_API FLevelEditorOutlinerSettings :  public TSharedFromThis<class FLevelEditorOutlinerSettings>
 {
 public:
 
+	/** Delegate to create a Filter for the Outliner */
+	DECLARE_DELEGATE_RetVal(TSharedPtr<FFilterBase<const ISceneOutlinerTreeItem&>>, FOutlinerFilterFactory);
+
+	/**  Add a custom filter to the outliner filter bar. These are all AND'd together
+	 *   @see FGenericFilter on how to create generic filters
+	 *   Note: Any filter added here will be shared between all 4 Outliners, so you cannot rely on the filter to know
+	 *   which outliner it is active in. Use the override that takes a delegate instead
+	 */
+
+	UE_DEPRECATED(5.2, "Use the AddCustomFilter override that takes in a factory instead so each Outliner can have a unique instance of the filter")
+	void AddCustomFilter(TSharedRef<FFilterBase<const ISceneOutlinerTreeItem&>> InCustomFilter);
+
 	/**  Add a custom filter to the outliner filter bar. These are all AND'd together
+	 *   This function takes in a delegate that will be called to create an instance of the custom filter for each Outliner that
+	 *   the level editor creates.
 	 *   @see FGenericFilter on how to create generic filters
 	 */
-	void AddCustomFilter(TSharedRef<FFilterBase<SceneOutliner::FilterBarType>> InCustomFilter);
+	void AddCustomFilter(FOutlinerFilterFactory InCreateCustomFilter);
 
 	/**  Add a custom class filter to the outliner filter bar. These represent asset/actor type filters and are OR'd
 	 *   Can be created using IAssetTypeActions or UClass (@see constructor)
@@ -60,15 +75,32 @@ public:
 	// Get the FFilterCategory attached to the given category name. Use this to add filters to the built in categories
 	TSharedPtr<FFilterCategory> GetFilterCategory(const FName& CategoryName);
 
+private:
+
+	void CreateSCCFilters();
+
+	bool DoesActorPassUnsavedFilter(const ISceneOutlinerTreeItem& InItem);
+	bool DoesActorPassUncontrolledFilter(const ISceneOutlinerTreeItem& InItem);
+
+	void OnUnsavedAssetAdded(const FString& InAsset);
+	void OnUnsavedAssetRemoved(const FString& InAsset);
+
+	void OnUncontrolledChangelistModuleChanged();
+	
+	// Refresh any Outliners that have the given filter active
+	void RefreshOutlinersWithActiveFilter(bool bFullRefresh, const FString& InFilterName);
 private:
 
 	/** These are the custom filters that the Scene Outliner will have. All active filters will be AND'd together to test
 	 *  against.
 	 */
-	TArray<TSharedRef<FFilterBase<SceneOutliner::FilterBarType>>> CustomFilters;
+	TArray<TSharedRef<FFilterBase<const ISceneOutlinerTreeItem&>>> CustomFilters;
+
+	/** These are delegates that will be used to create custom filters for each outliner that calls GetOutlinerFilters */ 
+	TArray<FOutlinerFilterFactory> CustomFilterDelegates;
 
 	/** These are the asset type filters that the Scene Outliner will have. All active filters will be OR'd together to
-	 *  test against.
+	*  test against.
 	 */
 	TArray<TSharedRef<FCustomClassFilterData>> CustomClassFilters;
 
@@ -77,4 +109,15 @@ private:
 
 	/* A map to convert placement mode built in categories to filter bar categories */
 	TMap<FName, FName> PlacementToFilterCategoryMap;
+
+	// Source Control Cache
+
+	// List of currently unsaved assets
+	TArray<FString> UnsavedAssets;
+
+	// List of uncontrolled changelist state
+	TArray<TSharedRef<FUncontrolledChangelistState>> UncontrolledChangelistStates;
+
+	static const FString UnsavedAssetsFilterName;
+	static const FString UncontrolledAssetsFilterName;
 };
diff --git a/Engine/Source/Editor/SceneOutliner/Private/SSceneOutliner.cpp b/Engine/Source/Editor/SceneOutliner/Private/SSceneOutliner.cpp
--- a/Engine/Source/Editor/SceneOutliner/Private/SSceneOutliner.cpp
+++ b/Engine/Source/Editor/SceneOutliner/Private/SSceneOutliner.cpp
@@ -2481,7 +2481,7 @@ bool SSceneOutliner::CompareItemWithClassName(SceneOutliner::FilterBarType InIte
 
 void SSceneOutliner::CreateFilterBar(const FSceneOutlinerFilterBarOptions& FilterBarOptions)
 {
-	if(!FilterBarOptions.bHasFilterBar)
+	if (!FilterBarOptions.bHasFilterBar)
 	{
 		return;
 	}
@@ -2507,4 +2507,17 @@ void SSceneOutliner::CreateFilterBar(const FSceneOutlinerFilterBarOptions& Filte
 	FilterBar->LoadSettings();
 }
 
+bool SSceneOutliner::IsFilterActive(const FString& FilterName) const
+{
+	if (FilterBar.IsValid())
+	{
+		if (TSharedPtr<FFilterBase<SceneOutliner::FilterBarType>> FoundFilter = FilterBar->GetFilter(FilterName))
+		{
+			return FilterBar->IsFilterActive(FoundFilter);
+		}
+	}
+
+	return false;
+}
+
 #undef LOCTEXT_NAMESPACE
\ No newline at end of file
diff --git a/Engine/Source/Editor/SceneOutliner/Private/SSourceControlWidget.cpp b/Engine/Source/Editor/SceneOutliner/Private/SSourceControlWidget.cpp
--- a/Engine/Source/Editor/SceneOutliner/Private/SSourceControlWidget.cpp
+++ b/Engine/Source/Editor/SceneOutliner/Private/SSourceControlWidget.cpp
@@ -1,6 +1,9 @@
 // Copyright Epic Games, Inc. All Rights Reserved.
 
 #include "SSourceControlWidget.h"
+#include "UncontrolledChangelistsModule.h"
+
+#define LOCTEXT_NAMESPACE "SceneOutlinerSourceControlWidget"
 
 void SSourceControlWidget::Construct(const FArguments& InArgs, TSharedPtr<FSceneOutlinerTreeItemSCC> InItemSourceControl)
 {
@@ -9,7 +12,9 @@ void SSourceControlWidget::Construct(const FArguments& InArgs, TSharedPtr<FScene
 	ItemSourceControl = InItemSourceControl;
 
 	ItemSourceControl->OnSourceControlStateChanged.BindSP(this, &SSourceControlWidget::UpdateSourceControlState);
+	ItemSourceControl->OnUncontrolledChangelistsStateChanged.BindSP(this, &SSourceControlWidget::UpdateUncontrolledChangelistState);
 
+	
 	SImage::Construct(
 		SImage::FArguments()
 		.ColorAndOpacity(this, &SSourceControlWidget::GetForegroundColor)
@@ -29,15 +34,62 @@ FReply SSourceControlWidget::OnMouseButtonDoubleClick(const FGeometry& InMyGeome
 
 void SSourceControlWidget::UpdateSourceControlState(FSourceControlStatePtr SourceControlState)
 {
-	if(SourceControlState.IsValid())
+	TWeakPtr<FUncontrolledChangelistState> UncontrolledChangelistState = ItemSourceControl->GetUncontrolledChangelistState();
+	
+	UpdateWidget(SourceControlState, UncontrolledChangelistState.Pin());
+}
+
+void SSourceControlWidget::UpdateUncontrolledChangelistState(TSharedPtr<FUncontrolledChangelistState> UncontrolledChangelistState)
+{
+	UpdateWidget(ItemSourceControl->GetSourceControlState(), UncontrolledChangelistState);
+}
+
+void SSourceControlWidget::UpdateWidget(FSourceControlStatePtr SourceControlState, const TSharedPtr<FUncontrolledChangelistState>& UncontrolledChangelistState)
+{
+	RemoveAllLayers();
+
+	static const FName UncontrolledChangelistIconName = TEXT("Icons.Unlink");
+
+	// Custom logic to combine uncontrolled and source control state
+	if (SourceControlState.IsValid() && UncontrolledChangelistState.IsValid())
+	{
+		// We use the uncontrolled icon as the base icon
+		const FSlateBrush* ImageBrush = FAppStyle::GetBrush(UncontrolledChangelistIconName);
+
+		// We use the overlay from the source control icon (if any)
+		const FSlateBrush* OverlayBrush = SourceControlState->GetIcon().GetOverlayIcon();
+
+		// We use the color of the source control icon
+		FSlateColor IconColor = SourceControlState->GetIcon().GetIcon()->TintColor;
+
+		SetImage(ImageBrush);
+
+		if (OverlayBrush)
+		{
+			AddLayer(OverlayBrush);
+		}
+		
+		SetColorAndOpacity(IconColor);
+
+		SetToolTipText(FText::Format(LOCTEXT("UncontrolledSourceControlCombinedTooltip",
+			"{0}\nStatus in source control: {1}"),
+			UncontrolledChangelistState->GetDisplayTooltip(), SourceControlState->GetDisplayTooltip()));
+	}
+	else if (SourceControlState.IsValid() && !UncontrolledChangelistState.IsValid())
 	{
 		SetFromSlateIcon(SourceControlState->GetIcon());
 		SetToolTipText(SourceControlState->GetDisplayTooltip());
 	}
+	else if (!SourceControlState.IsValid() && UncontrolledChangelistState.IsValid())
+	{
+		SetImage(FAppStyle::GetBrush(UncontrolledChangelistIconName));
+		SetToolTipText(UncontrolledChangelistState->GetDisplayTooltip());
+	}
 	else
 	{
 		SetImage(nullptr);
 		SetToolTipText(TAttribute<FText>());
-		RemoveAllLayers();
 	}
 }
+
+#undef LOCTEXT_NAMESPACE
\ No newline at end of file
diff --git a/Engine/Source/Editor/SceneOutliner/Private/SSourceControlWidget.h b/Engine/Source/Editor/SceneOutliner/Private/SSourceControlWidget.h
--- a/Engine/Source/Editor/SceneOutliner/Private/SSourceControlWidget.h
+++ b/Engine/Source/Editor/SceneOutliner/Private/SSourceControlWidget.h
@@ -20,6 +20,9 @@ private:
 	virtual FReply OnMouseButtonDoubleClick(const FGeometry& InMyGeometry, const FPointerEvent& InMouseEvent) override;
 
 	void UpdateSourceControlState(FSourceControlStatePtr SourceControlState);
+	void UpdateUncontrolledChangelistState(TSharedPtr<FUncontrolledChangelistState> UncontrolledChangelistState);
+
+	void UpdateWidget(FSourceControlStatePtr SourceControlState, const TSharedPtr<FUncontrolledChangelistState>& UncontrolledChangelistState);
 
 	/** The object that keeps the source control state for the outliner */
 	TSharedPtr<FSceneOutlinerTreeItemSCC> ItemSourceControl;
diff --git a/Engine/Source/Editor/SceneOutliner/Private/SceneOutlinerHelpers.cpp b/Engine/Source/Editor/SceneOutliner/Private/SceneOutlinerHelpers.cpp
new file mode 100644
--- /dev/null
+++ b/Engine/Source/Editor/SceneOutliner/Private/SceneOutlinerHelpers.cpp
@@ -0,0 +1,79 @@
+// Copyright Epic Games, Inc. All Rights Reserved.
+
+#include "SceneOutlinerHelpers.h"
+#include "ActorTreeItem.h"
+#include "ActorDescTreeItem.h"
+#include "ActorFolderTreeItem.h"
+#include "WorldPartition/WorldPartitionActorDesc.h"
+
+#include "SourceControlHelpers.h"
+
+namespace SceneOutliner
+{
+	FString FSceneOutlinerHelpers::GetExternalPackageName(const ISceneOutlinerTreeItem& TreeItem)
+	{
+		if (const FActorTreeItem* ActorItem = TreeItem.CastTo<FActorTreeItem>())
+		{
+			if (AActor* Actor = ActorItem->Actor.Get())
+			{
+				if (Actor->IsPackageExternal())
+				{
+					return USourceControlHelpers::PackageFilename(Actor->GetExternalPackage());
+				}
+
+			}
+		}
+		else if (const FActorFolderTreeItem* ActorFolderItem = TreeItem.CastTo<FActorFolderTreeItem>())
+		{
+			if (const UActorFolder* ActorFolder = ActorFolderItem->GetActorFolder())
+			{
+				if (ActorFolder->IsPackageExternal())
+				{
+					return USourceControlHelpers::PackageFilename(ActorFolder->GetExternalPackage());
+				}
+			}
+		}
+		else if (const FActorDescTreeItem* ActorDescItem = TreeItem.CastTo<FActorDescTreeItem>())
+		{
+			if (const FWorldPartitionActorDesc* ActorDesc = ActorDescItem->ActorDescHandle.Get())
+			{
+				return USourceControlHelpers::PackageFilename(ActorDesc->GetActorPackage().ToString());
+			}
+		}
+
+		return FString();
+	}
+	
+	UPackage* FSceneOutlinerHelpers::GetExternalPackage(const ISceneOutlinerTreeItem& TreeItem)
+	{
+		if (const FActorTreeItem* ActorItem = TreeItem.CastTo<FActorTreeItem>())
+		{
+			if (AActor* Actor = ActorItem->Actor.Get())
+			{
+				if (Actor->IsPackageExternal())
+				{
+					return Actor->GetExternalPackage();
+				}
+			}
+		}
+		else if (const FActorFolderTreeItem* ActorFolderItem = TreeItem.CastTo<FActorFolderTreeItem>())
+		{
+			if (const UActorFolder* ActorFolder = ActorFolderItem->GetActorFolder())
+			{
+				if (ActorFolder->IsPackageExternal())
+				{
+					return ActorFolder->GetExternalPackage();
+				}
+			}
+		}
+		else if (const FActorDescTreeItem* ActorDescItem = TreeItem.CastTo<FActorDescTreeItem>())
+		{
+			if (const FWorldPartitionActorDesc* ActorDesc = ActorDescItem->ActorDescHandle.Get())
+			{
+				return FindPackage(nullptr, *ActorDesc->GetActorPackage().ToString());
+			}
+		}
+
+		return nullptr;
+	}
+;}
\ No newline at end of file
diff --git a/Engine/Source/Editor/SceneOutliner/Private/SceneOutlinerModule.cpp b/Engine/Source/Editor/SceneOutliner/Private/SceneOutlinerModule.cpp
--- a/Engine/Source/Editor/SceneOutliner/Private/SceneOutlinerModule.cpp
+++ b/Engine/Source/Editor/SceneOutliner/Private/SceneOutlinerModule.cpp
@@ -12,6 +12,7 @@
 #include "SceneOutlinerActorSCCColumn.h"
 #include "SceneOutlinerPinnedColumn.h"
 #include "SceneOutlinerTextInfoColumn.h"
+#include "SceneOutlinerUnsavedColumn.h"
 #include "SceneOutlinerPublicTypes.h"
 
 #include "ActorPickingMode.h"
@@ -51,6 +52,7 @@ void FSceneOutlinerModule::StartupModule()
 	RegisterColumnType<FTypeInfoColumn>();
 	RegisterColumnType<FSceneOutlinerActorSCCColumn>();
 	RegisterColumnType<FSceneOutlinerPinnedColumn>();
+	RegisterColumnType<FSceneOutlinerActorUnsavedColumn>();
 }
 
 
@@ -61,6 +63,7 @@ void FSceneOutlinerModule::ShutdownModule()
 	UnRegisterColumnType<FTypeInfoColumn>();
 	UnRegisterColumnType<FSceneOutlinerActorSCCColumn>();
 	UnRegisterColumnType<FSceneOutlinerPinnedColumn>();
+	UnRegisterColumnType<FSceneOutlinerActorUnsavedColumn>();
 }
 
 TSharedRef<ISceneOutliner> FSceneOutlinerModule::CreateSceneOutliner(const FSceneOutlinerInitializationOptions& InitOptions) const
@@ -161,18 +164,23 @@ TSharedRef< ISceneOutliner > FSceneOutlinerModule::CreateActorBrowser(const FSce
 		InitOptions.ColumnMap.Add(FSceneOutlinerBuiltInColumnTypes::ActorInfo(), FSceneOutlinerColumnInfo(ESceneOutlinerColumnVisibility::Visible, 20, FCreateSceneOutlinerColumn(), true, TOptional<float>(), FSceneOutlinerBuiltInColumnTypes::ActorInfo_Localized()));
 		InitOptions.ColumnMap.Add(FSceneOutlinerBuiltInColumnTypes::SourceControl(), FSceneOutlinerColumnInfo(ESceneOutlinerColumnVisibility::Invisible, 30, FCreateSceneOutlinerColumn(), true, TOptional<float>(), FSceneOutlinerBuiltInColumnTypes::SourceControl_Localized()));
 
-		// TODO: Figure out a better way to get the current world
+		ESceneOutlinerColumnVisibility UnsavedColumnVisibility = ESceneOutlinerColumnVisibility::Invisible;
+		
 		UWorld* WorldPtr = GEditor->GetEditorWorldContext().World();
-
 		if (WorldPtr)
 		{
-			// We don't want the pinned column in non wp levels
 			if (WorldPtr->IsPartitionedWorld())
 			{
+				// We don't want the pinned column in non wp levels
 				InitOptions.ColumnMap.Add(FSceneOutlinerBuiltInColumnTypes::Pinned(), FSceneOutlinerColumnInfo(ESceneOutlinerColumnVisibility::Visible, 5, FCreateSceneOutlinerColumn(), true, TOptional<float>(), FSceneOutlinerBuiltInColumnTypes::Pinned_Localized()));
+
+				// We want the unsaved column to be visible by default in partitioned levels
+				UnsavedColumnVisibility = ESceneOutlinerColumnVisibility::Visible;
 			}
 			CreateActorInfoColumns(InitOptions, WorldPtr);
 		}
+
+		InitOptions.ColumnMap.Add(FSceneOutlinerBuiltInColumnTypes::Unsaved(), FSceneOutlinerColumnInfo(UnsavedColumnVisibility, 1, FCreateSceneOutlinerColumn(), true, TOptional<float>(), FSceneOutlinerBuiltInColumnTypes::Unsaved_Localized()));
 		
 	}
 
diff --git a/Engine/Source/Editor/SceneOutliner/Private/SceneOutlinerTreeItemSCC.cpp b/Engine/Source/Editor/SceneOutliner/Private/SceneOutlinerTreeItemSCC.cpp
--- a/Engine/Source/Editor/SceneOutliner/Private/SceneOutlinerTreeItemSCC.cpp
+++ b/Engine/Source/Editor/SceneOutliner/Private/SceneOutlinerTreeItemSCC.cpp
@@ -8,7 +8,9 @@
 #include "ISourceControlProvider.h"
 #include "ISourceControlModule.h"
 #include "ProfilingDebugging/CpuProfilerTrace.h"
+#include "SceneOutlinerHelpers.h"
 #include "WorldPartition/WorldPartitionActorDesc.h"
+#include "UncontrolledChangelistsModule.h"
 
 FSceneOutlinerTreeItemSCC::FSceneOutlinerTreeItemSCC(FSceneOutlinerTreeItemPtr InTreeItemPtr)
 {
@@ -16,17 +18,16 @@ FSceneOutlinerTreeItemSCC::FSceneOutlinerTreeItemSCC(FSceneOutlinerTreeItemPtr I
 
 	if (TreeItemPtr.IsValid())
 	{
+		ExternalPackageName = SceneOutliner::FSceneOutlinerHelpers::GetExternalPackageName(*TreeItemPtr.Get());
+		ExternalPackage = SceneOutliner::FSceneOutlinerHelpers::GetExternalPackage(*TreeItemPtr.Get());
+		
 		if (FActorTreeItem* ActorItem = TreeItemPtr->CastTo<FActorTreeItem>())
 		{
 			if (AActor* Actor = ActorItem->Actor.Get())
 			{
 				if (Actor->IsPackageExternal())
 				{
-					ExternalPackageName = USourceControlHelpers::PackageFilename(Actor->GetExternalPackage());
-					ExternalPackage = Actor->GetExternalPackage();
-				}
-
-				ActorPackingModeChangedDelegateHandle = Actor->OnPackagingModeChanged.AddLambda([this](AActor* InActor, bool bExternal)
+					ActorPackingModeChangedDelegateHandle = Actor->OnPackagingModeChanged.AddLambda([this](AActor* InActor, bool bExternal)
 					{
 						if (bExternal)
 						{
@@ -41,37 +42,29 @@ FSceneOutlinerTreeItemSCC::FSceneOutlinerTreeItemSCC(FSceneOutlinerTreeItemPtr I
 							DisconnectSourceControl();
 						}
 					});
-			}
-		}
-		else if (FActorFolderTreeItem* ActorFolderItem = TreeItemPtr->CastTo<FActorFolderTreeItem>())
-		{
-			if (const UActorFolder* ActorFolder = ActorFolderItem->GetActorFolder())
-			{
-				if (ActorFolder->IsPackageExternal())
-				{
-					ExternalPackageName = USourceControlHelpers::PackageFilename(ActorFolder->GetExternalPackage());
-					ExternalPackage = ActorFolder->GetExternalPackage();
 				}
 			}
 		}
-		else if (FActorDescTreeItem* ActorDescItem = TreeItemPtr->CastTo<FActorDescTreeItem>())
-		{
-			if (const FWorldPartitionActorDesc* ActorDesc = ActorDescItem->ActorDescHandle.Get())
-			{
-				ExternalPackageName =  USourceControlHelpers::PackageFilename(ActorDesc->GetActorPackage().ToString());
-				ExternalPackage = FindPackage(nullptr, *ActorDesc->GetActorPackage().ToString());
-			}
-		}
-
+		
 		if (!ExternalPackageName.IsEmpty())
 		{
 			ConnectSourceControl();
 		}
 	}
+
+	FUncontrolledChangelistsModule& UncontrolledChangelistModule = FUncontrolledChangelistsModule::Get();
+	UncontrolledChangelistChangedHandle = UncontrolledChangelistModule.OnUncontrolledChangelistModuleChanged.AddRaw(this, &FSceneOutlinerTreeItemSCC::HandleUncontrolledChangelistsStateChanged);
+
+	// Call the delegate to update the initial uncontrolled state
+	HandleUncontrolledChangelistsStateChanged();
+
 }
 
 FSceneOutlinerTreeItemSCC::~FSceneOutlinerTreeItemSCC()
 {
+	FUncontrolledChangelistsModule& UncontrolledChangelistModule = FUncontrolledChangelistsModule::Get();
+	UncontrolledChangelistModule.OnUncontrolledChangelistModuleChanged.Remove(UncontrolledChangelistChangedHandle);
+
 	DisconnectSourceControl();
 }
 
@@ -146,3 +139,27 @@ void FSceneOutlinerTreeItemSCC::BroadcastNewState(FSourceControlStatePtr SourceC
 {
 	OnSourceControlStateChanged.ExecuteIfBound(SourceControlState);
 }
+
+void FSceneOutlinerTreeItemSCC::HandleUncontrolledChangelistsStateChanged()
+{
+	TSharedPtr<FUncontrolledChangelistState> PrevUncontrolledChangelistState = UncontrolledChangelistState;
+	
+	UncontrolledChangelistState = nullptr;
+	
+	TArray<FUncontrolledChangelistStateRef> UncontrolledChangelistStates = FUncontrolledChangelistsModule::Get().GetChangelistStates();
+
+	for (const TSharedRef<FUncontrolledChangelistState>& UncontrolledChangelistStateRef : UncontrolledChangelistStates)
+	{
+		if(UncontrolledChangelistStateRef->GetFilenames().Contains(ExternalPackageName))
+		{
+			UncontrolledChangelistState = UncontrolledChangelistStateRef;
+			break;
+		}
+	}
+
+	// Broadcast the delegate if our uncontrolled status was changed
+	if(UncontrolledChangelistState != PrevUncontrolledChangelistState)
+	{
+		OnUncontrolledChangelistsStateChanged.ExecuteIfBound(UncontrolledChangelistState);
+	}
+}
\ No newline at end of file
diff --git a/Engine/Source/Editor/SceneOutliner/Private/SceneOutlinerUnsavedColumn.cpp b/Engine/Source/Editor/SceneOutliner/Private/SceneOutlinerUnsavedColumn.cpp
new file mode 100644
--- /dev/null
+++ b/Engine/Source/Editor/SceneOutliner/Private/SceneOutlinerUnsavedColumn.cpp
@@ -0,0 +1,161 @@
+// Copyright Epic Games, Inc. All Rights Reserved.
+
+#include "SceneOutlinerUnsavedColumn.h"
+#include "ActorTreeItem.h"
+#include "ActorDescTreeItem.h"
+#include "ActorFolderTreeItem.h"
+#include "SourceControlHelpers.h"
+#include "WorldPartition/WorldPartitionActorDesc.h"
+#include "UnsavedAssetsTrackerModule.h"
+#include "SortHelper.h"
+#include "ISceneOutliner.h"
+#include "ISceneOutlinerMode.h"
+#include "SceneOutlinerHelpers.h"
+
+#define LOCTEXT_NAMESPACE "SceneOutlinerActorUnsavedColumn"
+
+class SUnsavedActorWidget : public SImage
+{
+public:
+	SLATE_BEGIN_ARGS(SUnsavedActorWidget) {}
+	SLATE_END_ARGS()
+
+	/** Construct this widget */
+	void Construct(const FArguments& InArgs, const FString& InActorExternalPackageName);
+	
+	bool IsUnsaved() const;
+	
+private:
+	void UpdateImage();
+	void OnUnsavedAssetAdded(const FString& FileAbsPathname);
+	void OnUnsavedAssetRemoved(const FString& FileAbsPathname);
+
+private:
+	FString ExternalPackageName;
+	bool bIsUnsaved;
+};
+
+void SUnsavedActorWidget::Construct(const FArguments& InArgs, const FString& InActorExternalPackageName)
+{
+	ExternalPackageName = InActorExternalPackageName;
+
+	SImage::Construct(
+			SImage::FArguments()
+			.ColorAndOpacity(FSlateColor::UseForeground())
+			.Image(FStyleDefaults::GetNoBrush()));
+	
+	FUnsavedAssetsTrackerModule& UnsavedAssetsTrackerModule = FUnsavedAssetsTrackerModule::Get();
+	UnsavedAssetsTrackerModule.OnUnsavedAssetAdded.AddSP(this, &SUnsavedActorWidget::OnUnsavedAssetAdded);
+	UnsavedAssetsTrackerModule.OnUnsavedAssetRemoved.AddSP(this, &SUnsavedActorWidget::OnUnsavedAssetRemoved);
+
+	bIsUnsaved = UnsavedAssetsTrackerModule.IsAssetUnsaved(ExternalPackageName);
+
+	UpdateImage();
+}
+
+bool SUnsavedActorWidget::IsUnsaved() const
+{
+	return bIsUnsaved;
+}
+
+void SUnsavedActorWidget::OnUnsavedAssetAdded(const FString& FileAbsPathname)
+{
+	if (FileAbsPathname == ExternalPackageName)
+	{
+		// We should never be desynced, i.e if this item was added as an unsaved asset bIsUnsavedAsset MUST be false before
+		check(!bIsUnsaved)
+		bIsUnsaved = true;
+		
+		UpdateImage();
+	}
+}
+
+void SUnsavedActorWidget::OnUnsavedAssetRemoved(const FString& FileAbsPathname)
+{
+	if (FileAbsPathname == ExternalPackageName)
+	{
+		// We should never be desynced, i.e if this item was removed from the unsaved asset list bIsUnsavedAsset MUST be true before
+		check(bIsUnsaved)
+		bIsUnsaved = false;
+
+		UpdateImage();
+	}
+}
+
+void SUnsavedActorWidget::UpdateImage()
+{
+	if (bIsUnsaved)
+	{
+		SetImage(FAppStyle::GetBrush("Icons.DirtyBadge"));
+	}
+	else
+	{
+		SetImage(nullptr);
+	}
+}
+
+
+FName FSceneOutlinerActorUnsavedColumn::GetColumnID()
+{
+	return GetID();
+}
+
+SHeaderRow::FColumn::FArguments FSceneOutlinerActorUnsavedColumn::ConstructHeaderRowColumn()
+{
+	return SHeaderRow::Column(GetColumnID())
+		.FixedWidth(24.f)
+		.HAlignHeader(HAlign_Center)
+		.VAlignHeader(VAlign_Center)
+		.HAlignCell(HAlign_Center)
+		.VAlignCell(VAlign_Center)
+		.DefaultTooltip(FText::FromName(GetColumnID()))
+		[
+			SNew(SImage)
+			.Image(FAppStyle::GetBrush("Icons.DirtyBadge"))
+			.ColorAndOpacity(FSlateColor::UseForeground())
+		];
+}
+
+const TSharedRef<SWidget> FSceneOutlinerActorUnsavedColumn::ConstructRowWidget(FSceneOutlinerTreeItemRef TreeItem, const STableRow<FSceneOutlinerTreeItemPtr>& Row)
+{
+	FString AssetPath = SceneOutliner::FSceneOutlinerHelpers::GetExternalPackageName(TreeItem.Get());
+
+	if (AssetPath.IsEmpty())
+	{
+		return SNullWidget::NullWidget;
+	}
+
+	TSharedRef<SUnsavedActorWidget> ActualWidget = SNew(SUnsavedActorWidget, AssetPath);
+
+	// Store a reference to the widget for this asset so we can grab the unsaved status for sorting
+	UnsavedActorWidgets.Add(AssetPath, ActualWidget);
+	
+	return SNew(SHorizontalBox)
+			+ SHorizontalBox::Slot()
+			.AutoWidth()
+			.HAlign(HAlign_Center)
+			.VAlign(VAlign_Center)
+			[
+				ActualWidget
+			];
+}
+
+void FSceneOutlinerActorUnsavedColumn::SortItems(TArray<FSceneOutlinerTreeItemPtr>& RootItems, const EColumnSortMode::Type SortMode) const
+{
+	FSceneOutlinerSortHelper<bool, SceneOutliner::FNumericStringWrapper>()
+		/** Sort by unsaved first */
+		.Primary([this](const ISceneOutlinerTreeItem& Item)
+		{
+			if (const TSharedRef<SUnsavedActorWidget>* FoundWidget = UnsavedActorWidgets.Find(SceneOutliner::FSceneOutlinerHelpers::GetExternalPackageName(Item)))
+			{
+				// return the inverse because we want items with Unsaved = TRUE to show up on top in ascending sort (default)
+				return !(*FoundWidget)->IsUnsaved();
+			}
+			return true;
+		}, SortMode)
+		/** Then by type */
+		.Secondary([this](const ISceneOutlinerTreeItem& Item){ return SceneOutliner::FNumericStringWrapper(Item.GetDisplayString()); }, SortMode)
+		.Sort(RootItems);
+}
+
+#undef LOCTEXT_NAMESPACE
\ No newline at end of file
diff --git a/Engine/Source/Editor/SceneOutliner/Private/SceneOutlinerUnsavedColumn.h b/Engine/Source/Editor/SceneOutliner/Private/SceneOutlinerUnsavedColumn.h
new file mode 100644
--- /dev/null
+++ b/Engine/Source/Editor/SceneOutliner/Private/SceneOutlinerUnsavedColumn.h
@@ -0,0 +1,39 @@
+// Copyright Epic Games, Inc. All Rights Reserved.
+
+#pragma once
+
+#include "Widgets/SWidget.h"
+#include "Widgets/Views/SHeaderRow.h"
+#include "ISceneOutliner.h"
+#include "SceneOutlinerPublicTypes.h"
+#include "ISceneOutlinerColumn.h"
+
+template<typename ItemType> class STableRow;
+class SUnsavedActorWidget;
+
+/** A column for the SceneOutliner that displays if the actor is unsaved */
+class FSceneOutlinerActorUnsavedColumn : public ISceneOutlinerColumn
+{
+public:
+	FSceneOutlinerActorUnsavedColumn(ISceneOutliner& SceneOutliner) : WeakOutliner(StaticCastSharedRef<ISceneOutliner>(SceneOutliner.AsShared())) {}
+	virtual ~FSceneOutlinerActorUnsavedColumn() = default;
+
+	static FName GetID() { return FSceneOutlinerBuiltInColumnTypes::Unsaved(); }
+
+	//////////////////////////////////////////////////////////////////////////
+	// Begin ISceneOutlinerColumn Implementation
+	virtual FName GetColumnID() override;
+	virtual SHeaderRow::FColumn::FArguments ConstructHeaderRowColumn() override;
+	virtual const TSharedRef<SWidget> ConstructRowWidget(FSceneOutlinerTreeItemRef TreeItem, const STableRow<FSceneOutlinerTreeItemPtr>& Row) override;
+	virtual bool SupportsSorting() const override { return true; }
+	virtual void SortItems(TArray<FSceneOutlinerTreeItemPtr>& RootItems, const EColumnSortMode::Type SortMode) const override;
+	// End ISceneOutlinerColumn Implementation
+	//////////////////////////////////////////////////////////////////////////
+
+private:
+	TWeakPtr<ISceneOutliner> WeakOutliner;
+
+	// Map of asset path -> widgets for each asset we encounter
+	TMap<FString, TSharedRef<SUnsavedActorWidget>> UnsavedActorWidgets;
+	
+};
diff --git a/Engine/Source/Editor/SceneOutliner/Public/ISceneOutliner.h b/Engine/Source/Editor/SceneOutliner/Public/ISceneOutliner.h
--- a/Engine/Source/Editor/SceneOutliner/Public/ISceneOutliner.h
+++ b/Engine/Source/Editor/SceneOutliner/Public/ISceneOutliner.h
@@ -150,6 +150,9 @@ public:
 
 	/** Get the associated source control object for the specified item. */
 	virtual TSharedPtr<FSceneOutlinerTreeItemSCC> GetItemSourceControl(const FSceneOutlinerTreeItemPtr& InItem) = 0;
+
+	/** Check if a filter with the given name exists and is active in the filter bar for this Outliner (if this Outliner has a filter bar). */
+	virtual bool IsFilterActive(const FString& FilterName) const = 0;
 protected:
 	ISceneOutlinerMode* Mode;
 };
diff --git a/Engine/Source/Editor/SceneOutliner/Public/SSceneOutliner.h b/Engine/Source/Editor/SceneOutliner/Public/SSceneOutliner.h
--- a/Engine/Source/Editor/SceneOutliner/Public/SSceneOutliner.h
+++ b/Engine/Source/Editor/SceneOutliner/Public/SSceneOutliner.h
@@ -372,6 +372,9 @@ public:
 	/** Hook to add custom options to toolbar in a derived class */
 	virtual void CustomAddToToolbar(TSharedPtr<class SHorizontalBox> Toolbar) {}
 
+	/** Check if a filter with the given name exists and is active in the filter bar for this Outliner (if this Outliner has a filter bar). */
+	virtual bool IsFilterActive(const FString& FilterName) const override;
+
 public:
 	/** Event to react to a user double click on a item */
 	SceneOutliner::FTreeItemPtrEvent& GetDoubleClickEvent() { return OnDoubleClickOnTreeEvent; }
diff --git a/Engine/Source/Editor/SceneOutliner/Public/SceneOutlinerHelpers.h b/Engine/Source/Editor/SceneOutliner/Public/SceneOutlinerHelpers.h
new file mode 100644
--- /dev/null
+++ b/Engine/Source/Editor/SceneOutliner/Public/SceneOutlinerHelpers.h
@@ -0,0 +1,16 @@
+// Copyright Epic Games, Inc. All Rights Reserved.
+
+#pragma once
+
+#include "SceneOutlinerPublicTypes.h"
+
+namespace SceneOutliner
+{
+	// Class to hold common functionality needed by the Outliner and helpful to modules creating Outliner instances
+	class SCENEOUTLINER_API FSceneOutlinerHelpers
+	{
+	public:
+		static FString GetExternalPackageName(const ISceneOutlinerTreeItem& TreeItem);
+		static UPackage* GetExternalPackage(const ISceneOutlinerTreeItem& TreeItem);
+	};
+};
diff --git a/Engine/Source/Editor/SceneOutliner/Public/SceneOutlinerPublicTypes.h b/Engine/Source/Editor/SceneOutliner/Public/SceneOutlinerPublicTypes.h
--- a/Engine/Source/Editor/SceneOutliner/Public/SceneOutlinerPublicTypes.h
+++ b/Engine/Source/Editor/SceneOutliner/Public/SceneOutlinerPublicTypes.h
@@ -58,6 +58,7 @@ struct FSceneOutlinerBuiltInColumnTypes
 	DEFINE_SCENEOUTLINER_BUILTIN_COLUMN_TYPE(PackageShortName, "Package Short Name", "SceneOutlinerPackageShortNameColumn", "Package Short Name");
 	DEFINE_SCENEOUTLINER_BUILTIN_COLUMN_TYPE(UncachedLights, "Uncached Lights", "SceneOutlinerUncachedLightsColumn", "# Uncached Lights");
 	DEFINE_SCENEOUTLINER_BUILTIN_COLUMN_TYPE(Socket, "Socket", "SceneOutlinerSocketColumn", "Socket");
+	DEFINE_SCENEOUTLINER_BUILTIN_COLUMN_TYPE(Unsaved, "Unsaved", "UnsavedColumnName", "Unsaved");
 };
 
 /** Visibility enum for scene outliner columns */
diff --git a/Engine/Source/Editor/SceneOutliner/Public/SceneOutlinerTreeItemSCC.h b/Engine/Source/Editor/SceneOutliner/Public/SceneOutlinerTreeItemSCC.h
--- a/Engine/Source/Editor/SceneOutliner/Public/SceneOutlinerTreeItemSCC.h
+++ b/Engine/Source/Editor/SceneOutliner/Public/SceneOutlinerTreeItemSCC.h
@@ -7,7 +7,11 @@
 
 #define LOCTEXT_NAMESPACE "SceneOutlinerTreeItemSCC"
 
+class FUncontrolledChangelistState;
+
 DECLARE_DELEGATE_OneParam(FSourceControlStateChangedDelegate, FSourceControlStatePtr);
+DECLARE_DELEGATE_OneParam(FUncontrolledStateChangedDelegate, TSharedPtr<FUncontrolledChangelistState>);
+
 
 class FSceneOutlinerTreeItemSCC : public TSharedFromThis<FSceneOutlinerTreeItemSCC>
 {
@@ -27,6 +31,10 @@ public:
 	UPackage* GetPackage() { return ExternalPackage; }
 
 	FSourceControlStateChangedDelegate OnSourceControlStateChanged;
+	
+	FUncontrolledStateChangedDelegate OnUncontrolledChangelistsStateChanged;
+	
+	TWeakPtr<FUncontrolledChangelistState> GetUncontrolledChangelistState() { return UncontrolledChangelistState; }
 
 private:
 
@@ -40,6 +48,8 @@ private:
 
 	void BroadcastNewState(FSourceControlStatePtr SourceControlState);
 
+	void HandleUncontrolledChangelistsStateChanged();
+
 	/** The tree item we relate to */
 	FSceneOutlinerTreeItemPtr TreeItemPtr;
 
@@ -55,6 +65,12 @@ private:
 
 	/** Actor packaging mode changed delegate handle */
 	FDelegateHandle ActorPackingModeChangedDelegateHandle;
+
+	/** Uncontrolled Changelist Changed changed delegate handle */
+	FDelegateHandle UncontrolledChangelistChangedHandle;
+
+	/** The uncontrolled changelist this file belongs to, empty if this item is not uncontrolled */
+	TSharedPtr<FUncontrolledChangelistState> UncontrolledChangelistState;
 };
 
 #undef LOCTEXT_NAMESPACE
diff --git a/Engine/Source/Editor/SceneOutliner/SceneOutliner.Build.cs b/Engine/Source/Editor/SceneOutliner/SceneOutliner.Build.cs
--- a/Engine/Source/Editor/SceneOutliner/SceneOutliner.Build.cs
+++ b/Engine/Source/Editor/SceneOutliner/SceneOutliner.Build.cs
@@ -25,6 +25,8 @@ public class SceneOutliner : ModuleRules
 				"UncontrolledChangelists",
 				"EditorWidgets",
 				"ToolWidgets",
+				"UnsavedAssetsTracker",
+				"UncontrolledChangelists"
 			}
 		);
 		
