Added "LevelViewportToolBar.RightExtension"

diff --git a/Engine/Source/Editor/LevelEditor/Private/SLevelViewportToolBar.cpp b/Engine/Source/Editor/LevelEditor/Private/SLevelViewportToolBar.cpp
--- a/Engine/Source/Editor/LevelEditor/Private/SLevelViewportToolBar.cpp
+++ b/Engine/Source/Editor/LevelEditor/Private/SLevelViewportToolBar.cpp
@@ -305,6 +305,16 @@ void SLevelViewportToolBar::Construct( const FArguments& InArgs )
 					.ExtensionContext(ExtensionContextObject)
 				]
 				+ SHorizontalBox::Slot()
+				.Padding(ToolbarSlotPadding)
+				.AutoWidth()
+				.HAlign(HAlign_Right)
+				.VAlign(VAlign_Fill)
+				[
+					SNew(SExtensionPanel)
+					.ExtensionPanelID("LevelViewportToolBar.RightExtension")
+					.ExtensionContext(ExtensionContextObject)
+				]
+				+ SHorizontalBox::Slot()
 				.AutoWidth()
 				.MaxWidth(TAttribute<float>::CreateSP(this, &SLevelViewportToolBar::GetTransformToolbarWidth))
 				.Padding(ToolbarSlotPadding)
