Expose multiple viewport-related editor behaviors to editor utility widgets.

diff --git a/Engine/Source/Editor/LevelEditor/Private/LevelEditorSubsystem.cpp b/Engine/Source/Editor/LevelEditor/Private/LevelEditorSubsystem.cpp
--- a/Engine/Source/Editor/LevelEditor/Private/LevelEditorSubsystem.cpp
+++ b/Engine/Source/Editor/LevelEditor/Private/LevelEditorSubsystem.cpp
@@ -67,6 +67,14 @@ void ULevelEditorSubsystem::Initialize(FSubsystemCollectionBase& Collection)
 	FWorldDelegates::LevelAddedToWorld.AddUObject(this, &ULevelEditorSubsystem::OnLevelAddedOrRemoved);
 	FWorldDelegates::LevelRemovedFromWorld.AddUObject(this, &ULevelEditorSubsystem::OnLevelAddedOrRemoved);
 	FWorldDelegates::OnCurrentLevelChanged.AddUObject(this, &ULevelEditorSubsystem::OnCurrentLevelChanged);
+
+	FEditorDelegates::PreSaveWorldWithContext.AddUObject(this, &ULevelEditorSubsystem::HandleOnPreSaveWorldWithContext);
+	FEditorDelegates::PostSaveWorldWithContext.AddUObject(this, &ULevelEditorSubsystem::HandleOnPostSaveWorldWithContext);
+
+	FEditorDelegates::OnEditorCameraMoved.AddUObject(this, &ULevelEditorSubsystem::HandleOnEditorCameraMoved);
+
+	FEditorDelegates::MapChange.AddUObject(this, &ULevelEditorSubsystem::HandleOnMapChanged);
+	FEditorDelegates::OnMapOpened.AddUObject(this, &ULevelEditorSubsystem::HandleOnMapOpened);
 }
 
 void ULevelEditorSubsystem::Deinitialize()
@@ -77,6 +85,14 @@ void ULevelEditorSubsystem::Deinitialize()
 	FWorldDelegates::OnCurrentLevelChanged.RemoveAll(this);
 	UToolMenus::UnRegisterStartupCallback(this);
 	UToolMenus::UnregisterOwner(this);
+
+	FEditorDelegates::PreSaveWorldWithContext.RemoveAll(this);
+	FEditorDelegates::PostSaveWorldWithContext.RemoveAll(this);
+
+	FEditorDelegates::OnEditorCameraMoved.RemoveAll(this);
+
+	FEditorDelegates::MapChange.RemoveAll(this);
+	FEditorDelegates::OnMapOpened.RemoveAll(this);
 }
 
 void ULevelEditorSubsystem::ExtendQuickActionMenu()
@@ -901,6 +917,31 @@ void ULevelEditorSubsystem::OnCurrentLevelChanged(ULevel* InNewLevel, ULevel* In
 	ActorEditorContextClientChanged.Broadcast(this);
 }
 
+void ULevelEditorSubsystem::HandleOnPreSaveWorldWithContext(class UWorld* World, FObjectPreSaveContext ObjectSaveContext)
+{
+	OnPreSaveWorld.Broadcast(ObjectSaveContext.GetSaveFlags(), World);
+}
+
+void ULevelEditorSubsystem::HandleOnPostSaveWorldWithContext(class UWorld* World, FObjectPostSaveContext ObjectSaveContext)
+{
+	OnPostSaveWorld.Broadcast(ObjectSaveContext.GetSaveFlags(), World, ObjectSaveContext.SaveSucceeded());
+}
+
+void ULevelEditorSubsystem::HandleOnEditorCameraMoved(const FVector& Location, const FRotator& Rotation, ELevelViewportType ViewportType, int32 ViewIndex)
+{
+	OnEditorCameraMoved.Broadcast(Location, Rotation, ViewportType, ViewIndex);
+}
+
+void ULevelEditorSubsystem::HandleOnMapChanged(uint32 MapChangeFlags)
+{
+	OnMapChanged.Broadcast(MapChangeFlags);
+}
+
+void ULevelEditorSubsystem::HandleOnMapOpened(const FString& Filename, bool bAsTemplate)
+{
+	OnMapOpened.Broadcast(Filename, bAsTemplate);
+}
+
 TSharedRef<SWidget> ULevelEditorSubsystem::GetActorEditorContextWidget(UWorld* InWorld) const
 {
 	return SNew(SCurrentLevelWidget).World(InWorld);
diff --git a/Engine/Source/Editor/LevelEditor/Public/LevelEditorSubsystem.h b/Engine/Source/Editor/LevelEditor/Public/LevelEditorSubsystem.h
--- a/Engine/Source/Editor/LevelEditor/Public/LevelEditorSubsystem.h
+++ b/Engine/Source/Editor/LevelEditor/Public/LevelEditorSubsystem.h
@@ -4,8 +4,10 @@
 
 #include "CoreMinimal.h"
 #include "EditorSubsystem.h"
-#include "IActorEditorContextClient.h"
+#include "Editor/UnrealEdTypes.h"
 #include "Engine/EngineTypes.h"
+#include "IActorEditorContextClient.h"
+#include "UObject/ObjectSaveContext.h"
 
 #include "LevelEditorSubsystem.generated.h"
 
@@ -13,6 +15,19 @@ class UTypedElementSelectionSet;
 struct FToolMenuContext;
 class FEditorModeTools;
 
+/** Delegate type for pre save world events ( uint32 SaveFlags, UWorld* World ) */
+DECLARE_DYNAMIC_MULTICAST_DELEGATE_TwoParams(FOnLevelEditorPreSaveWorld, int32, SaveFlags, class UWorld*, World);
+/** Delegate type for post save world events ( uint32 SaveFlags, UWorld* World, bool bSuccess ) */
+DECLARE_DYNAMIC_MULTICAST_DELEGATE_ThreeParams(FOnLevelEditorPostSaveWorld, int32, SaveFlags, class UWorld*, World, bool, bSuccess);
+
+/** Delegate type for editor camera movement */
+DECLARE_DYNAMIC_MULTICAST_DELEGATE_FourParams(FOnLevelEditorEditorCameraMoved, const FVector&, Location, const FRotator&, Rotation, ELevelViewportType, ViewportType, int32, ViewIndex);
+
+/** Delegate type for map change events ( Params: uint32 MapChangeFlags (MapChangeEventFlags) ) */
+DECLARE_DYNAMIC_MULTICAST_DELEGATE_OneParam(FOnLevelEditorMapChanged, int32, MapChangeEventFlags);
+/** Delegate type for triggering when a map is opened */
+DECLARE_DYNAMIC_MULTICAST_DELEGATE_TwoParams(FOnLevelEditorMapOpened, const FString&, Filename, bool, bAsTemplate);
+
 /**
 * ULevelEditorSubsystem
 * Subsystem for exposing Level Editor related functionality to scripts
@@ -127,7 +142,8 @@ public:
 	ULevel* GetCurrentLevel();
 
 	/**
-	 * Get the level Editor typed element selection set for the current world
+	 * Get the selection set for the current world, you can use this to track
+	 * and create changes to the level editor's selection
 	 */
 	UFUNCTION(BlueprintCallable, Category = "Editor Scripting | Level Utility")
 	UTypedElementSelectionSet* GetSelectionSet();
@@ -155,6 +171,24 @@ public:
 	virtual FOnActorEditorContextClientChanged& GetOnActorEditorContextClientChanged() override { return ActorEditorContextClientChanged; }
 	//~ End IActorEditorContextClient interface
 
+	/** Expose PreSaveWorld to blueprints */
+	UPROPERTY(BlueprintAssignable, Category = "Editor Scripting | Level Utility")
+	FOnLevelEditorPreSaveWorld OnPreSaveWorld;
+	/** Expose PostSaveWorld to blueprints */
+	UPROPERTY(BlueprintAssignable, Category = "Editor Scripting | Level Utility")
+	FOnLevelEditorPostSaveWorld OnPostSaveWorld;
+
+	/** Expose EditorCameraMoved to blueprints */
+	UPROPERTY(BlueprintAssignable, Category = "Editor Scripting | Level Utility")
+	FOnLevelEditorEditorCameraMoved OnEditorCameraMoved;
+
+	/** Expose MapChanged to blueprints. Note: This executes too early for some editor scripting, consider using OnMapOpened if this doesn't work for you. */
+	UPROPERTY(BlueprintAssignable, Category = "Editor Scripting | Level Utility")
+	FOnLevelEditorMapChanged OnMapChanged;
+	/** Expose MapOpened to blueprints */
+	UPROPERTY(BlueprintAssignable, Category = "Editor Scripting | Level Utility")
+	FOnLevelEditorMapOpened OnMapOpened;
+
 private:
 
 	/** Called when a Level is added to a world or remove from a world */
@@ -165,4 +199,17 @@ private:
 
 	/** Delegate used to notify changes to ActorEditorContextSubsystem */
 	FOnActorEditorContextClientChanged ActorEditorContextClientChanged;
+
+	/** Called Pre SaveWorld */
+	void HandleOnPreSaveWorldWithContext(class UWorld* World, FObjectPreSaveContext ObjectSaveContext);
+	/** Called Post SaveWorld */
+	void HandleOnPostSaveWorldWithContext(class UWorld* World, FObjectPostSaveContext ObjectSaveContext);
+
+	/** Called on Camera Movement*/
+	void HandleOnEditorCameraMoved(const FVector& Location, const FRotator& Rotation, ELevelViewportType ViewportType, int32 ViewIndex);
+
+	/** Called on Map Change */
+	void HandleOnMapChanged(uint32 MapChangeFlags);
+	/** Called on Map Open */
+	void HandleOnMapOpened(const FString& Filename, bool bAsTemplate);
 };
\ No newline at end of file
diff --git a/Engine/Source/Editor/UnrealEd/Private/Subsystems/EditorActorSubsystem.cpp b/Engine/Source/Editor/UnrealEd/Private/Subsystems/EditorActorSubsystem.cpp
--- a/Engine/Source/Editor/UnrealEd/Private/Subsystems/EditorActorSubsystem.cpp
+++ b/Engine/Source/Editor/UnrealEd/Private/Subsystems/EditorActorSubsystem.cpp
@@ -134,6 +134,9 @@ namespace InternalActorUtilitiesSubsystemLibrary
 
 void UEditorActorSubsystem::Initialize(FSubsystemCollectionBase& Collection)
 {
+	FEditorDelegates::OnNewActorsDropped.AddUObject(this, &UEditorActorSubsystem::BroadcastEditNewActorsDropped);
+	FEditorDelegates::OnNewActorsPlaced.AddUObject(this, &UEditorActorSubsystem::BroadcastEditNewActorsPlaced);
+
 	FEditorDelegates::OnEditCutActorsBegin.AddUObject(this, &UEditorActorSubsystem::BroadcastEditCutActorsBegin);
 	FEditorDelegates::OnEditCutActorsEnd.AddUObject(this, &UEditorActorSubsystem::BroadcastEditCutActorsEnd);
 
@@ -152,6 +155,8 @@ void UEditorActorSubsystem::Initialize(FSubsystemCollectionBase& Collection)
 
 void UEditorActorSubsystem::Deinitialize()
 {
+	FEditorDelegates::OnNewActorsDropped.RemoveAll(this);
+	FEditorDelegates::OnNewActorsPlaced.RemoveAll(this);
 	FEditorDelegates::OnEditCutActorsBegin.RemoveAll(this);
 	FEditorDelegates::OnEditCutActorsEnd.RemoveAll(this);
 	FEditorDelegates::OnEditCopyActorsBegin.RemoveAll(this);
@@ -165,6 +170,18 @@ void UEditorActorSubsystem::Deinitialize()
 
 }
 
+/** To fire before an Actor is Dropped */
+void UEditorActorSubsystem::BroadcastEditNewActorsDropped(const TArray<UObject*>& DroppedObjects, const TArray<AActor*>& DroppedActors)
+{
+	OnNewActorsDropped.Broadcast(DroppedObjects, DroppedActors);
+}
+
+/** To fire before an Actor is Placed */
+void UEditorActorSubsystem::BroadcastEditNewActorsPlaced(UObject* ObjToUse, const TArray<AActor*>& PlacedActors)
+{
+	OnNewActorsPlaced.Broadcast(ObjToUse, PlacedActors);
+}
+
 /** To fire before an Actor is Cut */
 void UEditorActorSubsystem::BroadcastEditCutActorsBegin()
 {
diff --git a/Engine/Source/Editor/UnrealEd/Private/UnrealEdSrv.cpp b/Engine/Source/Editor/UnrealEd/Private/UnrealEdSrv.cpp
--- a/Engine/Source/Editor/UnrealEd/Private/UnrealEdSrv.cpp
+++ b/Engine/Source/Editor/UnrealEd/Private/UnrealEdSrv.cpp
@@ -2583,7 +2583,9 @@ bool UUnrealEdEngine::Exec_Actor( UWorld* InWorld, const TCHAR* Str, FOutputDevi
 					}
 					else
 					{
+						FEditorDelegates::OnDeleteActorsBegin.Broadcast();
 						CommonActions->DeleteSelectedElements(SelectionSet, InWorld, FTypedElementDeletionOptions());
+						FEditorDelegates::OnDeleteActorsEnd.Broadcast();
 					}
 				}
 			}
diff --git a/Engine/Source/Editor/UnrealEd/Public/Subsystems/EditorActorSubsystem.h b/Engine/Source/Editor/UnrealEd/Public/Subsystems/EditorActorSubsystem.h
--- a/Engine/Source/Editor/UnrealEd/Public/Subsystems/EditorActorSubsystem.h
+++ b/Engine/Source/Editor/UnrealEd/Public/Subsystems/EditorActorSubsystem.h
@@ -9,6 +9,12 @@
 
 #include "EditorActorSubsystem.generated.h"
 
+class AActor;
+
+/** delegate type for triggering when new actors are dropped on to the viewport via drag and drop */
+DECLARE_DYNAMIC_MULTICAST_DELEGATE_TwoParams(FOnEditNewActorsDropped, const TArray<UObject*>&, DroppedObjects, const TArray<AActor*>&, DroppedActors);
+/** delegate type for triggering when new actors are placed on to the viewport. Triggers before NewActorsDropped if placement is caused by a drop action */
+DECLARE_DYNAMIC_MULTICAST_DELEGATE_TwoParams(FOnEditNewActorsPlaced, UObject*, ObjToUse, const TArray<AActor*>&, PlacedActors);
 /** delegate type for before edit cut actors is handled */
 DECLARE_DYNAMIC_MULTICAST_DELEGATE(FOnEditCutActorsBegin);
 /** delegate type for after edit cut actors is handled */
@@ -16,7 +22,7 @@ DECLARE_DYNAMIC_MULTICAST_DELEGATE(FOnEditCutActorsEnd);
 /** delegate type for before edit copy actors is handled */
 DECLARE_DYNAMIC_MULTICAST_DELEGATE(FOnEditCopyActorsBegin);
 /** delegate type for after edit copy actors is handled */
-DECLARE_DYNAMIC_MULTICAST_DELEGATE(FOnEditCopyActorsEnd);
+DECLARE_DYNAMIC_MULTICAST_DELEGATE(FOnEditCopyActorsEnd); 
 /** delegate type for before edit paste actors is handled */
 DECLARE_DYNAMIC_MULTICAST_DELEGATE(FOnEditPasteActorsBegin);
 /** delegate type for after edit paste actors is handled */
@@ -44,6 +50,12 @@ public:
 	virtual void Initialize(FSubsystemCollectionBase& Collection) override;
 	virtual void Deinitialize() override;
 
+	UPROPERTY(BlueprintAssignable, Category = "Editor Scripting | Level Utility")
+	FOnEditNewActorsDropped OnNewActorsDropped;
+
+	UPROPERTY(BlueprintAssignable, Category = "Editor Scripting | Level Utility")
+	FOnEditNewActorsPlaced OnNewActorsPlaced;
+
 	UPROPERTY(BlueprintAssignable, Category = "Editor Scripting | Level Utility")
 	FOnEditCutActorsBegin OnEditCutActorsBegin;
 
@@ -240,6 +252,12 @@ public:
 
 private:
 
+	/** To fire before an Actor is Dropped */
+	void BroadcastEditNewActorsDropped(const TArray<UObject*>& DroppedObjects, const TArray<AActor*>& DroppedActors);
+
+	/** To fire before an Actor is Placed */
+	void BroadcastEditNewActorsPlaced(UObject* ObjToUse, const TArray<AActor*>& PlacedActors);
+
 	/** To fire before an Actor is Cut */
 	void BroadcastEditCutActorsBegin();
 
