Sequencer: Make the Window->Cinematics->Sequencer menu item always clickable. If there is an existing Sequencer window, draw attention to it when the menu item is clicked.

diff --git a/Engine/Source/Editor/LevelEditor/Private/SLevelEditor.cpp b/Engine/Source/Editor/LevelEditor/Private/SLevelEditor.cpp
--- a/Engine/Source/Editor/LevelEditor/Private/SLevelEditor.cpp
+++ b/Engine/Source/Editor/LevelEditor/Private/SLevelEditor.cpp
@@ -1034,6 +1034,19 @@ TSharedRef<SDockTab> SLevelEditor::SpawnLevelEditorTab( const FSpawnTabArgs& Arg
 	}
 	else if( TabIdentifier == TEXT("Sequencer") )
 	{
+		// If the Sequencer tab already exists, draw attention to it and return
+		TSharedPtr<FTabManager> LevelEditorTabManager = GetTabManager();
+		if (LevelEditorTabManager.IsValid())
+		{
+			TSharedPtr<SDockTab> Tab = LevelEditorTabManager->FindExistingLiveTab(LevelEditorTabIds::Sequencer);
+			if (Tab.IsValid())
+			{
+				LevelEditorTabManager->DrawAttention(Tab.ToSharedRef());
+				return Tab.ToSharedRef();
+			}
+		}
+		
+		// Otherwise, open a null sequencer widget
 		if (FSlateStyleRegistry::FindSlateStyle("LevelSequenceEditorStyle"))
 		{
 			// @todo sequencer: remove when world-centric mode is added
@@ -1100,11 +1113,6 @@ TSharedRef<SDockTab> SLevelEditor::SpawnLevelEditorTab( const FSpawnTabArgs& Arg
 
 bool SLevelEditor::CanSpawnLevelEditorTab(const FSpawnTabArgs& Args, FName TabIdentifier)
 {
-	if (TabIdentifier == LevelEditorTabIds::Sequencer && !IsModeActive("EM_SequencerMode"))
-	{
-		return false;
-	}
-
 	// HLOD Outliner not yet supported with World Partition
 	if (TabIdentifier == LevelEditorTabIds::LevelEditorHierarchicalLODOutliner && ensure(World) && World->IsPartitionedWorld())
 	{
