Disable Preview button now named after the Editor Shader Platform
E.G.
[EditorShaderPlatformName]
(Disable Preview)

diff --git a/Engine/Source/Editor/LevelEditor/Private/LevelEditor.cpp b/Engine/Source/Editor/LevelEditor/Private/LevelEditor.cpp
--- a/Engine/Source/Editor/LevelEditor/Private/LevelEditor.cpp
+++ b/Engine/Source/Editor/LevelEditor/Private/LevelEditor.cpp
@@ -1795,7 +1795,7 @@ void FLevelEditorModule::BindGlobalLevelEditorCommands()
 
 	const TArray<FPreviewPlatformMenuItem>& MenuItems = FDataDrivenPlatformInfoRegistry::GetAllPreviewPlatformMenuItems();
 	// We need one extra slot for the Disable Preview option
-	check(MenuItems.Num() + 1 == Commands.PreviewPlatformOverrides.Num());
+	check(MenuItems.Num() == Commands.PreviewPlatformOverrides.Num());
 
 	for (int32 Index=0; Index < MenuItems.Num(); Index++)
 	{
@@ -1806,11 +1806,9 @@ void FLevelEditorModule::BindGlobalLevelEditorCommands()
 		{
 			ERHIFeatureLevel::Type FeatureLevel = GetMaxSupportedFeatureLevel(ShaderPlatform);
 
-			const bool IsDefaultActive = FDataDrivenShaderPlatformInfo::GetPreviewShaderPlatformParent(ShaderPlatform) == GMaxRHIShaderPlatform;
-			const bool AllowPreview = !IsDefaultActive;
+			bool bIsDefaultShaderPlatform = FDataDrivenShaderPlatformInfo::GetPreviewShaderPlatformParent(ShaderPlatform) == GMaxRHIShaderPlatform;
 
-
-			FPreviewPlatformInfo PreviewFeatureLevelInfo(FeatureLevel, (EShaderPlatform)ShaderPlatform, IsDefaultActive ? NAME_None : Item.PlatformName, IsDefaultActive ? NAME_None : Item.ShaderFormat, IsDefaultActive ? NAME_None : Item.DeviceProfileName, AllowPreview, Item.PreviewShaderPlatformName);
+			FPreviewPlatformInfo PreviewFeatureLevelInfo(bIsDefaultShaderPlatform ? GMaxRHIFeatureLevel : FeatureLevel, bIsDefaultShaderPlatform ? GMaxRHIShaderPlatform : (EShaderPlatform)ShaderPlatform, bIsDefaultShaderPlatform ? NAME_None : Item.PlatformName, bIsDefaultShaderPlatform ? NAME_None : Item.ShaderFormat, bIsDefaultShaderPlatform ? NAME_None : Item.DeviceProfileName, true, bIsDefaultShaderPlatform ? NAME_None : Item.PreviewShaderPlatformName);
 
 			ActionList.MapAction(
 				Commands.PreviewPlatformOverrides[Index],
@@ -1820,16 +1818,6 @@ void FLevelEditorModule::BindGlobalLevelEditorCommands()
 		}
 	}
 
-	// Add the Disable Preview Menu Item
-	{
-		FPreviewPlatformInfo PreviewFeatureLevelInfo(GMaxRHIFeatureLevel, GMaxRHIShaderPlatform, NAME_None, NAME_None, NAME_None, true, NAME_None);
-		ActionList.MapAction(
-			Commands.PreviewPlatformOverrides[MenuItems.Num()],
-			FExecuteAction::CreateStatic(&FLevelEditorActionCallbacks::SetPreviewPlatform, PreviewFeatureLevelInfo),
-			FCanExecuteAction::CreateStatic(&FLevelEditorActionCallbacks::CanExecutePreviewPlatform, PreviewFeatureLevelInfo),
-			FIsActionChecked::CreateStatic(&FLevelEditorActionCallbacks::IsPreviewPlatformChecked, PreviewFeatureLevelInfo));
-	}
-
 	ActionList.MapAction(
 		Commands.OpenMergeActor,
 		FExecuteAction::CreateStatic(&FLevelEditorActionCallbacks::OpenMergeActor_Clicked)
diff --git a/Engine/Source/Editor/LevelEditor/Private/LevelEditorActions.cpp b/Engine/Source/Editor/LevelEditor/Private/LevelEditorActions.cpp
--- a/Engine/Source/Editor/LevelEditor/Private/LevelEditorActions.cpp
+++ b/Engine/Source/Editor/LevelEditor/Private/LevelEditorActions.cpp
@@ -3756,34 +3756,28 @@ void FLevelEditorCommands::RegisterCommands()
 	// Add preview platforms
 	for (const FPreviewPlatformMenuItem& Item : FDataDrivenPlatformInfoRegistry::GetAllPreviewPlatformMenuItems())
 	{
-		FText FriendlyName;
+		FTextBuilder FriendlyNameBuilder;
 		if (!IsRunningCommandlet() && !GUsingNullRHI)
 		{
 			EShaderPlatform ShaderPlatform = FDataDrivenShaderPlatformInfo::GetShaderPlatformFromName(Item.PreviewShaderPlatformName);
-			FriendlyName = FDataDrivenShaderPlatformInfo::GetFriendlyName(ShaderPlatform);
+			FriendlyNameBuilder.AppendLine(FDataDrivenShaderPlatformInfo::GetFriendlyName(ShaderPlatform));
+			if (FDataDrivenShaderPlatformInfo::GetShaderPlatformFromName(Item.ShaderPlatformToPreview) == GMaxRHIShaderPlatform)
+			{
+				FriendlyNameBuilder.AppendLine(NSLOCTEXT("PreviewPlatform", "PreviewMenuText_DisablePreview", "(Disable Preview)"));
+			}
 		}
 
 		PreviewPlatformOverrides.Add(
 			FUICommandInfoDecl(
 				this->AsShared(),
 				FName(*FString::Printf(TEXT("PreviewPlatformOverrides_%s_%s"), *Item.PlatformName.ToString(), *Item.ShaderFormat.ToString())),
-				FriendlyName,
+				FriendlyNameBuilder.ToText(),
 				Item.MenuTooltip)
 			.UserInterfaceType(EUserInterfaceActionType::Check)
 			.DefaultChord(FInputChord())
 		);
 	}
 
-	PreviewPlatformOverrides.Add(
-		FUICommandInfoDecl(
-			this->AsShared(),
-			FName(TEXT("Disable Preview Command")),
-			NSLOCTEXT("PreviewPlatform", "PreviewMenuText_DisablePreview", "Disable Preview"),
-			NSLOCTEXT("PreviewPlatform", "PreviewMenuTooltip_DisablePreview", "Disable the current Preview Shader Platform"))
-		.UserInterfaceType(EUserInterfaceActionType::Check)
-		.DefaultChord(FInputChord())
-	);
-
 	UI_COMMAND(OpenMergeActor, "Merge Actors", "Opens the Merge Actor panel", EUserInterfaceActionType::Button, FInputChord());
 }
 
