Virtual Scouting: Fix shift-v hotkey exit for modes which do not derive from UVREditorMode. Followup to Release-5.1 hotfix change merged in CL 23300617.

#rnx

diff --git a/Engine/Source/Editor/LevelEditor/Private/LevelEditorActions.cpp b/Engine/Source/Editor/LevelEditor/Private/LevelEditorActions.cpp
--- a/Engine/Source/Editor/LevelEditor/Private/LevelEditorActions.cpp
+++ b/Engine/Source/Editor/LevelEditor/Private/LevelEditorActions.cpp
@@ -2379,7 +2379,7 @@ void FLevelEditorActionCallbacks::ImportContent()
 void FLevelEditorActionCallbacks::ToggleVR()
 {
 	IVREditorModule& VREditorModule = IVREditorModule::Get();
-	VREditorModule.EnableVREditor( VREditorModule.GetVRMode() == nullptr );
+	VREditorModule.EnableVREditor( VREditorModule.GetVRModeBase() == nullptr );
 }
 
 
