The File > Recent Levels submenu now contains a Clear Recent Levels command.

diff --git a/Engine/Source/Editor/LevelEditor/Private/LevelEditor.cpp b/Engine/Source/Editor/LevelEditor/Private/LevelEditor.cpp
--- a/Engine/Source/Editor/LevelEditor/Private/LevelEditor.cpp
+++ b/Engine/Source/Editor/LevelEditor/Private/LevelEditor.cpp
@@ -677,6 +677,8 @@ void FLevelEditorModule::BindGlobalLevelEditorCommands()
 		ActionList.MapAction( Commands.OpenRecentFileCommands[ CurRecentIndex ], FExecuteAction::CreateStatic( &FLevelEditorActionCallbacks::OpenRecentFile, CurRecentIndex ), DefaultExecuteAction );
 	}
 
+	ActionList.MapAction( Commands.ClearRecentFiles, FExecuteAction::CreateStatic(&FLevelEditorActionCallbacks::ClearRecentFiles) );
+
 	for (int32 CurFavoriteIndex = 0; CurFavoriteIndex < FLevelEditorCommands::MaxFavoriteFiles; ++CurFavoriteIndex)
 	{
 		ActionList.MapAction(Commands.OpenFavoriteFileCommands[CurFavoriteIndex], FExecuteAction::CreateStatic(&FLevelEditorActionCallbacks::OpenFavoriteFile, CurFavoriteIndex), DefaultExecuteAction);
diff --git a/Engine/Source/Editor/LevelEditor/Private/LevelEditorActions.cpp b/Engine/Source/Editor/LevelEditor/Private/LevelEditorActions.cpp
--- a/Engine/Source/Editor/LevelEditor/Private/LevelEditorActions.cpp
+++ b/Engine/Source/Editor/LevelEditor/Private/LevelEditorActions.cpp
@@ -302,28 +302,35 @@ void FLevelEditorActionCallbacks::DeltaTransform()
 void FLevelEditorActionCallbacks::OpenRecentFile( int32 RecentFileIndex )
 {
 	IMainFrameModule& MainFrameModule = FModuleManager::LoadModuleChecked<IMainFrameModule>( "MainFrame" );
-	FMainMRUFavoritesList* RecentsAndFavorites = MainFrameModule.GetMRUFavoritesList();
-
-	FString NewPackageName;
-	if( RecentsAndFavorites->VerifyMRUFile( RecentFileIndex, NewPackageName ) )
+	
+	if (FMainMRUFavoritesList* RecentsAndFavorites = MainFrameModule.GetMRUFavoritesList())
 	{
-		// Prompt the user to save any outstanding changes.
-		if( FEditorFileUtils::SaveDirtyPackages(true, true, false) )
+		FString NewPackageName;
+		if( RecentsAndFavorites->VerifyMRUFile( RecentFileIndex, NewPackageName ) )
 		{
-			FString NewFilename;
-			if (FPackageName::TryConvertLongPackageNameToFilename(NewPackageName, NewFilename, FPackageName::GetMapPackageExtension()))
+			// Prompt the user to save any outstanding changes.
+			if( FEditorFileUtils::SaveDirtyPackages(true, true, false) )
 			{
-				// Load the requested level.
-				FEditorFileUtils::LoadMap(NewFilename);
+				FString NewFilename;
+				if (FPackageName::TryConvertLongPackageNameToFilename(NewPackageName, NewFilename, FPackageName::GetMapPackageExtension()))
+				{
+					// Load the requested level.
+					FEditorFileUtils::LoadMap(NewFilename);
+				}
 			}
 		}
-		else
-		{
-			// something went wrong or the user pressed cancel.  Return to the editor so the user doesn't lose their changes		
-		}
 	}
 }
 
+void FLevelEditorActionCallbacks::ClearRecentFiles()
+{
+	IMainFrameModule& MainFrameModule = FModuleManager::LoadModuleChecked<IMainFrameModule>( "MainFrame" );
+	
+	if (FMainMRUFavoritesList* RecentsAndFavorites = MainFrameModule.GetMRUFavoritesList())
+	{
+		RecentsAndFavorites->ClearMRUItems();
+	}
+}
 
 void FLevelEditorActionCallbacks::OpenFavoriteFile( int32 FavoriteFileIndex )
 {
@@ -3488,6 +3495,7 @@ void FLevelEditorCommands::RegisterCommands()
 		OpenFavoriteFileCommands.Add(OpenFavoriteFile);
 	}
 
+	UI_COMMAND( ClearRecentFiles, "Clear Recent Levels", "Clear the list of recently opened levels", EUserInterfaceActionType::Button, FInputChord() );
 	UI_COMMAND( ImportScene, "Import Into Level...", "Imports a scene from a FBX or OBJ format into the current level", EUserInterfaceActionType::Button, FInputChord());
 	UI_COMMAND( ExportAll, "Export All...", "Exports the entire level to a file on disk (multiple formats are supported.)", EUserInterfaceActionType::Button, FInputChord() );
 	UI_COMMAND( ExportSelected, "Export Selected...", "Exports currently-selected objects to a file on disk (multiple formats are supported.)", EUserInterfaceActionType::Button, FInputChord() );
diff --git a/Engine/Source/Editor/LevelEditor/Private/LevelEditorMenu.cpp b/Engine/Source/Editor/LevelEditor/Private/LevelEditorMenu.cpp
--- a/Engine/Source/Editor/LevelEditor/Private/LevelEditorMenu.cpp
+++ b/Engine/Source/Editor/LevelEditor/Private/LevelEditorMenu.cpp
@@ -151,28 +151,34 @@ void FLevelEditorMenu::RegisterLevelEditorMenus()
 			{
 				struct FRecentLevelMenu
 				{
-					static void MakeRecentLevelMenu( UToolMenu* InMenu )
+					static void MakeRecentLevelMenu(UToolMenu* InMenu)
 					{
-						const FMainMRUFavoritesList& MRUFavorites = *FModuleManager::LoadModuleChecked<IMainFrameModule>( "MainFrame" ).GetMRUFavoritesList();
+						const FMainMRUFavoritesList& MRUFavorites = *FModuleManager::LoadModuleChecked<IMainFrameModule>("MainFrame").GetMRUFavoritesList();
 						const int32 NumRecents = MRUFavorites.GetNumItems();
 
-						const int32 AllowedRecents = FMath::Min( NumRecents, FLevelEditorCommands::Get().OpenRecentFileCommands.Num() );
-						for ( int32 CurRecentIndex = 0; CurRecentIndex < AllowedRecents; ++CurRecentIndex )
+						FToolMenuSection& Section = InMenu->FindOrAddSection("Recent");
+
+						const int32 AllowedRecents = FMath::Min(NumRecents, FLevelEditorCommands::Get().OpenRecentFileCommands.Num());
+						for (int32 CurRecentIndex = 0; CurRecentIndex < AllowedRecents; ++CurRecentIndex)
 						{
-							TSharedPtr< FUICommandInfo > OpenRecentFile = FLevelEditorCommands::Get().OpenRecentFileCommands[ CurRecentIndex ];
+							TSharedPtr< FUICommandInfo > OpenRecentFile = FLevelEditorCommands::Get().OpenRecentFileCommands[CurRecentIndex];
 
 							if (!MRUFavorites.MRUItemPassesCurrentFilter(CurRecentIndex))
 							{
 								continue;
 							}
 
-							const FString CurRecent = MRUFavorites.GetMRUItem( CurRecentIndex );
+							const FString& CurRecent = MRUFavorites.GetMRUItem(CurRecentIndex);
 
-							const FText ToolTip = FText::Format( LOCTEXT( "RecentLevelToolTip", "Opens recent level: {0}" ), FText::FromString( CurRecent ) );
-							const FText Label = FText::FromString( FPaths::GetBaseFilename( CurRecent ) );
+							const FText ToolTip = FText::Format(LOCTEXT("RecentLevelToolTip", "Opens recent level: {0}"), FText::FromString(CurRecent));
+							const FText Label = FText::FromString(FPaths::GetBaseFilename(CurRecent));
 
-							InMenu->FindOrAddSection("Recent").AddMenuEntry( OpenRecentFile, Label, ToolTip ).Name = NAME_None;
+							Section.AddMenuEntry(OpenRecentFile, Label, ToolTip).Name = NAME_None;
 						}
+
+						Section.AddSeparator("AfterRecentLevels");
+
+						Section.AddMenuEntry("ClearRecentLevels", FLevelEditorCommands::Get().ClearRecentFiles);
 					}
 				};
 
diff --git a/Engine/Source/Editor/LevelEditor/Public/LevelEditorActions.h b/Engine/Source/Editor/LevelEditor/Public/LevelEditorActions.h
--- a/Engine/Source/Editor/LevelEditor/Public/LevelEditorActions.h
+++ b/Engine/Source/Editor/LevelEditor/Public/LevelEditorActions.h
@@ -53,6 +53,8 @@ public:
 	TArray< TSharedPtr< FUICommandInfo > > OpenRecentFileCommands;
 	static const int32 MaxFavoriteFiles = 20;
 	TArray< TSharedPtr< FUICommandInfo > > OpenFavoriteFileCommands;
+	
+	TSharedPtr< FUICommandInfo > ClearRecentFiles;
 
 	TSharedPtr< FUICommandInfo > ToggleFavorite;
 
@@ -670,6 +672,9 @@ public:
 	 */
 	static void OpenRecentFile( int32 RecentFileIndex );
 
+	/** Clear the list of recent files. */
+	static void ClearRecentFiles();
+
 	/**
 	 * Opens a favorite file
 	 *
diff --git a/Engine/Source/Editor/UnrealEd/Private/MRUList.cpp b/Engine/Source/Editor/UnrealEd/Private/MRUList.cpp
--- a/Engine/Source/Editor/UnrealEd/Private/MRUList.cpp
+++ b/Engine/Source/Editor/UnrealEd/Private/MRUList.cpp
@@ -38,6 +38,12 @@ void FMRUList::Cull()
 	}
 }
 
+void FMRUList::ClearMRUItems()
+{
+	Items.Empty();
+
+	WriteToINI();
+}
 
 void FMRUList::ReadFromINI()
 {
diff --git a/Engine/Source/Editor/UnrealEd/Public/MRUList.h b/Engine/Source/Editor/UnrealEd/Public/MRUList.h
--- a/Engine/Source/Editor/UnrealEd/Public/MRUList.h
+++ b/Engine/Source/Editor/UnrealEd/Public/MRUList.h
@@ -63,6 +63,11 @@ public:
 	 */
 	void RemoveMRUItem(int32 ItemIndex);
 
+	/**
+	 * Clears the list of MRU items. 
+	 */
+	void ClearMRUItems();
+
 	/**
 	 * Checks to make sure the file specified still exists.  If it does, it is
 	 * moved to the top of the MRU list and we return true.  If not, we remove it
