Outliner: Fix an issue causing duplicate filters when you load the layout

diff --git a/Engine/Source/Editor/LevelEditor/Private/LevelEditorOutlinerSettings.cpp b/Engine/Source/Editor/LevelEditor/Private/LevelEditorOutlinerSettings.cpp
--- a/Engine/Source/Editor/LevelEditor/Private/LevelEditorOutlinerSettings.cpp
+++ b/Engine/Source/Editor/LevelEditor/Private/LevelEditorOutlinerSettings.cpp
@@ -284,6 +284,10 @@ void FLevelEditorOutlinerSettings::CreateDefaultFilters()
 	{
 		FindOrAddClassFilter(AAmbientSound::StaticClass(), *AudioFilterCategory);
 	}
+
+
+	// Destroy and recreate our built in custom filters
+	BuiltInCustomFilters.Empty();
 	
 	CreateSCCFilters();
 }
@@ -334,7 +338,7 @@ void FLevelEditorOutlinerSettings::CreateSCCFilters()
 	FGenericFilter<const ISceneOutlinerTreeItem&>::FOnItemFiltered UncontrolledFilterDelegate = FGenericFilter<const ISceneOutlinerTreeItem&>::FOnItemFiltered::CreateSP(this, &FLevelEditorOutlinerSettings::DoesActorPassUncontrolledFilter);
 	TSharedPtr<FGenericFilter<const ISceneOutlinerTreeItem&>> UncontrolledFilter = MakeShared<FGenericFilter<const ISceneOutlinerTreeItem&>>(SCCFiltersCategory, UncontrolledAssetsFilterName, LOCTEXT("UncontrolledFilterName", "Uncontrolled"), UncontrolledFilterDelegate);
 	UncontrolledFilter->SetToolTipText(LOCTEXT("UncontrolledFilterTooltip", "Only show items that are uncontrolled (locally modified outside of revision control)"));
-	CustomFilters.Add(UncontrolledFilter.ToSharedRef());
+	BuiltInCustomFilters.Add(UncontrolledFilter.ToSharedRef());
 	
 	// File Management Category
 	TSharedPtr<FFilterCategory> FileManagementFiltersCategory = MakeShared<FFilterCategory>(LOCTEXT("FileManagementFiltersCategory", "File Management"), FText::GetEmpty());
@@ -348,7 +352,7 @@ void FLevelEditorOutlinerSettings::CreateSCCFilters()
 	FGenericFilter<const ISceneOutlinerTreeItem&>::FOnItemFiltered UnsavedActorsFilterDelegate = FGenericFilter<const ISceneOutlinerTreeItem&>::FOnItemFiltered::CreateSP(this, & FLevelEditorOutlinerSettings::DoesActorPassUnsavedFilter);
 	TSharedPtr<FGenericFilter<const ISceneOutlinerTreeItem&>> UnsavedAssetsFilter = MakeShared<FGenericFilter<const ISceneOutlinerTreeItem&>>(FileManagementFiltersCategory, UnsavedAssetsFilterName, LOCTEXT("UnsavedFilterName", "Unsaved"), UnsavedActorsFilterDelegate);
 	UnsavedAssetsFilter->SetToolTipText(LOCTEXT("UnsavedAssetsFilterTooltip", "Only show items that are unsaved"));
-	CustomFilters.Add(UnsavedAssetsFilter.ToSharedRef());
+	BuiltInCustomFilters.Add(UnsavedAssetsFilter.ToSharedRef());
 }
 
 void FLevelEditorOutlinerSettings::OnUncontrolledChangelistModuleChanged()
@@ -370,6 +374,7 @@ void FLevelEditorOutlinerSettings::GetOutlinerFilters(FSceneOutlinerFilterBarOpt
 	
 	OutFilterBarOptions.CustomClassFilters.Append(CustomClassFilters);
 	OutFilterBarOptions.CustomFilters.Append(CustomFilters);
+	OutFilterBarOptions.CustomFilters.Append(BuiltInCustomFilters);
 
 	for(FOutlinerFilterFactory& CreateFilterDelegate : CustomFilterDelegates)
 	{
diff --git a/Engine/Source/Editor/LevelEditor/Public/LevelEditorOutlinerSettings.h b/Engine/Source/Editor/LevelEditor/Public/LevelEditorOutlinerSettings.h
--- a/Engine/Source/Editor/LevelEditor/Public/LevelEditorOutlinerSettings.h
+++ b/Engine/Source/Editor/LevelEditor/Public/LevelEditorOutlinerSettings.h
@@ -104,6 +104,9 @@ private:
 	 */
 	TArray<TSharedRef<FCustomClassFilterData>> CustomClassFilters;
 
+	/** The built in custom filters created by this instance */
+	TArray<TSharedRef<FFilterBase<const ISceneOutlinerTreeItem&>>> BuiltInCustomFilters;
+
 	/* A map of the categories the Outliner filter bar will have */
 	TMap<FName, TSharedPtr<FFilterCategory>> FilterBarCategories;
 
