Fix crash when in level editor when hairstrands plugin is not enabled.

diff --git a/Engine/Source/Editor/LevelEditor/Private/SLevelViewport.cpp b/Engine/Source/Editor/LevelEditor/Private/SLevelViewport.cpp
--- a/Engine/Source/Editor/LevelEditor/Private/SLevelViewport.cpp
+++ b/Engine/Source/Editor/LevelEditor/Private/SLevelViewport.cpp
@@ -1555,7 +1555,10 @@ void SLevelViewport::BindViewCommands( FUICommandList& OutCommandList )
 	{
 		FStrataVisualizationMenuCommands::Get().BindCommands(OutCommandList, Client);
 	}
-	FGroomVisualizationMenuCommands::Get().BindCommands(OutCommandList, Client);
+	if (IsGroomEnabled())
+	{
+		FGroomVisualizationMenuCommands::Get().BindCommands(OutCommandList, Client);
+	}
 	FVirtualShadowMapVisualizationMenuCommands::Get().BindCommands(OutCommandList, Client);
 }
 
