Renamed "Build Minimap" to "Build World Partition Editor Mininmap" to remove confusion with in-game or HUD minimap

diff --git a/Engine/Source/Editor/LevelEditor/Private/LevelEditorActions.cpp b/Engine/Source/Editor/LevelEditor/Private/LevelEditorActions.cpp
--- a/Engine/Source/Editor/LevelEditor/Private/LevelEditorActions.cpp
+++ b/Engine/Source/Editor/LevelEditor/Private/LevelEditorActions.cpp
@@ -3511,7 +3511,7 @@ void FLevelEditorCommands::RegisterCommands()
 	UI_COMMAND( BuildGeometryOnly_OnlyCurrentLevel, "Build Geometry (Current Level)", "Builds geometry, only for the current level", EUserInterfaceActionType::Button, FInputChord() );
 	UI_COMMAND( BuildPathsOnly, "Build Paths", "Only builds paths (all levels.)", EUserInterfaceActionType::Button, FInputChord() );
 	UI_COMMAND( BuildHLODs, "Build HLODs", "Builds all HLODs for the current world", EUserInterfaceActionType::Button, FInputChord() );
-	UI_COMMAND( BuildMinimap, "Build Minimap", "Builds the minimap for the current world", EUserInterfaceActionType::Button, FInputChord() );
+	UI_COMMAND( BuildMinimap, "Build World Partition Editor Minimap", "Builds the minimap for the current world", EUserInterfaceActionType::Button, FInputChord() );
 	UI_COMMAND( BuildLandscapeSplineMeshes, "Build Landscape Spline Meshes", "Builds landscape spline meshes for the current world", EUserInterfaceActionType::Button, FInputChord());
 	UI_COMMAND( BuildTextureStreamingOnly, "Build Texture Streaming", "Build texture streaming data", EUserInterfaceActionType::Button, FInputChord() );
 	UI_COMMAND( BuildVirtualTextureOnly, "Build Virtual Textures", "Build runtime virtual texture low mips streaming data", EUserInterfaceActionType::Button, FInputChord());
