Resolve crash caused by UWorld being GC while in PIE

diff --git a/Engine/Source/Editor/LevelEditor/Private/SLevelViewport.cpp b/Engine/Source/Editor/LevelEditor/Private/SLevelViewport.cpp
--- a/Engine/Source/Editor/LevelEditor/Private/SLevelViewport.cpp
+++ b/Engine/Source/Editor/LevelEditor/Private/SLevelViewport.cpp
@@ -4066,10 +4066,10 @@ void SLevelViewport::StartPlayInEditorSession(UGameViewportClient* PlayClient, c
 
 	// register for preview feature level change
 	UEditorEngine* Editor = CastChecked<UEditorEngine>(GEngine);
-	UWorld* PIEWorld = PlayClient->GetWorld();
-	PIEPreviewFeatureLevelChangedHandle = Editor->OnPreviewFeatureLevelChanged().AddLambda([PIEWorld](ERHIFeatureLevel::Type NewFeatureLevel)
+	
+	PIEPreviewFeatureLevelChangedHandle = Editor->OnPreviewFeatureLevelChanged().AddLambda([PlayClient](ERHIFeatureLevel::Type NewFeatureLevel)
 		{
-			PIEWorld->GetWorld()->ChangeFeatureLevel(NewFeatureLevel);
+			PlayClient->GetWorld()->ChangeFeatureLevel(NewFeatureLevel);
 		});
 
 }
