Bring back the Preview Shader Platform sub-menu and toggle button

diff --git a/Engine/Source/Editor/LevelEditor/Private/LevelEditorToolBar.cpp b/Engine/Source/Editor/LevelEditor/Private/LevelEditorToolBar.cpp
--- a/Engine/Source/Editor/LevelEditor/Private/LevelEditorToolBar.cpp
+++ b/Engine/Source/Editor/LevelEditor/Private/LevelEditorToolBar.cpp
@@ -64,6 +64,86 @@
 #include "PlatformInfo.h"
 #include "DataDrivenShaderPlatformInfo.h"
 
+namespace PreviewModeFunctionality
+{
+	FText GetPreviewModeText()
+	{
+		const FPreviewPlatformMenuItem* Item = FDataDrivenPlatformInfoRegistry::GetAllPreviewPlatformMenuItems().FindByPredicate([](const FPreviewPlatformMenuItem& TestItem)
+			{
+				return GEditor->PreviewPlatform.PreviewPlatformName == TestItem.PlatformName && GEditor->PreviewPlatform.PreviewShaderFormatName == TestItem.ShaderFormat && GEditor->PreviewPlatform.PreviewShaderPlatformName == TestItem.PreviewShaderPlatformName;
+			});
+		return Item ? Item->IconText : FText();
+	}
+
+	FText GetPreviewModeTooltip()
+	{
+#define LOCTEXT_NAMESPACE "LevelEditorToolBar"
+		EShaderPlatform PreviewShaderPlatform = GEditor->PreviewPlatform.PreviewShaderPlatformName != NAME_None ?
+			FDataDrivenShaderPlatformInfo::GetShaderPlatformFromName(GEditor->PreviewPlatform.PreviewShaderPlatformName) :
+			GetFeatureLevelShaderPlatform(GEditor->PreviewPlatform.PreviewFeatureLevel);
+
+		EShaderPlatform MaxRHIFeatureLevelPlatform = GetFeatureLevelShaderPlatform(GMaxRHIFeatureLevel);
+
+		{
+			const FText& RenderingAsPlatformName = FDataDrivenShaderPlatformInfo::GetFriendlyName(GEditor->PreviewPlatform.bPreviewFeatureLevelActive ? PreviewShaderPlatform : MaxRHIFeatureLevelPlatform);
+			const FText& SwitchToPlatformName = FDataDrivenShaderPlatformInfo::GetFriendlyName(GEditor->PreviewPlatform.bPreviewFeatureLevelActive ? MaxRHIFeatureLevelPlatform : PreviewShaderPlatform);
+			if (PreviewShaderPlatform == MaxRHIFeatureLevelPlatform)
+			{
+				return FText::Format(LOCTEXT("PreviewModeViewingAs", "Viewing {0}."), RenderingAsPlatformName);
+			}
+			else if (GWorld->FeatureLevel == GMaxRHIFeatureLevel)
+			{
+				return FText::Format(LOCTEXT("PreviewModeViewingAsSwitchTo", "Viewing {0}. Click to preview {1}."), RenderingAsPlatformName, SwitchToPlatformName);
+			}
+			else
+			{
+				return FText::Format(LOCTEXT("PreviewModePreviewingAsSwitchTo", "Previewing {0}. Click to view {1}."), RenderingAsPlatformName, SwitchToPlatformName);
+			}
+		}
+#undef LOCTEXT_NAMESPACE
+	}
+
+	FSlateIcon  GetPreviewModeIcon()
+	{
+		const FPreviewPlatformMenuItem* Item = FDataDrivenPlatformInfoRegistry::GetAllPreviewPlatformMenuItems().FindByPredicate([](const FPreviewPlatformMenuItem& TestItem)
+			{
+				return GEditor->PreviewPlatform.PreviewPlatformName == TestItem.PlatformName && GEditor->PreviewPlatform.PreviewShaderFormatName == TestItem.ShaderFormat && GEditor->PreviewPlatform.PreviewShaderPlatformName == TestItem.PreviewShaderPlatformName;
+			});
+		if (Item)
+		{
+			return FSlateIcon(FAppStyle::GetAppStyleSetName(), GEditor->IsFeatureLevelPreviewActive() ? Item->ActiveIconName : Item->InactiveIconName);
+		}
+
+		EShaderPlatform ShaderPlatform = FDataDrivenShaderPlatformInfo::GetShaderPlatformFromName(GEditor->PreviewPlatform.PreviewShaderPlatformName);
+
+		if (ShaderPlatform == SP_NumPlatforms)
+		{
+			ShaderPlatform = GetFeatureLevelShaderPlatform(GEditor->PreviewPlatform.PreviewFeatureLevel);
+		}
+		switch (GEditor->PreviewPlatform.PreviewFeatureLevel)
+		{
+		case ERHIFeatureLevel::ES3_1:
+		{
+			return FSlateIcon(FAppStyle::GetAppStyleSetName(), GEditor->IsFeatureLevelPreviewActive() ? "LevelEditor.PreviewMode.Enabled" : "LevelEditor.PreviewMode.Disabled");
+		}
+		default:
+		{
+			return FSlateIcon(FAppStyle::GetAppStyleSetName(), GEditor->IsFeatureLevelPreviewActive() ? "LevelEditor.PreviewMode.Enabled" : "LevelEditor.PreviewMode.Disabled");
+		}
+		}
+	}
+
+	LEVELEDITOR_API void AddPreviewToggleButton(FToolMenuSection& Section)
+	{
+		Section.AddEntry(FToolMenuEntry::InitToolBarButton(
+			FLevelEditorCommands::Get().ToggleFeatureLevelPreview,
+			TAttribute<FText>::Create(&GetPreviewModeText),
+			TAttribute<FText>::Create(&GetPreviewModeTooltip),
+			TAttribute<FSlateIcon>::Create(&GetPreviewModeIcon)
+		));
+	}
+}
+
 namespace LevelEditorActionHelpers
 {
 	/** Filters out any classes for the Class Picker when creating or selecting classes in the Blueprints dropdown */
@@ -1397,82 +1477,8 @@ void FLevelEditorToolBar::RegisterLevelEditorToolBar( const TSharedRef<FUIComman
 
 		PlaySection.AddSeparator(NAME_None);
 
-		struct FPreviewModeFunctionality
-		{
-			static FText GetPreviewModeText()
-			{
-				const FPreviewPlatformMenuItem* Item = FDataDrivenPlatformInfoRegistry::GetAllPreviewPlatformMenuItems().FindByPredicate([](const FPreviewPlatformMenuItem& TestItem)
-					{
-						return GEditor->PreviewPlatform.PreviewPlatformName == TestItem.PlatformName && GEditor->PreviewPlatform.PreviewShaderFormatName == TestItem.ShaderFormat && GEditor->PreviewPlatform.PreviewShaderPlatformName == TestItem.PreviewShaderPlatformName;
-					});
-				return Item ? Item->IconText : FText();
-			}
-
-			static FText GetPreviewModeTooltip()
-			{
-				EShaderPlatform PreviewShaderPlatform = GEditor->PreviewPlatform.PreviewShaderPlatformName != NAME_None ?
-					FDataDrivenShaderPlatformInfo::GetShaderPlatformFromName(GEditor->PreviewPlatform.PreviewShaderPlatformName) :
-					GetFeatureLevelShaderPlatform(GEditor->PreviewPlatform.PreviewFeatureLevel);
+		PreviewModeFunctionality::AddPreviewToggleButton(PlaySection);
 
-				EShaderPlatform MaxRHIFeatureLevelPlatform = GetFeatureLevelShaderPlatform(GMaxRHIFeatureLevel);
-
-				{
-					const FText& RenderingAsPlatformName = FDataDrivenShaderPlatformInfo::GetFriendlyName(GEditor->PreviewPlatform.bPreviewFeatureLevelActive ? PreviewShaderPlatform : MaxRHIFeatureLevelPlatform);
-                    const FText& SwitchToPlatformName = FDataDrivenShaderPlatformInfo::GetFriendlyName(GEditor->PreviewPlatform.bPreviewFeatureLevelActive ? MaxRHIFeatureLevelPlatform : PreviewShaderPlatform);
-					if (PreviewShaderPlatform == MaxRHIFeatureLevelPlatform)
-					{
-						return FText::Format(LOCTEXT("PreviewModeViewingAs", "Viewing {0}."), RenderingAsPlatformName);
-					}
-                    else if (GWorld->FeatureLevel == GMaxRHIFeatureLevel)
-                    {
-                        return FText::Format(LOCTEXT("PreviewModeViewingAsSwitchTo", "Viewing {0}. Click to preview {1}."), RenderingAsPlatformName, SwitchToPlatformName);
-                    }
-                    else
-                    {
-                        return FText::Format(LOCTEXT("PreviewModePreviewingAsSwitchTo", "Previewing {0}. Click to view {1}."), RenderingAsPlatformName, SwitchToPlatformName);
-                    }
-				}
-			}
-
-			static FSlateIcon GetPreviewModeIcon()
-			{
-				const FPreviewPlatformMenuItem* Item = FDataDrivenPlatformInfoRegistry::GetAllPreviewPlatformMenuItems().FindByPredicate([](const FPreviewPlatformMenuItem& TestItem)
-					{
-						return GEditor->PreviewPlatform.PreviewPlatformName == TestItem.PlatformName && GEditor->PreviewPlatform.PreviewShaderFormatName == TestItem.ShaderFormat && GEditor->PreviewPlatform.PreviewShaderPlatformName == TestItem.PreviewShaderPlatformName;
-					});
-				if(Item)
-				{
-					return FSlateIcon(FAppStyle::GetAppStyleSetName(), GEditor->IsFeatureLevelPreviewActive() ? Item->ActiveIconName : Item->InactiveIconName);
-				}
-
-				EShaderPlatform ShaderPlatform = FDataDrivenShaderPlatformInfo::GetShaderPlatformFromName(GEditor->PreviewPlatform.PreviewShaderPlatformName);
-
-				if (ShaderPlatform == SP_NumPlatforms)
-				{
-					ShaderPlatform = GetFeatureLevelShaderPlatform(GEditor->PreviewPlatform.PreviewFeatureLevel);
-				}
-				switch (GEditor->PreviewPlatform.PreviewFeatureLevel)
-				{
-					case ERHIFeatureLevel::ES3_1:
-					{
-						return FSlateIcon(FAppStyle::GetAppStyleSetName(), GEditor->IsFeatureLevelPreviewActive() ? "LevelEditor.PreviewMode.Enabled" : "LevelEditor.PreviewMode.Disabled");
-					}
-					default:
-					{
-						return FSlateIcon(FAppStyle::GetAppStyleSetName(), GEditor->IsFeatureLevelPreviewActive() ? "LevelEditor.PreviewMode.Enabled" : "LevelEditor.PreviewMode.Disabled");
-					}
-				}
-			}
-		};
-
-		PlaySection.AddEntry(FToolMenuEntry::InitToolBarButton(
-			FLevelEditorCommands::Get().ToggleFeatureLevelPreview,
-			TAttribute<FText>::Create(&FPreviewModeFunctionality::GetPreviewModeText),
-        	TAttribute<FText>::Create(&FPreviewModeFunctionality::GetPreviewModeTooltip),
-        	TAttribute<FSlateIcon>::Create(&FPreviewModeFunctionality::GetPreviewModeIcon)
-			));
-
-	
 		// Add the shared play-world commands that will be shown on the Kismet toolbar as well
 		FPlayWorldCommands::BuildToolbar(PlaySection, true);
 
