Save viewport camera settings into per-user editor config file.

diff --git a/Engine/Source/Editor/LevelEditor/Private/SLevelViewport.cpp b/Engine/Source/Editor/LevelEditor/Private/SLevelViewport.cpp
--- a/Engine/Source/Editor/LevelEditor/Private/SLevelViewport.cpp
+++ b/Engine/Source/Editor/LevelEditor/Private/SLevelViewport.cpp
@@ -1211,6 +1211,11 @@ void SLevelViewport::OnMapChanged( UWorld* World, EMapChangeType MapChangeType )
 	{
 		if( MapChangeType == EMapChangeType::LoadMap )
 		{
+			if (FLevelEditorViewporEditorViews* PerUserEditorViews = GetMutableDefault<ULevelEditorViewportSettings>()->EditorViews.Find(World))
+			{
+				World->EditorViews = PerUserEditorViews->LevelViewportsInfo;
+			}
+
 			if (World->EditorViews[LevelViewportClient->ViewportType].CamOrthoZoom == 0.0f)
 			{
 				World->EditorViews[LevelViewportClient->ViewportType].CamOrthoZoom = DEFAULT_ORTHOZOOM;
@@ -1237,7 +1242,7 @@ void SLevelViewport::OnMapChanged( UWorld* World, EMapChangeType MapChangeType )
 				}
 			}
 		}
-		else if( MapChangeType == EMapChangeType::SaveMap )
+		else if( (MapChangeType == EMapChangeType::SaveMap) || (MapChangeType == EMapChangeType::TearDownWorld))
 		{
 			//@todo there could potentially be more than one of the same viewport type.  This effectively takes the last one of a specific type
 			World->EditorViews[LevelViewportClient->ViewportType] = 
@@ -1245,9 +1250,12 @@ void SLevelViewport::OnMapChanged( UWorld* World, EMapChangeType MapChangeType )
 					LevelViewportClient->GetViewLocation(),
 					LevelViewportClient->GetViewRotation(), 
 					LevelViewportClient->GetOrthoZoom() );
+
+			GetMutableDefault<ULevelEditorViewportSettings>()->EditorViews.FindOrAdd(World).LevelViewportsInfo = World->EditorViews;
+			GetMutableDefault<ULevelEditorViewportSettings>()->SaveConfig();
 		}
 		else if( MapChangeType == EMapChangeType::NewMap )
-		{
+		{		
 		
 			ResetNewLevelViewFlags();
 
diff --git a/Engine/Source/Editor/UnrealEd/Classes/Settings/LevelEditorViewportSettings.h b/Engine/Source/Editor/UnrealEd/Classes/Settings/LevelEditorViewportSettings.h
--- a/Engine/Source/Editor/UnrealEd/Classes/Settings/LevelEditorViewportSettings.h
+++ b/Engine/Source/Editor/UnrealEd/Classes/Settings/LevelEditorViewportSettings.h
@@ -8,6 +8,7 @@
 #include "UObject/SoftObjectPath.h"
 #include "Engine/EngineTypes.h"
 #include "Engine/EngineBaseTypes.h"
+#include "Engine/World.h"
 #include "Viewports.h"
 #include "Editor/UnrealEdTypes.h"
 #include "LevelEditorViewportSettings.generated.h"
@@ -249,6 +250,15 @@ struct UNREALED_API FSnapToSurfaceSettings
 	bool bSnapRotation;
 };
 
+USTRUCT()
+struct UNREALED_API FLevelEditorViewporEditorViews
+{
+	GENERATED_USTRUCT_BODY()
+
+	UPROPERTY()
+	TArray<FLevelViewportInfo> LevelViewportsInfo;
+};
+
 /**
  * Implements the Level Editor's view port settings.
  */
@@ -570,6 +580,10 @@ public:
 	UPROPERTY(EditAnywhere, config, Category = Behavior)
 	TMap<EMaterialKind, FName> MaterialParamsForDroppedTextures;
 
+	/** Store the last camera settings for all the viewport of each world loaded in the editor. */
+	UPROPERTY(config)
+	TMap<TSoftObjectPtr<UWorld>, FLevelEditorViewporEditorViews> EditorViews;
+
 private:
 
 	// Per-instance viewport settings.
