Fixed common typo (verticies->vertices)

#skipci
#rnx

diff --git a/Engine/Plugins/Runtime/GeometryProcessing/Source/DynamicMesh/Private/Operations/IntrinsicTriangulationMesh.cpp b/Engine/Plugins/Runtime/GeometryProcessing/Source/DynamicMesh/Private/Operations/IntrinsicTriangulationMesh.cpp
--- a/Engine/Plugins/Runtime/GeometryProcessing/Source/DynamicMesh/Private/Operations/IntrinsicTriangulationMesh.cpp
+++ b/Engine/Plugins/Runtime/GeometryProcessing/Source/DynamicMesh/Private/Operations/IntrinsicTriangulationMesh.cpp
@@ -4005,7 +4005,7 @@ UE::Geometry::EMeshResult UE::Geometry::FIntrinsicMesh::PokeTriangle(int32 Intri
 												{
 													const int32 NumBoundaryPts = BoundingConvexPolyVerts.Num();
 													// In the barycentric space of the intrinsic triangle:
-													// find the (min norm) linear combination convex-poly verticies that is equivalent to the new vertex location
+													// find the (min norm) linear combination convex-poly vertices that is equivalent to the new vertex location
 													TArray<double> MinNormSolution;
 													TArray<double> MatRows[3];
 													for (int32 i = 0; i < 3; ++i)
diff --git a/Engine/Plugins/Runtime/GeometryProcessing/Source/DynamicMesh/Public/Operations/UniformTessellate.h b/Engine/Plugins/Runtime/GeometryProcessing/Source/DynamicMesh/Public/Operations/UniformTessellate.h
--- a/Engine/Plugins/Runtime/GeometryProcessing/Source/DynamicMesh/Public/Operations/UniformTessellate.h
+++ b/Engine/Plugins/Runtime/GeometryProcessing/Source/DynamicMesh/Public/Operations/UniformTessellate.h
@@ -17,7 +17,7 @@ class FDynamicMesh3;
 /**
  * Given an input mesh and a tessellation level, this operator generates a new tessellated mesh where every triangle in 
  * the input mesh is uniformly subtriangulated into (TessellationNum + 1)^2 triangles. Per-vertex normals, uvs, 
- * colors and extended per-vertex attributes are linearly interpolated to the new verticies. Per-triangle group 
+ * colors and extended per-vertex attributes are linearly interpolated to the new vertices. Per-triangle group 
  * identifiers/materials and extended triangle attributes for the new triangles are inherited from the corresponding 
  * input mesh triangles they replaced. 
  * 
diff --git a/Engine/Source/Editor/LevelEditor/Private/LevelEditorActions.cpp b/Engine/Source/Editor/LevelEditor/Private/LevelEditorActions.cpp
--- a/Engine/Source/Editor/LevelEditor/Private/LevelEditorActions.cpp
+++ b/Engine/Source/Editor/LevelEditor/Private/LevelEditorActions.cpp
@@ -3631,7 +3631,7 @@ void FLevelEditorCommands::RegisterCommands()
 
 	UI_COMMAND( MergePolys, "Merge", "Merges multiple polygons on a brush face into as few as possible", EUserInterfaceActionType::Button, FInputChord() );
 	UI_COMMAND( SeparatePolys, "Separate", "Reverses the effect of a previous merge", EUserInterfaceActionType::Button, FInputChord() );
-	UI_COMMAND( AlignBrushVerticesToGrid, "Align Brush Verticies To Grid", "Align brush verticies to the grid", EUserInterfaceActionType::Button, FInputChord() );
+	UI_COMMAND( AlignBrushVerticesToGrid, "Align Brush Vertices To Grid", "Align brush vertices to the grid", EUserInterfaceActionType::Button, FInputChord() );
 
 	// RegroupActors uses GroupActors for it's label and tooltip when simply grouping a selection of actors using overrides. This is to provide display of the chord which is the same for both.
 	UI_COMMAND( GroupActors, "Group", "Groups the selected actors", EUserInterfaceActionType::Button, FInputChord( /*EKeys::G, EModifierKey::Control*/ ) );
diff --git a/Engine/Source/Editor/LevelEditor/Public/LevelEditorActions.h b/Engine/Source/Editor/LevelEditor/Public/LevelEditorActions.h
--- a/Engine/Source/Editor/LevelEditor/Public/LevelEditorActions.h
+++ b/Engine/Source/Editor/LevelEditor/Public/LevelEditorActions.h
@@ -308,7 +308,7 @@ public:
 	/** Reverse a merge */
 	TSharedPtr< FUICommandInfo > SeparatePolys;
 
-	/** Align brush verticies to the grid */
+	/** Align brush vertices to the grid */
 	TSharedPtr<FUICommandInfo> AlignBrushVerticesToGrid;
 
 	/**
@@ -1340,7 +1340,7 @@ public:
 	static void SnapElementsToElement_Clicked( bool InAlign, bool InUseLineTrace, bool InUseBounds, bool InUsePivot );
 
 	/**
-	 * Aligns brush verticies to the nearest grid point.
+	 * Aligns brush vertices to the nearest grid point.
 	 */
 	static void AlignBrushVerticesToGrid_Execute();
 
diff --git a/Engine/Source/Runtime/Engine/Public/ParticleEmitterInstances.h b/Engine/Source/Runtime/Engine/Public/ParticleEmitterInstances.h
--- a/Engine/Source/Runtime/Engine/Public/ParticleEmitterInstances.h
+++ b/Engine/Source/Runtime/Engine/Public/ParticleEmitterInstances.h
@@ -183,7 +183,7 @@ struct FParticleEmitterBuildInfo
 
 	int32 ScreenAlignment;
 
-	/** An offset in UV space for the positioning of a sprites verticies. */
+	/** An offset in UV space for the positioning of a sprites vertices. */
 	FVector2D PivotOffset;
 
 	/** If true, local vector fields ignore the component transform. */
