Update Landscape Build menu entries

diff --git a/Engine/Source/Editor/LevelEditor/Private/LevelEditorActions.cpp b/Engine/Source/Editor/LevelEditor/Private/LevelEditorActions.cpp
--- a/Engine/Source/Editor/LevelEditor/Private/LevelEditorActions.cpp
+++ b/Engine/Source/Editor/LevelEditor/Private/LevelEditorActions.cpp
@@ -3530,7 +3530,7 @@ void FLevelEditorCommands::RegisterCommands()
 	UI_COMMAND( BuildLandscapeSplineMeshes, "Build Landscape Spline Meshes", "Builds landscape spline meshes for the current world", EUserInterfaceActionType::Button, FInputChord());
 	UI_COMMAND( BuildTextureStreamingOnly, "Build Texture Streaming", "Build texture streaming data", EUserInterfaceActionType::Button, FInputChord() );
 	UI_COMMAND( BuildVirtualTextureOnly, "Build Virtual Textures", "Build runtime virtual texture low mips streaming data", EUserInterfaceActionType::Button, FInputChord());
-	UI_COMMAND( BuildAllLandscape, "Build All Landscape", "Build All Landscape Data(Grass, Baked Textures)", EUserInterfaceActionType::Button, FInputChord());
+	UI_COMMAND( BuildAllLandscape, "Build Landscape", "Build all data related to Landscape (grass maps, GI textures, physical material, Nanite, dirty height and weight maps)", EUserInterfaceActionType::Button, FInputChord());
 
 	UI_COMMAND( LightingQuality_Production, "Production", "Sets precomputed lighting quality to highest possible quality (slowest computation time.)", EUserInterfaceActionType::RadioButton, FInputChord() );
 	UI_COMMAND( LightingQuality_High, "High", "Sets precomputed lighting quality to high quality", EUserInterfaceActionType::RadioButton, FInputChord() );
