Sort the Placed Cameras menu by camera name to make the list more predictable for users.

diff --git a/Engine/Source/Editor/LevelEditor/Private/SLevelViewportToolBar.cpp b/Engine/Source/Editor/LevelEditor/Private/SLevelViewportToolBar.cpp
--- a/Engine/Source/Editor/LevelEditor/Private/SLevelViewportToolBar.cpp
+++ b/Engine/Source/Editor/LevelEditor/Private/SLevelViewportToolBar.cpp
@@ -49,6 +49,7 @@
 #include "ToolMenu.h"
 #include "WorldPartition/WorldPartitionSubsystem.h"
 #include "SLevelViewport.h"
+#include "SortHelper.h"
 
 #define LOCTEXT_NAMESPACE "LevelViewportToolBar"
 
@@ -912,6 +913,16 @@ void SLevelViewportToolBar::GeneratePlacedCameraMenuEntries(FToolMenuSection& Se
 {
 	FSlateIcon CameraIcon( FAppStyle::GetAppStyleSetName(), "ClassIcon.CameraComponent" );
 
+	// Sort the cameras to make the ordering predictable for users.
+	Cameras.StableSort([](const ACameraActor& Left, const ACameraActor& Right)
+	{
+		// Do "natural sorting" via SceneOutliner::FNumericStringWrapper to make more sense to humans (also matches the Scene Outliner). This sorts "Camera2" before "Camera10" which a normal lexicographical sort wouldn't.
+		SceneOutliner::FNumericStringWrapper LeftWrapper(FString(Left.GetActorLabel()));
+		SceneOutliner::FNumericStringWrapper RightWrapper(FString(Right.GetActorLabel()));
+
+		return LeftWrapper < RightWrapper;
+	});
+
 	for( ACameraActor* CameraActor : Cameras )
 	{
 		// Needed for the delegate hookup to work below
