Fix DPI issue in STransformViewportToolBar that broke collapse functionality for small viewports

This adjusts a function added in CL 22676675 to fix the issue introduced by 22635797 also in the case of DPI scaling.

diff --git a/Engine/Source/Editor/LevelEditor/Private/SLevelViewportToolBar.cpp b/Engine/Source/Editor/LevelEditor/Private/SLevelViewportToolBar.cpp
--- a/Engine/Source/Editor/LevelEditor/Private/SLevelViewportToolBar.cpp
+++ b/Engine/Source/Editor/LevelEditor/Private/SLevelViewportToolBar.cpp
@@ -50,6 +50,7 @@
 #include "WorldPartition/WorldPartitionSubsystem.h"
 #include "SLevelViewport.h"
 #include "SortHelper.h"
+#include "Interfaces/IMainFrameModule.h"
 
 #define LOCTEXT_NAMESPACE "LevelViewportToolBar"
 
@@ -1805,12 +1806,22 @@ float SLevelViewportToolBar::GetTransformToolbarWidth() const
 				const float ToolbarWidthMinusPreviousTransformToolbar = GetDesiredSize().X - TransformToolbar_CachedMaxWidth;
 				const float ToolbarWidthEstimate = ToolbarWidthMinusPreviousTransformToolbar + TransformToolbarWidth;
 
-				const float OverflowWidth = ToolbarWidthEstimate - ViewportWidth;
+				float DpiScale = 1.0f;
+				{
+					IMainFrameModule& MainFrameModule = FModuleManager::LoadModuleChecked<IMainFrameModule>(TEXT("MainFrame"));
+					const TSharedPtr<SWindow>& MainFrameParentWindow = MainFrameModule.GetParentWindow();
+					if (MainFrameParentWindow.IsValid())
+					{
+						DpiScale = MainFrameParentWindow->GetDPIScaleFactor();
+					}
+				}
+
+				const float OverflowWidth = ToolbarWidthEstimate * DpiScale - ViewportWidth;
 				if (OverflowWidth > 0.0f)
 				{
 					// There isn't enough space in the viewport to show the toolbar!
 					// Try and shrink the transform toolbar (which has an overflow area) to make things fit
-					TransformToolbar_CachedMaxWidth = FMath::Max(FMath::Min(4.0f, TransformToolbarWidth), TransformToolbarWidth - OverflowWidth);
+					TransformToolbar_CachedMaxWidth = FMath::Max(FMath::Min(4.0f, TransformToolbarWidth), TransformToolbarWidth - OverflowWidth / DpiScale);
 				}
 				else
 				{
