Vulkan:  Fix the missing "Disable Preview" option in in Vulkan when using preview modes.  Disable preview modes for the same FeatureLevel until all issues are resolved.

#rnx

diff --git a/Engine/Config/VulkanPC/DataDrivenPlatformInfo.ini b/Engine/Config/VulkanPC/DataDrivenPlatformInfo.ini
--- a/Engine/Config/VulkanPC/DataDrivenPlatformInfo.ini
+++ b/Engine/Config/VulkanPC/DataDrivenPlatformInfo.ini
@@ -146,3 +146,10 @@ bSupportsIndexBufferUAVs=true
 bSupportsManualVertexFetch=false
 bSupportsSceneDataCompressedTransforms=true
 FriendlyName=LOCTEXT("FriendlyShaderPlatformName_VulkanPC_ES31", "PC Vulkan ES31")
+
+
+[PreviewPlatform VULKAN_SM5]
+PlatformName=VULKAN_SM5
+ShaderFormat=SF_VULKAN_SM5
+ShaderPlatform=VULKAN_SM5
+MenuTooltip=LOCTEXT("PreviewMenuTooltip_VulkanPC_SM5", "Linux using SM5 profile")
diff --git a/Engine/Source/Editor/LevelEditor/Private/LevelEditor.cpp b/Engine/Source/Editor/LevelEditor/Private/LevelEditor.cpp
--- a/Engine/Source/Editor/LevelEditor/Private/LevelEditor.cpp
+++ b/Engine/Source/Editor/LevelEditor/Private/LevelEditor.cpp
@@ -1816,16 +1816,16 @@ void FLevelEditorModule::BindGlobalLevelEditorCommands()
 		
 		if (ShaderPlatform < SP_NumPlatforms)
 		{
-			ERHIFeatureLevel::Type FeatureLevel = GetMaxSupportedFeatureLevel(ShaderPlatform);
+			const ERHIFeatureLevel::Type FeatureLevel = GetMaxSupportedFeatureLevel(ShaderPlatform);
 
-			bool bIsDefaultShaderPlatform = FDataDrivenShaderPlatformInfo::GetPreviewShaderPlatformParent(ShaderPlatform) == GMaxRHIShaderPlatform;
+			const bool bIsDefaultShaderPlatform = FDataDrivenShaderPlatformInfo::GetPreviewShaderPlatformParent(ShaderPlatform) == GMaxRHIShaderPlatform;
 
 			FPreviewPlatformInfo PreviewFeatureLevelInfo(bIsDefaultShaderPlatform ? GMaxRHIFeatureLevel : FeatureLevel, bIsDefaultShaderPlatform ? GMaxRHIShaderPlatform : (EShaderPlatform)ShaderPlatform, bIsDefaultShaderPlatform ? NAME_None : Item.PlatformName, bIsDefaultShaderPlatform ? NAME_None : Item.ShaderFormat, bIsDefaultShaderPlatform ? NAME_None : Item.DeviceProfileName, true, bIsDefaultShaderPlatform ? NAME_None : Item.PreviewShaderPlatformName);
 
 			ActionList.MapAction(
 				Commands.PreviewPlatformOverrides[Index],
 				FExecuteAction::CreateStatic(&FLevelEditorActionCallbacks::SetPreviewPlatform, PreviewFeatureLevelInfo),
-				FCanExecuteAction::CreateStatic(&FLevelEditorActionCallbacks::CanExecutePreviewPlatform, PreviewFeatureLevelInfo),
+				bIsDefaultShaderPlatform ? FCanExecuteAction() : FCanExecuteAction::CreateStatic(&FLevelEditorActionCallbacks::CanExecutePreviewPlatform, PreviewFeatureLevelInfo),
 				FIsActionChecked::CreateStatic(&FLevelEditorActionCallbacks::IsPreviewPlatformChecked, PreviewFeatureLevelInfo));
 		}
 	}
diff --git a/Engine/Source/Editor/LevelEditor/Private/LevelEditorActions.cpp b/Engine/Source/Editor/LevelEditor/Private/LevelEditorActions.cpp
--- a/Engine/Source/Editor/LevelEditor/Private/LevelEditorActions.cpp
+++ b/Engine/Source/Editor/LevelEditor/Private/LevelEditorActions.cpp
@@ -671,7 +671,20 @@ void FLevelEditorActionCallbacks::SetPreviewPlatform(FPreviewPlatformInfo NewPre
 
 bool FLevelEditorActionCallbacks::CanExecutePreviewPlatform(FPreviewPlatformInfo NewPreviewPlatform)
 {
-	return NewPreviewPlatform.PreviewFeatureLevel <= GMaxRHIFeatureLevel;
+	// TODO: Prevent switching from a platform with VSM to a platform without VSM at the same FeatureLevel until all issues are resolved.
+	// (for instance, GlobalShaderMap does not contain the shaders necessary for FVirtualShadowMapArrayCacheManager::ProcessInvalidations anymore)
+	// Currently this is mostly meant to isolate going to/from VULKAN_SM5 which is wildly different and causes issues with preview mechanics.
+	if (NewPreviewPlatform.PreviewFeatureLevel == GMaxRHIFeatureLevel)
+	{
+		if (FDataDrivenShaderPlatformInfo::IsValid(NewPreviewPlatform.ShaderPlatform) &&
+			FDataDrivenShaderPlatformInfo::GetIsPreviewPlatform(NewPreviewPlatform.ShaderPlatform))
+		{
+			const EShaderPlatform ParentShaderPlatform = FDataDrivenShaderPlatformInfo::GetPreviewShaderPlatformParent(NewPreviewPlatform.ShaderPlatform);
+			return (DoesPlatformSupportVirtualShadowMaps(GMaxRHIShaderPlatform) == DoesPlatformSupportVirtualShadowMaps(ParentShaderPlatform));
+		}
+	}
+
+	return (NewPreviewPlatform.PreviewFeatureLevel < GMaxRHIFeatureLevel);
 }
 
 bool FLevelEditorActionCallbacks::IsPreviewPlatformChecked(FPreviewPlatformInfo PreviewPlatform)
