PR #10303 :FEditorDelegates::OnLightingBuildSucceeded delegate is removed correctly
Previously causes the editor to crash when used multiple times and when the prev level is unloaded.

(Contributed by [REDACTED_PERSON])

diff --git a/Engine/Source/Editor/LevelEditor/Private/LevelEditorSubsystem.cpp b/Engine/Source/Editor/LevelEditor/Private/LevelEditorSubsystem.cpp
--- a/Engine/Source/Editor/LevelEditor/Private/LevelEditorSubsystem.cpp
+++ b/Engine/Source/Editor/LevelEditor/Private/LevelEditorSubsystem.cpp
@@ -716,7 +716,7 @@ bool ULevelEditorSubsystem::BuildLightMaps(ELightingBuildQuality Quality, bool b
 	}
 
 	FEditorDelegates::OnLightingBuildFailed.Remove(BuildFailedDelegateHandle);
-	FEditorDelegates::OnLightingBuildFailed.Remove(BuildSucceededDelegateHandle);
+	FEditorDelegates::OnLightingBuildSucceeded.Remove(BuildSucceededDelegateHandle);
 
 	return Success;
 }
