Sync "reset to default" state in the Actor details view on archetype edits that occur while archetype instances are actively selected to it, e.g. while changing Blueprint class defaults.

diff --git a/Engine/Source/Editor/LevelEditor/Private/SActorDetails.cpp b/Engine/Source/Editor/LevelEditor/Private/SActorDetails.cpp
--- a/Engine/Source/Editor/LevelEditor/Private/SActorDetails.cpp
+++ b/Engine/Source/Editor/LevelEditor/Private/SActorDetails.cpp
@@ -223,6 +223,7 @@ void SActorDetails::Construct(const FArguments& InArgs, UTypedElementSelectionSe
 	checkf(SelectionSet, TEXT("SActorDetails must be constructed with a valid selection set!"));
 
 	FCoreUObjectDelegates::OnObjectsReplaced.AddRaw(this, &SActorDetails::OnObjectsReplaced);
+	FCoreUObjectDelegates::OnObjectPropertyChanged.AddRaw(this, &SActorDetails::OnObjectPropertyChanged);
 
 	FLevelEditorModule& LevelEditor = FModuleManager::GetModuleChecked<FLevelEditorModule>("LevelEditor");
 	LevelEditor.OnComponentsEdited().AddRaw(this, &SActorDetails::OnComponentsEditedInWorld);
@@ -405,6 +406,7 @@ SActorDetails::~SActorDetails()
 	}
 
 	FCoreUObjectDelegates::OnObjectsReplaced.RemoveAll(this);
+	FCoreUObjectDelegates::OnObjectPropertyChanged.RemoveAll(this);
 
 	RemoveBPComponentCompileEventDelegate();
 
@@ -1052,4 +1054,29 @@ void SActorDetails::OnObjectsReplaced(const TMap<UObject*, UObject*>& InReplacem
 	}
 }
 
+void SActorDetails::OnObjectPropertyChanged(UObject* ObjectBeingModified, FPropertyChangedEvent& PropertyChangedEvent)
+{
+	// Listen for archetype changes to properties that are editable on the instance
+	if (!ObjectBeingModified
+		|| !ObjectBeingModified->HasAnyFlags(RF_ClassDefaultObject | RF_ArchetypeObject)
+		|| !PropertyChangedEvent.Property
+		|| !PropertyChangedEvent.Property->HasAnyPropertyFlags(CPF_Edit)
+		|| PropertyChangedEvent.Property->HasAnyPropertyFlags(CPF_DisableEditOnInstance))
+	{
+		return;
+	}
+
+	// If the object that changes matches the archetype of an instance that's selected to the property view,
+	// invalidate the cached state so that things like the "reset to default" icon is synced to the archetype.
+	for (const TWeakObjectPtr<>& SelectedObject : DetailsView->GetSelectedObjects())
+	{
+		if (SelectedObject.IsValid()
+			&& SelectedObject->GetArchetype() == ObjectBeingModified)
+		{
+			DetailsView->InvalidateCachedState();
+			break;
+		}
+	}
+}
+
 #undef LOCTEXT_NAMESPACE
diff --git a/Engine/Source/Editor/LevelEditor/Private/SActorDetails.h b/Engine/Source/Editor/LevelEditor/Private/SActorDetails.h
--- a/Engine/Source/Editor/LevelEditor/Private/SActorDetails.h
+++ b/Engine/Source/Editor/LevelEditor/Private/SActorDetails.h
@@ -24,6 +24,7 @@ class UTypedElementSelectionSet;
 class SSubobjectEditor;
 class ISCSEditorUICustomization;
 class FSubobjectEditorTreeNode;
+struct FPropertyChangedEvent;
 
 namespace UE::LevelEditor::Private
 {
@@ -105,6 +106,7 @@ private:
 	void RemoveBPComponentCompileEventDelegate();
 	void OnBlueprintComponentCompiled(UBlueprint* ComponentBlueprint);
 	void OnObjectsReplaced(const TMap<UObject*, UObject*>& InReplacementObjects);
+	void OnObjectPropertyChanged(UObject* ObjectBeingModified, FPropertyChangedEvent& PropertyChangedEvent);
 
 private:
 	TSharedPtr<SSplitter> DetailsSplitter;
