Further block Vulkan_SM5 as a preview mode until all issues can be resolved.

#rnx

diff --git a/Engine/Source/Editor/LevelEditor/Private/LevelEditorActions.cpp b/Engine/Source/Editor/LevelEditor/Private/LevelEditorActions.cpp
--- a/Engine/Source/Editor/LevelEditor/Private/LevelEditorActions.cpp
+++ b/Engine/Source/Editor/LevelEditor/Private/LevelEditorActions.cpp
@@ -684,6 +684,15 @@ bool FLevelEditorActionCallbacks::CanExecutePreviewPlatform(FPreviewPlatformInfo
 		}
 	}
 
+	// TODO: Prevent previewing VULKAN_SM5 with a D3D renderer until all issues are resolved.
+	if (FDataDrivenShaderPlatformInfo::IsValid(NewPreviewPlatform.ShaderPlatform) &&
+		FDataDrivenShaderPlatformInfo::GetIsPreviewPlatform(NewPreviewPlatform.ShaderPlatform) &&
+		FDataDrivenShaderPlatformInfo::GetIsLanguageD3D(GMaxRHIShaderPlatform))
+	{
+		const EShaderPlatform ParentShaderPlatform = FDataDrivenShaderPlatformInfo::GetPreviewShaderPlatformParent(NewPreviewPlatform.ShaderPlatform);
+		return !(FDataDrivenShaderPlatformInfo::GetIsLanguageVulkan(ParentShaderPlatform) && IsFeatureLevelSupported(ParentShaderPlatform, ERHIFeatureLevel::SM5));
+	}
+
 	return (NewPreviewPlatform.PreviewFeatureLevel < GMaxRHIFeatureLevel);
 }
 
