Fixing CanExecutePreviewPlatform to work correctly with multiple conditional disables. By returning early when checking the first condition, the other conditions were being skipped over, which led to things like being able to preview SM6 from a SM5 max capable runtime.

diff --git a/Engine/Source/Editor/LevelEditor/Private/LevelEditorActions.cpp b/Engine/Source/Editor/LevelEditor/Private/LevelEditorActions.cpp
--- a/Engine/Source/Editor/LevelEditor/Private/LevelEditorActions.cpp
+++ b/Engine/Source/Editor/LevelEditor/Private/LevelEditorActions.cpp
@@ -671,6 +671,11 @@ void FLevelEditorActionCallbacks::SetPreviewPlatform(FPreviewPlatformInfo NewPre
 
 bool FLevelEditorActionCallbacks::CanExecutePreviewPlatform(FPreviewPlatformInfo NewPreviewPlatform)
 {
+	if (NewPreviewPlatform.PreviewFeatureLevel > GMaxRHIFeatureLevel)
+	{
+		return false;
+	}
+
 	// TODO: Prevent switching from a platform with VSM to a platform without VSM at the same FeatureLevel until all issues are resolved.
 	// (for instance, GlobalShaderMap does not contain the shaders necessary for FVirtualShadowMapArrayCacheManager::ProcessInvalidations anymore)
 	// Currently this is mostly meant to isolate going to/from VULKAN_SM5 which is wildly different and causes issues with preview mechanics.
@@ -680,7 +685,10 @@ bool FLevelEditorActionCallbacks::CanExecutePreviewPlatform(FPreviewPlatformInfo
 			FDataDrivenShaderPlatformInfo::GetIsPreviewPlatform(NewPreviewPlatform.ShaderPlatform))
 		{
 			const EShaderPlatform ParentShaderPlatform = FDataDrivenShaderPlatformInfo::GetPreviewShaderPlatformParent(NewPreviewPlatform.ShaderPlatform);
-			return (DoesPlatformSupportVirtualShadowMaps(GMaxRHIShaderPlatform) == DoesPlatformSupportVirtualShadowMaps(ParentShaderPlatform));
+			if (DoesPlatformSupportVirtualShadowMaps(GMaxRHIShaderPlatform) != DoesPlatformSupportVirtualShadowMaps(ParentShaderPlatform))
+			{
+				return false;
+			}
 		}
 	}
 
@@ -690,10 +698,13 @@ bool FLevelEditorActionCallbacks::CanExecutePreviewPlatform(FPreviewPlatformInfo
 		FDataDrivenShaderPlatformInfo::GetIsLanguageD3D(GMaxRHIShaderPlatform))
 	{
 		const EShaderPlatform ParentShaderPlatform = FDataDrivenShaderPlatformInfo::GetPreviewShaderPlatformParent(NewPreviewPlatform.ShaderPlatform);
-		return !(FDataDrivenShaderPlatformInfo::GetIsLanguageVulkan(ParentShaderPlatform) && IsFeatureLevelSupported(ParentShaderPlatform, ERHIFeatureLevel::SM5));
+		if (FDataDrivenShaderPlatformInfo::GetIsLanguageVulkan(ParentShaderPlatform) && IsFeatureLevelSupported(ParentShaderPlatform, ERHIFeatureLevel::SM5))
+		{
+			return false;
+		}
 	}
 
-	return (NewPreviewPlatform.PreviewFeatureLevel < GMaxRHIFeatureLevel);
+	return true;
 }
 
 bool FLevelEditorActionCallbacks::IsPreviewPlatformChecked(FPreviewPlatformInfo PreviewPlatform)
