Add a Browse button to the Level Editor Toolbar

diff --git a/Engine/Source/Editor/LevelEditor/Private/LevelEditor.cpp b/Engine/Source/Editor/LevelEditor/Private/LevelEditor.cpp
--- a/Engine/Source/Editor/LevelEditor/Private/LevelEditor.cpp
+++ b/Engine/Source/Editor/LevelEditor/Private/LevelEditor.cpp
@@ -673,6 +673,7 @@ void FLevelEditorModule::BindGlobalLevelEditorCommands()
 	ActionList.MapAction( Commands.Save, FExecuteAction::CreateStatic( &FLevelEditorActionCallbacks::Save ), FCanExecuteAction::CreateStatic( &FLevelEditorActionCallbacks::CanSaveWorld ) );
 	ActionList.MapAction( Commands.SaveAs, FExecuteAction::CreateStatic( &FLevelEditorActionCallbacks::SaveCurrentAs ), FCanExecuteAction::CreateStatic( &FLevelEditorActionCallbacks::CanSaveCurrentAs), FGetActionCheckState(), FIsActionButtonVisible::CreateStatic( &FLevelEditorActionCallbacks::CanSaveCurrentAs) );
 	ActionList.MapAction( Commands.SaveAllLevels, FExecuteAction::CreateStatic( &FLevelEditorActionCallbacks::SaveAllLevels ), FCanExecuteAction::CreateStatic( &FLevelEditorActionCallbacks::CanSaveUnpartitionedWorld), FGetActionCheckState(), FIsActionButtonVisible::CreateStatic(&FLevelEditorActionCallbacks::CanSaveUnpartitionedWorld) );
+	ActionList.MapAction( Commands.BrowseLevel, FExecuteAction::CreateStatic( &FLevelEditorActionCallbacks::Browse ), FCanExecuteAction::CreateStatic( &FLevelEditorActionCallbacks::CanBrowse ) );
 	ActionList.MapAction( Commands.ToggleFavorite, FExecuteAction::CreateStatic( &FLevelEditorActionCallbacks::ToggleFavorite ), FCanExecuteAction::CreateStatic( &FLevelEditorActionCallbacks::ToggleFavorite_CanExecute ), FIsActionChecked::CreateStatic( &FLevelEditorActionCallbacks::ToggleFavorite_IsChecked ) );
 
 	for( int32 CurRecentIndex = 0; CurRecentIndex < FLevelEditorCommands::MaxRecentFiles; ++CurRecentIndex )
diff --git a/Engine/Source/Editor/LevelEditor/Private/LevelEditorActions.cpp b/Engine/Source/Editor/LevelEditor/Private/LevelEditorActions.cpp
--- a/Engine/Source/Editor/LevelEditor/Private/LevelEditorActions.cpp
+++ b/Engine/Source/Editor/LevelEditor/Private/LevelEditorActions.cpp
@@ -515,6 +515,20 @@ void FLevelEditorActionCallbacks::SaveAllLevels()
 	FEditorFileUtils::SaveDirtyPackages( bPromptUserToSave, bSaveMapPackages, bSaveContentPackages, bFastSave );
 }
 
+void FLevelEditorActionCallbacks::Browse()
+{
+	if (const ULevel* CurrentLevel = GetWorld()->GetCurrentLevel())
+	{
+		TArray<UObject*> Assets = { CurrentLevel->GetOuter() };
+		GEditor->SyncBrowserToObjects(Assets);
+	}
+}
+
+bool FLevelEditorActionCallbacks::CanBrowse()
+{
+	return !FPackageName::IsTempPackage(GetWorld()->GetPackage()->GetName());
+}
+
 
 void FLevelEditorActionCallbacks::ImportScene_Clicked()
 {
@@ -3541,6 +3555,7 @@ void FLevelEditorCommands::RegisterCommands()
 	UI_COMMAND( Save, "Save Current Level", "Saves the current level to disk", EUserInterfaceActionType::Button, FInputChord(EModifierKey::Control, EKeys::S) );
 	UI_COMMAND( SaveAs, "Save Current Level As...", "Save the current level as...", EUserInterfaceActionType::Button, FInputChord( EModifierKey::Control|EModifierKey::Alt, EKeys::S ) );
 	UI_COMMAND( SaveAllLevels, "Save All Levels", "Saves all unsaved levels to disk", EUserInterfaceActionType::Button, FInputChord() );
+	UI_COMMAND( BrowseLevel, "Browse To Level", "Browses to the associated level and selects it in the most recently used Content Browser (summoning one if necessary)", EUserInterfaceActionType::Button, FInputChord() );
 	UI_COMMAND( ToggleFavorite, "Toggle Favorite", "Sets whether the currently loaded level will appear in the list of favorite levels", EUserInterfaceActionType::Button, FInputChord() );
 
 	for( int32 CurRecentIndex = 0; CurRecentIndex < FLevelEditorCommands::MaxRecentFiles; ++CurRecentIndex )
diff --git a/Engine/Source/Editor/LevelEditor/Private/LevelEditorToolBar.cpp b/Engine/Source/Editor/LevelEditor/Private/LevelEditorToolBar.cpp
--- a/Engine/Source/Editor/LevelEditor/Private/LevelEditorToolBar.cpp
+++ b/Engine/Source/Editor/LevelEditor/Private/LevelEditorToolBar.cpp
@@ -1305,6 +1305,14 @@ void FLevelEditorToolBar::RegisterLevelEditorToolBar( const TSharedRef<FUIComman
 				NAME_None,
 				FName("SaveAllLevels")
 			));
+
+			// Browse Level
+			Section.AddEntry(FToolMenuEntry::InitToolBarButton(
+				FLevelEditorCommands::Get().BrowseLevel,
+				TAttribute<FText>(),
+				TAttribute<FText>(),
+				FSlateIcon(FAppStyle::GetAppStyleSetName(), "SystemWideCommands.FindInContentBrowser")
+			));
 		}
 
 		TWeakPtr<SLevelEditor> LevelEditorPtr = InLevelEditor;
diff --git a/Engine/Source/Editor/LevelEditor/Public/LevelEditorActions.h b/Engine/Source/Editor/LevelEditor/Public/LevelEditorActions.h
--- a/Engine/Source/Editor/LevelEditor/Public/LevelEditorActions.h
+++ b/Engine/Source/Editor/LevelEditor/Public/LevelEditorActions.h
@@ -48,6 +48,7 @@ public:
 	TSharedPtr< FUICommandInfo > Save;
 	TSharedPtr< FUICommandInfo > SaveAs;
 	TSharedPtr< FUICommandInfo > SaveAllLevels;
+	TSharedPtr< FUICommandInfo > BrowseLevel;
 
 	static const int32 MaxRecentFiles = 10;
 	TArray< TSharedPtr< FUICommandInfo > > OpenRecentFileCommands;
@@ -712,6 +713,10 @@ public:
 	/** Saves all unsaved maps (but not packages) */
 	static void SaveAllLevels();
 
+	/** Browses to the current map */
+	static void Browse();
+	static bool CanBrowse();
+
 
 	/**
 	 * Called when import scene is selected
