Fix warnings about menus being registered twice
#rnx

diff --git a/Engine/Source/Editor/LevelEditor/Private/LevelEditorMenu.cpp b/Engine/Source/Editor/LevelEditor/Private/LevelEditorMenu.cpp
--- a/Engine/Source/Editor/LevelEditor/Private/LevelEditorMenu.cpp
+++ b/Engine/Source/Editor/LevelEditor/Private/LevelEditorMenu.cpp
@@ -209,7 +209,7 @@ void FLevelEditorMenu::RegisterLevelEditorMenus()
 
 		static void ExtendEditMenu()
 		{
-			UToolMenu* Menu = UToolMenus::Get()->RegisterMenu("LevelEditor.MainMenu.Edit", "MainFrame.MainMenu.Edit");
+			UToolMenu* Menu = UToolMenus::Get()->RegisterMenu("LevelEditor.MainMenu.Edit", "MainFrame.MainMenu.Edit", EMultiBoxType::Menu, /*bWarnIfAlreadyRegistered*/false);
 			{
 				// Edit Actor
 				{
@@ -228,7 +228,7 @@ void FLevelEditorMenu::RegisterLevelEditorMenus()
 
 		static void ExtendHelpMenu()
 		{
-			UToolMenu* Menu = UToolMenus::Get()->RegisterMenu("LevelEditor.MainMenu.Help", "MainFrame.MainMenu.Help");
+			UToolMenu* Menu = UToolMenus::Get()->RegisterMenu("LevelEditor.MainMenu.Help", "MainFrame.MainMenu.Help", EMultiBoxType::Menu, /*bWarnIfAlreadyRegistered*/false);
 			FToolMenuSection& Section = Menu->AddSection("HelpResources", NSLOCTEXT("MainHelpMenu", "LevelEditorHelpResources", "Level Editor Resources"), FToolMenuInsert("Learn", EToolMenuInsertType::First));
 			{
 				Section.AddMenuEntry( FLevelEditorCommands::Get().BrowseDocumentation );
@@ -297,10 +297,11 @@ void FLevelEditorMenu::RegisterLevelEditorMenus()
 	};
 
 	UToolMenus* ToolMenus = UToolMenus::Get();
-	ToolMenus->RegisterMenu("LevelEditor.MainMenu", "MainFrame.MainMenu", EMultiBoxType::MenuBar);
-	ToolMenus->RegisterMenu("LevelEditor.MainMenu.File", "MainFrame.MainTabMenu.File");
-	ToolMenus->RegisterMenu("LevelEditor.MainMenu.Window", "MainFrame.MainMenu.Window");
-	ToolMenus->RegisterMenu("LevelEditor.MainMenu.Tools", "MainFrame.MainMenu.Tools");
+	const bool bWarnIfAlreadyRegistered = false;
+	ToolMenus->RegisterMenu("LevelEditor.MainMenu", "MainFrame.MainMenu", EMultiBoxType::MenuBar, bWarnIfAlreadyRegistered);
+	ToolMenus->RegisterMenu("LevelEditor.MainMenu.File", "MainFrame.MainTabMenu.File", EMultiBoxType::Menu, bWarnIfAlreadyRegistered);
+	ToolMenus->RegisterMenu("LevelEditor.MainMenu.Window", "MainFrame.MainMenu.Window", EMultiBoxType::Menu, bWarnIfAlreadyRegistered);
+	ToolMenus->RegisterMenu("LevelEditor.MainMenu.Tools", "MainFrame.MainMenu.Tools", EMultiBoxType::Menu, bWarnIfAlreadyRegistered);
 
 	// Add other top level menus
 	Local::ExtendMenuBar();
@@ -339,7 +340,7 @@ TSharedRef< SWidget > FLevelEditorMenu::MakeLevelEditorMenu( const TSharedPtr<FU
 void FLevelEditorMenu::RegisterBuildMenu()
 {
 	static const FName BaseMenuName = "LevelEditor.MainMenu.Build";
-	UToolMenu* Menu = UToolMenus::Get()->RegisterMenu(BaseMenuName);
+	UToolMenu* Menu = UToolMenus::Get()->RegisterMenu(BaseMenuName, NAME_None, EMultiBoxType::Menu, /*bWarnIfAlreadyRegistered*/false);
 
 	struct FLightingMenus
 	{
@@ -356,7 +357,7 @@ void FLevelEditorMenu::RegisterBuildMenu()
 		/** Generates a lighting quality sub-menu */
 		static void RegisterLightingQualityMenu(const FName InBaseMenuName)
 		{
-			UToolMenu* SubMenu = UToolMenus::Get()->RegisterMenu(UToolMenus::JoinMenuPaths(InBaseMenuName, "LightingQuality"));
+			UToolMenu* SubMenu = UToolMenus::Get()->RegisterMenu(UToolMenus::JoinMenuPaths(InBaseMenuName, "LightingQuality"), NAME_None, EMultiBoxType::Menu, /*bWarnIfAlreadyRegistered*/false);
 
 			{
 				FToolMenuSection& Section = SubMenu->AddSection("LevelEditorBuildLightingQuality", LOCTEXT("LightingQualityHeading", "Quality Level"));
@@ -370,7 +371,7 @@ void FLevelEditorMenu::RegisterBuildMenu()
 		/** Generates a lighting density sub-menu */
 		static void RegisterLightingDensityMenu(const FName InBaseMenuName)
 		{
-			UToolMenu* SubMenu = UToolMenus::Get()->RegisterMenu(UToolMenus::JoinMenuPaths(InBaseMenuName, "LightingDensity"));
+			UToolMenu* SubMenu = UToolMenus::Get()->RegisterMenu(UToolMenus::JoinMenuPaths(InBaseMenuName, "LightingDensity"), NAME_None, EMultiBoxType::Menu, /*bWarnIfAlreadyRegistered*/false);
 
 			{
 				FToolMenuSection& Section = SubMenu->AddSection("LevelEditorBuildLightingDensity", LOCTEXT("LightingDensityHeading", "Density Rendering"));
@@ -436,7 +437,7 @@ void FLevelEditorMenu::RegisterBuildMenu()
 		/** Generates a lighting resolution sub-menu */
 		static void RegisterLightingResolutionMenu(const FName InBaseMenuName)
 		{
-			UToolMenu* SubMenu = UToolMenus::Get()->RegisterMenu(UToolMenus::JoinMenuPaths(InBaseMenuName, "LightingResolution"));
+			UToolMenu* SubMenu = UToolMenus::Get()->RegisterMenu(UToolMenus::JoinMenuPaths(InBaseMenuName, "LightingResolution"), NAME_None, EMultiBoxType::Menu, /*bWarnIfAlreadyRegistered*/false);
 
 			{
 				FToolMenuSection& Section = SubMenu->AddSection("LevelEditorBuildLightingResolution1", LOCTEXT("LightingResolutionHeading1", "Primitive Types"));
@@ -544,7 +545,7 @@ void FLevelEditorMenu::RegisterBuildMenu()
 			RegisterLightingDensityMenu(UToolMenus::JoinMenuPaths(InBaseMenuName, "LightingInfo"));
 			RegisterLightingResolutionMenu(UToolMenus::JoinMenuPaths(InBaseMenuName, "LightingInfo"));
 
-			UToolMenu* SubMenu = UToolMenus::Get()->RegisterMenu(UToolMenus::JoinMenuPaths(InBaseMenuName, "LightingInfo"));
+			UToolMenu* SubMenu = UToolMenus::Get()->RegisterMenu(UToolMenus::JoinMenuPaths(InBaseMenuName, "LightingInfo"), NAME_None, EMultiBoxType::Menu, /*bWarnIfAlreadyRegistered*/false);
 
 			{
 				FToolMenuSection& Section = SubMenu->AddSection("LevelEditorBuildLightingInfo", LOCTEXT("LightingInfoHeading", "Lighting Info Dialogs"));
@@ -746,7 +747,7 @@ void FLevelEditorMenu::RegisterBuildMenu()
 void FLevelEditorMenu::RegisterSelectMenu()
 {
 	static const FName BaseMenuName = "LevelEditor.MainMenu.Select";
-	UToolMenu* Menu = UToolMenus::Get()->RegisterMenu(BaseMenuName);
+	UToolMenu* Menu = UToolMenus::Get()->RegisterMenu(BaseMenuName, NAME_None, EMultiBoxType::Menu, /*bWarnIfAlreadyRegistered*/false);
 
 	{
 		FToolMenuSection& Section = Menu->AddSection("SelectActorGeneral", NSLOCTEXT("LevelViewportContextMenu", "SelectAnyHeading", "General"));
diff --git a/Engine/Source/Editor/LevelEditor/Private/LevelEditorToolBar.cpp b/Engine/Source/Editor/LevelEditor/Private/LevelEditorToolBar.cpp
--- a/Engine/Source/Editor/LevelEditor/Private/LevelEditorToolBar.cpp
+++ b/Engine/Source/Editor/LevelEditor/Private/LevelEditorToolBar.cpp
@@ -1281,16 +1281,22 @@ const FSlateBrush* FLevelEditorToolBar::GetActiveModeIcon(TWeakPtr<SLevelEditor>
 
 void FLevelEditorToolBar::RegisterLevelEditorToolBar( const TSharedRef<FUICommandList>& InCommandList, const TSharedRef<SLevelEditor> InLevelEditor)
 {
-	RegisterSourceControlMenu();
-	RegisterCinematicsMenu();
+	static bool bHasRegistered = false;
+	if (!bHasRegistered)
+	{
+		bHasRegistered = true;
+
+		RegisterSourceControlMenu();
+		RegisterCinematicsMenu();
 
-	RegisterQuickSettingsMenu();
-	RegisterOpenBlueprintMenu();
-	RegisterAddMenu();
+		RegisterQuickSettingsMenu();
+		RegisterOpenBlueprintMenu();
+		RegisterAddMenu();
+	}
 
 #define LOCTEXT_NAMESPACE "LevelEditorToolBar"
 
-	UToolMenu* ModesToolbar = UToolMenus::Get()->RegisterMenu("LevelEditor.LevelEditorToolBar.ModesToolBar", NAME_None, EMultiBoxType::SlimHorizontalToolBar);
+	UToolMenu* ModesToolbar = UToolMenus::Get()->RegisterMenu("LevelEditor.LevelEditorToolBar.ModesToolBar", NAME_None, EMultiBoxType::SlimHorizontalToolBar, /*warn*/false);
 	ModesToolbar->StyleName = "AssetEditorToolbar";
 	{
 		{
@@ -1438,7 +1444,7 @@ void FLevelEditorToolBar::RegisterLevelEditorToolBar( const TSharedRef<FUIComman
 			}));
 	}
 
-	UToolMenu* AssetsToolBar = UToolMenus::Get()->RegisterMenu("LevelEditor.LevelEditorToolBar.AssetsToolBar", NAME_None, EMultiBoxType::SlimHorizontalToolBar);
+	UToolMenu* AssetsToolBar = UToolMenus::Get()->RegisterMenu("LevelEditor.LevelEditorToolBar.AssetsToolBar", NAME_None, EMultiBoxType::SlimHorizontalToolBar, false);
 	AssetsToolBar->StyleName = "AssetEditorToolbar";
 	{
 		{
@@ -1480,7 +1486,7 @@ void FLevelEditorToolBar::RegisterLevelEditorToolBar( const TSharedRef<FUIComman
 
 	}
 
-	UToolMenu* PlayToolBar = UToolMenus::Get()->RegisterMenu("LevelEditor.LevelEditorToolBar.PlayToolBar", NAME_None, EMultiBoxType::SlimHorizontalToolBar);
+	UToolMenu* PlayToolBar = UToolMenus::Get()->RegisterMenu("LevelEditor.LevelEditorToolBar.PlayToolBar", NAME_None, EMultiBoxType::SlimHorizontalToolBar, false);
 	PlayToolBar->StyleName = "AssetEditorToolbar";
 	{
 		FToolMenuSection& PlaySection = PlayToolBar->AddSection("Play");
@@ -1495,10 +1501,10 @@ void FLevelEditorToolBar::RegisterLevelEditorToolBar( const TSharedRef<FUIComman
 
 	}
 
-	UToolMenu* UserToolbar = UToolMenus::Get()->RegisterMenu("LevelEditor.LevelEditorToolBar.User", NAME_None, EMultiBoxType::SlimHorizontalToolBar);
+	UToolMenu* UserToolbar = UToolMenus::Get()->RegisterMenu("LevelEditor.LevelEditorToolBar.User", NAME_None, EMultiBoxType::SlimHorizontalToolBar, false);
 	UserToolbar->StyleName = "AssetEditorToolbar";
 
-	UToolMenu* SettingsToolbar = UToolMenus::Get()->RegisterMenu("LevelEditor.LevelEditorToolBar.SettingsToolbar", NAME_None, EMultiBoxType::SlimHorizontalToolBar);
+	UToolMenu* SettingsToolbar = UToolMenus::Get()->RegisterMenu("LevelEditor.LevelEditorToolBar.SettingsToolbar", NAME_None, EMultiBoxType::SlimHorizontalToolBar, false);
 	SettingsToolbar->StyleName = "AssetEditorToolbar";
 	{
 		FToolMenuSection& SettingsSection = SettingsToolbar->AddSection("ProjectSettings");
@@ -1545,7 +1551,7 @@ TSharedRef< SWidget > FLevelEditorToolBar::MakeLevelEditorSecondaryModeToolbar(
 
 void FLevelEditorToolBar::RegisterLevelEditorSecondaryModeToolbar()
 {
-	UToolMenu* ModesToolbar = UToolMenus::Get()->RegisterMenu(SecondaryModeToolbarName, NAME_None, EMultiBoxType::SlimHorizontalToolBar);
+	UToolMenu* ModesToolbar = UToolMenus::Get()->RegisterMenu(SecondaryModeToolbarName, NAME_None, EMultiBoxType::SlimHorizontalToolBar, false);
 	ModesToolbar->StyleName = "SecondaryToolbar";
 }
 
@@ -1705,6 +1711,10 @@ TSharedRef< SWidget > FLevelEditorToolBar::GenerateQuickSettingsMenu(TSharedRef<
 void FLevelEditorToolBar::RegisterQuickSettingsMenu()
 {
 #define LOCTEXT_NAMESPACE "LevelToolBarViewMenu"
+	if (UToolMenus::Get()->IsMenuRegistered("LevelEditor.LevelEditorToolBar.LevelToolbarQuickSettings"))
+	{
+		return;
+	}
 
 	UToolMenu* Menu = UToolMenus::Get()->RegisterMenu("LevelEditor.LevelEditorToolBar.LevelToolbarQuickSettings");
 
@@ -1864,6 +1874,11 @@ TSharedRef< SWidget > FLevelEditorToolBar::GenerateOpenBlueprintMenuContent( TSh
 void FLevelEditorToolBar::RegisterOpenBlueprintMenu()
 {
 #define LOCTEXT_NAMESPACE "LevelToolBarViewMenu"
+	if (UToolMenus::Get()->IsMenuRegistered("LevelEditor.LevelEditorToolBar.OpenBlueprint"))
+	{
+		return;
+	}
+
 	UToolMenu* Menu = UToolMenus::Get()->RegisterMenu("LevelEditor.LevelEditorToolBar.OpenBlueprint");
 
 	struct FBlueprintMenus
@@ -2009,6 +2024,11 @@ void FLevelEditorToolBar::RegisterOpenBlueprintMenu()
 void FLevelEditorToolBar::RegisterAddMenu()
 {
 #define LOCTEXT_NAMESPACE "LevelToolBarViewMenu"
+	if (UToolMenus::Get()->IsMenuRegistered("LevelEditor.LevelEditorToolBar.AddQuickMenu"))
+	{
+		return;
+	}
+
 	UToolMenu* AddMenu = UToolMenus::Get()->RegisterMenu("LevelEditor.LevelEditorToolBar.AddQuickMenu");
 	{
 		FToolMenuSection& Section = AddMenu->FindOrAddSection("Content");
@@ -2063,6 +2083,10 @@ TSharedRef< SWidget > FLevelEditorToolBar::GenerateAddMenuWidget(TSharedRef<FUIC
 void FLevelEditorToolBar::RegisterCinematicsMenu()
 {
 #define LOCTEXT_NAMESPACE "LevelToolBarCinematicsMenu"
+	if (UToolMenus::Get()->IsMenuRegistered("LevelEditor.LevelEditorToolBar.Cinematics"))
+	{
+		return;
+	}
 
 	UToolMenu* Menu = UToolMenus::Get()->RegisterMenu("LevelEditor.LevelEditorToolBar.Cinematics");
 	Menu->bShouldCloseWindowAfterMenuSelection = true;
diff --git a/Engine/Source/Editor/LevelEditor/Private/SLevelEditor.cpp b/Engine/Source/Editor/LevelEditor/Private/SLevelEditor.cpp
--- a/Engine/Source/Editor/LevelEditor/Private/SLevelEditor.cpp
+++ b/Engine/Source/Editor/LevelEditor/Private/SLevelEditor.cpp
@@ -2033,6 +2033,7 @@ void SLevelEditor::OnLevelActorOuterChanged(AActor* InActor, UObject* InOldOuter
 void SLevelEditor::RegisterStatusBarTools()
 {
 #if WITH_LIVE_CODING
+	if (!UToolMenus::Get()->IsMenuRegistered("LevelEditor.StatusBar.ToolBar.CompileComboButton"))
 	{
 		UToolMenu* Menu = UToolMenus::Get()->RegisterMenu("LevelEditor.StatusBar.ToolBar.CompileComboButton");
 		{
