Tentative fix for ensure when exiting pie and fix to error message.

diff --git a/Engine/Source/Editor/LevelEditor/Private/SLevelViewport.cpp b/Engine/Source/Editor/LevelEditor/Private/SLevelViewport.cpp
--- a/Engine/Source/Editor/LevelEditor/Private/SLevelViewport.cpp
+++ b/Engine/Source/Editor/LevelEditor/Private/SLevelViewport.cpp
@@ -136,7 +136,7 @@ SLevelViewport::SLevelViewport()
 	, ViewTransitionType( EViewTransition::None )
 	, bViewTransitionAnimPending( false )
 	, DeviceProfile("Default")
-	, PIEOverlaySlotIndex(0)
+	, PIEOverlaySlotIndex(INDEX_NONE)
 	, bPIEHasFocus(false)
 	, bPIEContainsFocus(false)
 	, UserAllowThrottlingValue(0)
@@ -1112,7 +1112,7 @@ void SLevelViewport::Tick( const FGeometry& AllottedGeometry, const double InCur
 	{
 		if(ViewTransitionType == EViewTransition::StartingPlayInEditor)
 		{
-			if(PIEOverlaySlotIndex)
+			if(PIEOverlaySlotIndex != INDEX_NONE)
 			{
 				PIEOverlayAnim = FCurveSequence(0.0f, SLevelViewportPIEAnimation::MouseControlLabelFadeout, ECurveEaseFunction::CubicInOut);
 				PIEOverlayAnim.Play(this->AsShared());
@@ -4208,8 +4208,11 @@ void SLevelViewport::ShowMouseCaptureLabel(ELabelAnchorMode AnchorMode)
 
 void SLevelViewport::HideMouseCaptureLabel()
 {
-	ViewportOverlay->RemoveSlot(PIEOverlaySlotIndex);
-	PIEOverlaySlotIndex = 0;
+	if (PIEOverlaySlotIndex != INDEX_NONE)
+	{
+		ViewportOverlay->RemoveSlot(PIEOverlaySlotIndex);
+		PIEOverlaySlotIndex = INDEX_NONE;
+	}
 }
 
 void SLevelViewport::ResetNewLevelViewFlags()
@@ -4276,10 +4279,7 @@ void SLevelViewport::EndPlayInEditorSession()
 	// No longer need to store the content 
 	InactiveViewportWidgetEditorContent.Reset();
 
-	if(PIEOverlaySlotIndex)
-	{
-		HideMouseCaptureLabel();
-	}
+	HideMouseCaptureLabel();
 
 	// Kick off a quick transition effect (border graphics)
 	ViewTransitionType = EViewTransition::ReturningToEditor;
@@ -4307,10 +4307,7 @@ void SLevelViewport::SwapViewportsForSimulateInEditor()
 	check( IsPlayInEditorViewportActive() );
 	
 	// Remove the mouse control label - not relevant for SIE
-	if(PIEOverlaySlotIndex)
-	{
-		HideMouseCaptureLabel();
-	}
+	HideMouseCaptureLabel();
 
 	// Unregister the game viewport with slate which will release mouse capture and lock
 	FSlateApplication::Get().UnregisterGameViewport();
diff --git a/Engine/Source/Runtime/SlateCore/Private/Widgets/SOverlay.cpp b/Engine/Source/Runtime/SlateCore/Private/Widgets/SOverlay.cpp
--- a/Engine/Source/Runtime/SlateCore/Private/Widgets/SOverlay.cpp
+++ b/Engine/Source/Runtime/SlateCore/Private/Widgets/SOverlay.cpp
@@ -271,7 +271,7 @@ void SOverlay::RemoveSlot( int32 ZOrder )
 			}
 		}
 
-		ensureMsgf(false, TEXT("Could not remove slot. There are no children with ZOrder %d."));
+		ensureMsgf(false, TEXT("Could not remove slot. There are no children with ZOrder %d."), ZOrder);
 	}
 	else if (Children.Num() > 0)
 	{
