Fix a bug causing "Edit Asset" to stop working from the Level Editor context menu

diff --git a/Engine/Source/Editor/LevelEditor/Private/SLevelEditor.cpp b/Engine/Source/Editor/LevelEditor/Private/SLevelEditor.cpp
--- a/Engine/Source/Editor/LevelEditor/Private/SLevelEditor.cpp
+++ b/Engine/Source/Editor/LevelEditor/Private/SLevelEditor.cpp
@@ -107,8 +107,8 @@ void SLevelEditor::BindCommands()
 	LevelEditorCommands->Append( FPlayWorldCommands::GlobalPlayWorldActions.ToSharedRef() );
 
 	LevelEditorCommands->MapAction( 
-		Actions.EditAssetNoConfirmMultiple, 
-		FExecuteAction::CreateStatic( &FLevelEditorActionCallbacks::EditAsset_Clicked, EToolkitMode::Standalone, TWeakPtr< SLevelEditor >( SharedThis( this ) ), false ),
+		Actions.EditAsset, 
+		FExecuteAction::CreateStatic( &FLevelEditorActionCallbacks::EditAsset_Clicked, EToolkitMode::Standalone, TWeakPtr< SLevelEditor >( SharedThis( this ) ), true ),
 		FCanExecuteAction::CreateStatic( &FLevelEditorActionCallbacks::EditAsset_CanExecute ) );
 
 	LevelEditorCommands->MapAction( 
