Add an option to the Outliner to bulk edit components in the Property Matrix

diff --git a/Engine/Plugins/Editor/ContentBrowser/ContentBrowserAssetDataSource/Source/ContentBrowserAssetDataSource/Private/AssetFileContextMenu.cpp b/Engine/Plugins/Editor/ContentBrowser/ContentBrowserAssetDataSource/Source/ContentBrowserAssetDataSource/Private/AssetFileContextMenu.cpp
--- a/Engine/Plugins/Editor/ContentBrowser/ContentBrowserAssetDataSource/Source/ContentBrowserAssetDataSource/Private/AssetFileContextMenu.cpp
+++ b/Engine/Plugins/Editor/ContentBrowser/ContentBrowserAssetDataSource/Source/ContentBrowserAssetDataSource/Private/AssetFileContextMenu.cpp
@@ -473,7 +473,7 @@ void FAssetFileContextMenu::MakeAssetActionsSubMenu(UToolMenu* Menu)
 
 			Section.AddMenuEntry(
 				"PropertyMatrix",
-				LOCTEXT("PropertyMatrix", "Open Selection in Property Matrix"),
+				LOCTEXT("PropertyMatrix", "Edit Selection in Property Matrix"),
 				DynamicTooltipAttribute,
 				FSlateIcon(),
 				FUIAction(
diff --git a/Engine/Source/Editor/LevelEditor/Private/LevelEditorActions.cpp b/Engine/Source/Editor/LevelEditor/Private/LevelEditorActions.cpp
--- a/Engine/Source/Editor/LevelEditor/Private/LevelEditorActions.cpp
+++ b/Engine/Source/Editor/LevelEditor/Private/LevelEditorActions.cpp
@@ -3666,7 +3666,7 @@ void FLevelEditorCommands::RegisterCommands()
 
 	UI_COMMAND( EditAsset, "Edit Asset", "Edits the asset associated with the selected actor", EUserInterfaceActionType::Button, FInputChord( EKeys::E, EModifierKey::Control ) );
 	UI_COMMAND( EditAssetNoConfirmMultiple, "Edit Multiple Assets", "Edits multiple assets associated with the selected actor without a confirmation prompt", EUserInterfaceActionType::Button, FInputChord( EKeys::E, EModifierKey::Control | EModifierKey::Shift ) );
-	UI_COMMAND( OpenSelectionInPropertyMatrix, "Open Selection in Property Matrix", "Bulk edit the selected assets in the Property Matrix", EUserInterfaceActionType::Button, FInputChord() );
+	UI_COMMAND( OpenSelectionInPropertyMatrix, "Edit Selection in Property Matrix", "Bulk edit the selected assets in the Property Matrix", EUserInterfaceActionType::Button, FInputChord() );
 
 	UI_COMMAND( GoHere, "Go Here", "Moves the camera to the current mouse position", EUserInterfaceActionType::Button, FInputChord() );
 
diff --git a/Engine/Source/Editor/LevelEditor/Private/LevelEditorContextMenu.cpp b/Engine/Source/Editor/LevelEditor/Private/LevelEditorContextMenu.cpp
--- a/Engine/Source/Editor/LevelEditor/Private/LevelEditorContextMenu.cpp
+++ b/Engine/Source/Editor/LevelEditor/Private/LevelEditorContextMenu.cpp
@@ -52,6 +52,7 @@
 #include "IMergeActorsTool.h"
 #include "SLevelEditor.h"
 #include "SLevelViewport.h"
+#include "PropertyEditorModule.h"
 
 #define LOCTEXT_NAMESPACE "LevelViewportContextMenu"
 
@@ -122,6 +123,8 @@ public:
 	static void FillEditMenu(UToolMenu* Menu) { FillEditMenu(Menu, nullptr); }
 	static void FillEditMenu(UToolMenu* Menu, FToolMenuSection* InSection);
 
+	static void FillBulkEditComponentsMenu(UToolMenu* Menu);
+
 	/**
 	 * Fills in the menu options for the Asset Tools submenu
 	 *
@@ -270,6 +273,8 @@ void FLevelEditorContextMenu::RegisterActorContextMenu()
 		FSelectedActorInfo& SelectionInfo = FLevelEditorContextMenuImpl::SelectionInfo;
 		SelectionInfo = AssetSelectionUtils::BuildSelectedActorInfo(SelectedActors);
 
+		int32 NumSelectedActors = SelectedActors.Num();
+
 		{
 			// General actions that apply to most actors and their underlying assets
 			// In most cases, you DO NOT want to extend this section; look at ActorUETools or ActorTypeTools below
@@ -322,13 +327,6 @@ void FLevelEditorContextMenu::RegisterActorContextMenu()
 				}
 			}
 
-			Section.AddMenuEntry(
-				FLevelEditorCommands::Get().OpenSelectionInPropertyMatrix,
-				TAttribute<FText>(), // use command's label
-				TAttribute<FText>(), // use command's tooltip
-				FSlateIcon(FAppStyle::GetAppStyleSetName(), "DetailsView.EditRawProperties")
-				);
-
 			if (LevelEditorContext->ContextType == ELevelEditorMenuContext::MainMenu)
 			{
 				Section.AddSubMenu(
@@ -361,6 +359,33 @@ void FLevelEditorContextMenu::RegisterActorContextMenu()
 
 			LevelEditorCreateActorMenu::FillAddReplaceContextMenuSections(Section, LevelEditorContext);
 		}
+		
+		{	// Options that relate to bulk editing assets
+			// In most cases, you DO NOT want to extend this section; look at ActorUETools or ActorTypeTools below
+			FToolMenuSection& Section = InMenu->AddSection("ActorBulkEdit", LOCTEXT("ActorBulkEditHeading", "Bulk Editing"));
+			
+			Section.AddMenuEntry(
+				FLevelEditorCommands::Get().OpenSelectionInPropertyMatrix,
+				TAttribute<FText>(), // use command's label
+				TAttribute<FText>(), // use command's tooltip
+				FSlateIcon(FAppStyle::GetAppStyleSetName(), "DetailsView.EditRawProperties")
+			);
+
+			Section.AddSubMenu(
+				"BulkEditComponentsSubmenu",
+				LOCTEXT("BulkEditComponentsSubmenuName", "Edit Components in the Property Matrix"),
+				LOCTEXT("BulkEditComponentsSubmenuTooltip", "Bulk Edit any editable components that are common between the selected actors in the Property Matrix"),
+				FNewToolMenuDelegate::CreateStatic(&FLevelEditorContextMenuImpl::FillBulkEditComponentsMenu),
+				FUIAction(FExecuteAction(), FCanExecuteAction(), FIsActionChecked(), FIsActionButtonVisible::CreateLambda([NumSelectedActors]()
+				{
+					return NumSelectedActors >= 2;
+				})),
+				EUserInterfaceActionType::Button,
+				false, // default value
+				FSlateIcon(FAppStyle::GetAppStyleSetName(), "DetailsView.EditRawProperties")
+			);
+		}
+
 
 		if (LevelEditorContext->ContextType == ELevelEditorMenuContext::Viewport || LevelEditorContext->ContextType == ELevelEditorMenuContext::SceneOutliner)
 		{
@@ -1389,6 +1414,138 @@ void FLevelEditorContextMenuImpl::FillEditMenu( UToolMenu* Menu, FToolMenuSectio
 	Section.AddMenuEntry( FGenericCommands::Get().Rename );
 }
 
+void FLevelEditorContextMenuImpl::FillBulkEditComponentsMenu( UToolMenu* Menu )
+{
+	/** This is how the logic for selecting components to bulk edit works at a high level:
+	 *  We try to find the common types of components that are shared between all the actors in the selection, and then
+	 *  add a menu entry for each of the common component types which opens them in the Property Matrix
+	 *  To achieve this, we pick the components of the first actor in the selection as a base, and then remove any
+	 *  that are not shared with every subsequent actor in the selection, early exiting if there are no common components.
+	 */
+	
+	FToolMenuSection& Section = Menu->AddSection("Section");
+
+	ULevelEditorContextMenuContext* LevelEditorContext = Menu->FindContext<ULevelEditorContextMenuContext>();
+	
+	if (!LevelEditorContext)
+	{
+		return;
+	}
+
+	TObjectPtr<const UTypedElementSelectionSet> Selection = LevelEditorContext->CurrentSelection;
+	TArray<AActor*> SelectedActors = Selection->GetSelectedObjects<AActor>();
+
+	/* Map to keep track of components that are common between the current selection
+	 * The key represents the type of the component, while the value contains the actual components from each object in the selection
+	 */
+	TMap<TSubclassOf<UActorComponent>, TArray< UObject* >> CommonComponents;
+
+	// We start by filling in CommonComponents with all the components of the first actor in the selection
+	const AActor* FirstActor = SelectedActors[0];
+	for( UActorComponent* Component : FirstActor->GetComponents())
+	{
+		// Make sure the component can be edited
+		if(FComponentEditorUtils::CanEditComponentInstance(Component, Cast<USceneComponent>(Component), false))
+		{
+			// If this type of component already exists, simply add it to the array
+			if(TArray< UObject* >* FoundCommonComponent = CommonComponents.Find(Component->GetClass()))
+			{
+				FoundCommonComponent->Add(Component);
+			}
+			// Otherwise add a new entry for this component type
+			else
+			{
+				TArray< UObject* > CurrentActorComponents;
+				CurrentActorComponents.Add(Component);
+				CommonComponents.Add(Component->GetClass(), CurrentActorComponents);
+			}
+			
+		}
+	}
+
+	// We start iterating from the second object since we've already used the first as a base
+	TArray<AActor*>::TIterator ActorIt = SelectedActors.CreateIterator();
+	++ActorIt;
+	
+	for(; ActorIt; ++ActorIt)
+	{
+		// Iterate through each common component
+		for(TMap<TSubclassOf<UActorComponent>, TArray< UObject* >>::TIterator ComponentIt = CommonComponents.CreateIterator(); ComponentIt; ++ComponentIt)
+		{
+			// Get all components of the current component type from this actor
+			TArray<UActorComponent*> ComponentsOfType;
+			(*ActorIt)->GetComponents(ComponentIt.Key(), ComponentsOfType);
+
+			bool bHasComponentOfCurrentType = false;
+
+			/* For each component the current actor has, we make sure it is EXACTLY the same type as the common component
+			 * because GetComponents also checks SubClasses. i.e if CommonComponent is SceneComponent and CurrentComponent
+			 * is StaticMeshComponent, we do not want to multi edit them currently
+			 */
+			for(UActorComponent* CurrentComponent : ComponentsOfType)
+			{
+				if(ComponentIt.Key() == CurrentComponent->GetClass() && FComponentEditorUtils::CanEditComponentInstance(CurrentComponent, Cast<USceneComponent>(CurrentComponent), false))
+				{
+					ComponentIt.Value().Add(CurrentComponent);
+					bHasComponentOfCurrentType = true;
+				}
+			}
+
+			// If this actor has no components of this type, discard this type
+			if(!bHasComponentOfCurrentType)
+			{
+				ComponentIt.RemoveCurrent();
+			}
+		}
+
+		// Early out if there are no common components between all actors in the selection
+		if(CommonComponents.IsEmpty())
+		{
+			break;
+		}
+		
+	}
+
+	// If we have no common components, show some text indicating so
+	if(CommonComponents.IsEmpty())
+	{
+		Section.AddEntry(FToolMenuEntry::InitWidget(
+		NAME_None,
+		SNew(SBox)
+		.Padding(FMargin(6.f, 0.f))
+		[
+			SNew(STextBlock)
+			.Text(LOCTEXT("NoCommonComponents", "The selected actors have no common editable components."))
+			.TextStyle(FAppStyle::Get(), "HintText")
+		],
+		FText::GetEmpty(), /*bNoIndent*/ true, /*bSearchable*/ false, false,
+		LOCTEXT("NoCommonComponentsTooltip", "There are no editable components of the same type between the currently selected actors.")));
+		
+		return;
+	}
+
+	// Add a menu entry to open each set of components in the property matrix
+	for(auto it = CommonComponents.CreateIterator(); it; ++it)
+	{
+		TArray<UObject*>& Components = it.Value();
+		
+		FUIAction BulkEditAction(
+		FExecuteAction::CreateLambda([Components]()
+		{
+			FPropertyEditorModule& PropertyEditorModule = FModuleManager::LoadModuleChecked<FPropertyEditorModule>( "PropertyEditor" );
+			PropertyEditorModule.CreatePropertyEditorToolkit(TSharedPtr<IToolkitHost>(), Components );
+		}));
+		
+		Section.AddMenuEntry(
+			NAME_None,
+			FText::FromName(it.Key()->GetFName()),
+			FText::GetEmpty(),
+			FSlateIcon(),
+			BulkEditAction);
+	}
+}
+
+
 
 void FLevelEditorContextMenuImpl::FillAssetToolsMenu(UToolMenu* Menu)
 {
diff --git a/Engine/Source/Editor/PropertyEditor/Private/SDetailsView.cpp b/Engine/Source/Editor/PropertyEditor/Private/SDetailsView.cpp
--- a/Engine/Source/Editor/PropertyEditor/Private/SDetailsView.cpp
+++ b/Engine/Source/Editor/PropertyEditor/Private/SDetailsView.cpp
@@ -290,7 +290,7 @@ void SDetailsView::Construct(const FArguments& InArgs, const FDetailsViewArgs& I
 				.OnClicked(this, &SDetailsView::OnOpenRawPropertyEditorClicked)
 				.IsEnabled(this, &SDetailsView::CanOpenRawPropertyEditor)
 				.Visibility(this, &SDetailsView::CanShowRawPropertyEditorButton, DetailsViewArgs.bShowPropertyMatrixButton)
-				.ToolTipText(LOCTEXT("RawPropertyEditorButtonLabel", "Open Selection in Property Matrix"))
+				.ToolTipText(LOCTEXT("RawPropertyEditorButtonLabel", "Edit Selection in Property Matrix"))
 				[
 					SNew(SImage)
 					.ColorAndOpacity(FSlateColor::UseForeground())
