ActorEditorContext:
- Allow switching Current level from context in WP World
- Unify display of Current Level when editing Level Instance (WP and non-WP)
- Level Instance Mode will lock/unlock selection and render postprocess based on current level being the Level Instance level or not.

#rnx

diff --git a/Engine/Source/Editor/LevelEditor/Private/LevelEditorSubsystem.cpp b/Engine/Source/Editor/LevelEditor/Private/LevelEditorSubsystem.cpp
--- a/Engine/Source/Editor/LevelEditor/Private/LevelEditorSubsystem.cpp
+++ b/Engine/Source/Editor/LevelEditor/Private/LevelEditorSubsystem.cpp
@@ -56,6 +56,33 @@ TSharedPtr<SLevelViewport> GetLevelViewport(const FName& ViewportConfigKey)
 	return nullptr;
 }
 
+bool IsEditingLevelInstanceCurrentLevel(UWorld* InWorld)
+{
+	if (InWorld)
+	{
+		if (ULevelInstanceSubsystem* LevelInstanceSubsystem = InWorld->GetSubsystem<ULevelInstanceSubsystem>())
+		{
+			if (ILevelInstanceInterface* LevelInstance = LevelInstanceSubsystem->GetEditingLevelInstance())
+			{
+				return LevelInstanceSubsystem->GetLevelInstanceLevel(LevelInstance) == InWorld->GetCurrentLevel();
+			}
+		}
+	}
+
+	return false;
+}
+
+AActor* GetEditingLevelInstance(UWorld* InWorld)
+{
+	ULevelInstanceSubsystem* LevelInstanceSubsystem = InWorld ? InWorld->GetSubsystem<ULevelInstanceSubsystem>() : nullptr;
+	return LevelInstanceSubsystem ? Cast<AActor>(LevelInstanceSubsystem->GetEditingLevelInstance()) : nullptr;
+}
+
+bool IsActorEditorContextVisible(UWorld* InWorld)
+{
+	return InWorld && (InWorld->GetCurrentLevel()->OwningWorld->GetLevels().Num() > 1 && (!InWorld->IsPartitionedWorld() || GetEditingLevelInstance(InWorld) != nullptr));
+}
+
 }
 
 void ULevelEditorSubsystem::Initialize(FSubsystemCollectionBase& Collection)
@@ -750,11 +777,16 @@ public:
 		World = (InArgs._World);
 		GEditor->GetEditorWorldContext().AddRef(World);
 		CommandList = MakeShareable(new FUICommandList);
-
-		// No option to change current level for partitioned worlds
-		if (World && World->IsPartitionedWorld())
-		{
-			ChildSlot
+				
+		ChildSlot
+		[
+			SNew(SComboButton)
+			.Cursor(EMouseCursor::Default)
+			.VAlign(VAlign_Center)
+			.ComboButtonStyle(FAppStyle::Get(), "SimpleComboButton")
+			.Visibility(this, &SCurrentLevelWidget::GetCurrentLevelButtonVisibility)
+			.OnGetMenuContent(this, &SCurrentLevelWidget::GenerateLevelMenu)
+			.ButtonContent()
 			[
 				// Current Level
 				SNew(SVerticalBox)
@@ -771,7 +803,7 @@ public:
 				.AutoHeight()
 				[
 					SNew(SHorizontalBox)
-					.Visibility_Lambda([this]() { return GetEditingLevelInstance() ? EVisibility::Visible : EVisibility::Collapsed; })
+					.Visibility_Lambda([this]() { return IsEditingLevelInstanceCurrentLevel() ? EVisibility::Visible : EVisibility::Collapsed; })
 					+ SHorizontalBox::Slot()
 					.AutoWidth()
 					[
@@ -795,7 +827,7 @@ public:
 					.AutoWidth()
 					[
 						SNew(STextBlock)
-						.Text_Lambda([this]() 
+						.Text_Lambda([this]()
 						{
 							AActor* LevelInstance = GetEditingLevelInstance();
 							return LevelInstance ? FText::FromString(LevelInstance->GetActorLabel()) : FText::GetEmpty();
@@ -810,34 +842,19 @@ public:
 						.Text(LOCTEXT("FromEnd", ")"))
 					]
 				]
-			];
-		}
-		else
-		{
-			ChildSlot
-			[
-				SNew(SComboButton)
-				.Cursor(EMouseCursor::Default)
-				.VAlign(VAlign_Center)
-				.ComboButtonStyle(FAppStyle::Get(), "SimpleComboButton")
-				.Visibility(this, &SCurrentLevelWidget::GetCurrentLevelButtonVisibility)
-				.OnGetMenuContent(this, &SCurrentLevelWidget::GenerateLevelMenu)
-				.ButtonContent()
-				[
-					SNew(STextBlock)
-					.Visibility(this, &SCurrentLevelWidget::GetCurrentLevelTextVisibility)
-					.Text(this, &SCurrentLevelWidget::GetCurrentLevelText)
-					.ShadowOffset(FVector2D(1, 1))
-				]
-			];
-		}
+			]
+		];
 	}
 
 private:
 	AActor* GetEditingLevelInstance() const
 	{
-		ULevelInstanceSubsystem* LevelInstanceSubsystem = GetWorld()->GetSubsystem<ULevelInstanceSubsystem>();
-		return LevelInstanceSubsystem ? Cast<AActor>(LevelInstanceSubsystem->GetEditingLevelInstance()) : nullptr;
+		return InternalEditorLevelLibrary::GetEditingLevelInstance(GetWorld());
+	}
+
+	bool IsEditingLevelInstanceCurrentLevel() const
+	{
+		return InternalEditorLevelLibrary::IsEditingLevelInstanceCurrentLevel(GetWorld());
 	}
 
 	FText GetCurrentLevelText() const
@@ -859,7 +876,7 @@ private:
 	
 	bool IsVisible() const
 	{
-		return (GetWorld() && (GetWorld()->GetCurrentLevel()->OwningWorld->GetLevels().Num() > 1) && (!GetWorld()->IsPartitionedWorld() || (GetWorld()->GetCurrentLevel() != GetWorld()->PersistentLevel)));
+		return InternalEditorLevelLibrary::IsActorEditorContextVisible(GetWorld());
 	}
 
 	EVisibility GetCurrentLevelTextVisibility() const
@@ -897,8 +914,7 @@ private:
 
 bool ULevelEditorSubsystem::GetActorEditorContextDisplayInfo(UWorld* InWorld, FActorEditorContextClientDisplayInfo& OutDiplayInfo) const
 {
-	const bool bIsVisible = InWorld && (InWorld->GetCurrentLevel()->OwningWorld->GetLevels().Num() > 1) && (!InWorld->IsPartitionedWorld() || (InWorld->GetCurrentLevel() != InWorld->PersistentLevel));
-	if (bIsVisible)
+	if (InternalEditorLevelLibrary::IsActorEditorContextVisible(InWorld))
 	{
 		OutDiplayInfo.Title = TEXT("Level");
 		OutDiplayInfo.Brush = FSlateIconFinder::FindIconBrushForClass(UWorld::StaticClass());
diff --git a/Engine/Source/Editor/LevelInstanceEditor/Private/LevelInstanceEditorMode.cpp b/Engine/Source/Editor/LevelInstanceEditor/Private/LevelInstanceEditorMode.cpp
--- a/Engine/Source/Editor/LevelInstanceEditor/Private/LevelInstanceEditorMode.cpp
+++ b/Engine/Source/Editor/LevelInstanceEditor/Private/LevelInstanceEditorMode.cpp
@@ -48,7 +48,7 @@ void ULevelInstanceEditorMode::UpdateEngineShowFlags()
 		{
 			if(ULevelInstanceSubsystem* LevelInstanceSubsystem = LevelVC->GetWorld()->GetSubsystem<ULevelInstanceSubsystem>())
 			{
-				const bool bEditingLevelInstance = !!LevelInstanceSubsystem->GetEditingLevelInstance() && IsContextRestrictionCommandEnabled();
+				const bool bEditingLevelInstance = IsContextRestrictedForWorld(LevelVC->GetWorld());
 				// Make sure we update both Game/Editor showflags
 				LevelVC->EngineShowFlags.EditingLevelInstance = bEditingLevelInstance;
 				LevelVC->LastEngineShowFlags.EditingLevelInstance = bEditingLevelInstance;
@@ -82,6 +82,13 @@ void ULevelInstanceEditorMode::CreateToolkit()
 	Toolkit = MakeShared<FLevelInstanceEditorModeToolkit>();
 }
 
+void ULevelInstanceEditorMode::ModeTick(float DeltaTime)
+{
+	Super::ModeTick(DeltaTime);
+
+	UpdateEngineShowFlags();
+}
+
 bool ULevelInstanceEditorMode::IsCompatibleWith(FEditorModeID OtherModeID) const
 {
 	return (OtherModeID != FBuiltinEditorModes::EM_Foliage) && (OtherModeID != FBuiltinEditorModes::EM_Landscape);
@@ -126,38 +133,37 @@ void ULevelInstanceEditorMode::BindCommands()
 
 bool ULevelInstanceEditorMode::IsSelectionDisallowed(AActor* InActor, bool bInSelection) const
 {
-	const bool bRestrict = bContextRestriction && bInSelection;
+	UWorld* World = InActor->GetWorld();
+	const bool bRestrict = bInSelection && IsContextRestrictedForWorld(World);
 
 	if (bRestrict)
 	{
-		if (UWorld* World = InActor->GetWorld())
+		check(World);
+		if (ILevelInstanceInterface* LevelInstance = Cast<ILevelInstanceInterface>(InActor))
 		{
-			if (ILevelInstanceInterface* LevelInstance = Cast<ILevelInstanceInterface>(InActor))
+			if (LevelInstance->IsEditing())
 			{
-				if (LevelInstance->IsEditing())
-				{
-					return false;
-				}
+				return false;
 			}
+		}
 
-			if (ULevelInstanceSubsystem* LevelInstanceSubsystem = World->GetSubsystem<ULevelInstanceSubsystem>())
+		if (ULevelInstanceSubsystem* LevelInstanceSubsystem = World->GetSubsystem<ULevelInstanceSubsystem>())
+		{
+			ILevelInstanceInterface* EditingLevelInstance = LevelInstanceSubsystem->GetEditingLevelInstance();
+			ILevelInstanceInterface* LevelInstance = LevelInstanceSubsystem->GetParentLevelInstance(InActor);
+			// Allow selection on actors that are part of the currently edited Level Instance hierarchy because AActor::GetRootSelectionParent() will eventually
+			// Bubble up the selection to its parent.
+			while (LevelInstance != nullptr)
 			{
-				ILevelInstanceInterface* EditingLevelInstance = LevelInstanceSubsystem->GetEditingLevelInstance();
-				ILevelInstanceInterface* LevelInstance = LevelInstanceSubsystem->GetParentLevelInstance(InActor);
-				// Allow selection on actors that are part of the currently edited Level Instance hierarchy because AActor::GetRootSelectionParent() will eventually
-				// Bubble up the selection to its parent.
-				while (LevelInstance != nullptr)
+				if (LevelInstance == EditingLevelInstance)
 				{
-					if (LevelInstance == EditingLevelInstance)
-					{
-						return false;
-					}
-
-					LevelInstance = LevelInstanceSubsystem->GetParentLevelInstance(CastChecked<AActor>(LevelInstance));
+					return false;
 				}
 
-				return EditingLevelInstance != nullptr;
+				LevelInstance = LevelInstanceSubsystem->GetParentLevelInstance(CastChecked<AActor>(LevelInstance));
 			}
+
+			return EditingLevelInstance != nullptr;
 		}
 	}
 
@@ -190,4 +196,17 @@ bool ULevelInstanceEditorMode::IsContextRestrictionCommandEnabled() const
 	return bContextRestriction;
 }
 
+bool ULevelInstanceEditorMode::IsContextRestrictedForWorld(UWorld* InWorld) const
+{
+	if (ULevelInstanceSubsystem* LevelInstanceSubsystem = InWorld? InWorld->GetSubsystem<ULevelInstanceSubsystem>() : nullptr)
+	{
+		if (ILevelInstanceInterface* EditingLevelInstance = LevelInstanceSubsystem->GetEditingLevelInstance())
+		{
+			return bContextRestriction && LevelInstanceSubsystem->GetLevelInstanceLevel(EditingLevelInstance) == InWorld->GetCurrentLevel();
+		}
+	}
+
+	return false;
+}
+
 #undef LOCTEXT_NAMESPACE
diff --git a/Engine/Source/Editor/LevelInstanceEditor/Private/LevelInstanceEditorMode.h b/Engine/Source/Editor/LevelInstanceEditor/Private/LevelInstanceEditorMode.h
--- a/Engine/Source/Editor/LevelInstanceEditor/Private/LevelInstanceEditorMode.h
+++ b/Engine/Source/Editor/LevelInstanceEditor/Private/LevelInstanceEditorMode.h
@@ -23,6 +23,7 @@ public:
 	virtual void Enter() override;
 	virtual void Exit() override;
 	virtual void CreateToolkit() override;
+	virtual void ModeTick(float DeltaTime) override;
 	
 	virtual bool IsCompatibleWith(FEditorModeID OtherModeID) const override;
 
@@ -40,6 +41,7 @@ private:
 	void ExitModeCommand();
 	void ToggleContextRestrictionCommand();
 	bool IsContextRestrictionCommandEnabled() const;
+	bool IsContextRestrictedForWorld(UWorld* InWorld) const;
 
 	bool bContextRestriction;
 };
