Fix crash when interacting with a script actor containing an ISM when the level is locked

diff --git a/Engine/Source/Editor/LevelEditor/Private/Elements/SMInstance/SMInstanceElementLevelEditorSelectionCustomization.cpp b/Engine/Source/Editor/LevelEditor/Private/Elements/SMInstance/SMInstanceElementLevelEditorSelectionCustomization.cpp
--- a/Engine/Source/Editor/LevelEditor/Private/Elements/SMInstance/SMInstanceElementLevelEditorSelectionCustomization.cpp
+++ b/Engine/Source/Editor/LevelEditor/Private/Elements/SMInstance/SMInstanceElementLevelEditorSelectionCustomization.cpp
@@ -7,6 +7,10 @@
 #include "Elements/Framework/TypedElementList.h"
 #include "Elements/Framework/EngineElementsLibrary.h"
 
+#include "LevelUtils.h"
+#include "EditorModeManager.h"
+#include "Toolkits/IToolkitHost.h"
+
 DEFINE_LOG_CATEGORY_STATIC(LogSMInstanceLevelEditorSelection, Log, All);
 
 bool FSMInstanceElementLevelEditorSelectionCustomization::CanSelectElement(const TTypedElement<ITypedElementSelectionInterface>& InElementSelectionHandle, const FTypedElementSelectionOptions& InSelectionOptions)
@@ -16,6 +20,15 @@ bool FSMInstanceElementLevelEditorSelectionCustomization::CanSelectElement(const
 	{
 		return false;
 	}
+
+	AActor* Owner = SMInstance.GetISMComponent()->GetOwner();
+	AActor* SelectionRoot = Owner->GetRootSelectionParent();
+	ULevel* SelectionLevel = (SelectionRoot != nullptr) ? SelectionRoot->GetLevel() : Owner->GetLevel();
+
+	if (!Owner->IsTemplate() && FLevelUtils::IsLevelLocked(SelectionLevel))
+	{
+		return false;
+	}
 	
 	// Bail if global selection is locked
 	return !GEdSelectionLock;
@@ -29,6 +42,16 @@ bool FSMInstanceElementLevelEditorSelectionCustomization::CanDeselectElement(con
 		return false;
 	}
 
+	AActor* Owner = SMInstance.GetISMComponent()->GetOwner();
+	AActor* SelectionRoot = Owner->GetRootSelectionParent();
+	ULevel* SelectionLevel = (SelectionRoot != nullptr) ? SelectionRoot->GetLevel() : Owner->GetLevel();
+
+	if (const IToolkitHost* ToolkitHostPtr = GetToolkitHost())
+	{
+		// Allow active modes to determine whether the deselection is allowed
+		return ToolkitHostPtr->GetEditorModeManager().IsSelectionAllowed(Owner, /*bInSelected*/false);
+	}
+
 	// Bail if global selection is locked
 	return !GEdSelectionLock;
 }
diff --git a/Engine/Source/Editor/UnrealEd/Private/EditorSelectUtils.cpp b/Engine/Source/Editor/UnrealEd/Private/EditorSelectUtils.cpp
--- a/Engine/Source/Editor/UnrealEd/Private/EditorSelectUtils.cpp
+++ b/Engine/Source/Editor/UnrealEd/Private/EditorSelectUtils.cpp
@@ -515,6 +515,20 @@ void UUnrealEdEngine::SelectActor(AActor* Actor, bool bInSelected, bool bNotify,
 			? SelectionSet->SelectElement(UEngineElementsLibrary::AcquireEditorActorElementHandle(Actor), SelectionOptions)
 			: SelectionSet->DeselectElement(UEngineElementsLibrary::AcquireEditorActorElementHandle(Actor), SelectionOptions);
 
+		// Clear all active selected elements if the selection is disabled by locking the level
+		// FLevelUtils::IsLevelLocked(SelectionLevel) is 'false' when the user just locked the level, 
+		// so it was previously 'false' and the user is trying to make it 'true' now.
+		// 
+		// StreamingLevel is used to handle the case when a sub-level is locked but not the main level
+		// if we don't check "StreamingLevel == NULL: when a user locks a sub-level, the user's current selections will be cleared.
+		ULevel* SelectionLevel = Actor->GetLevel();
+		ULevelStreaming* StreamingLevel = FLevelUtils::FindStreamingLevel( SelectionLevel );
+
+		if (StreamingLevel == NULL && !FLevelUtils::IsLevelLocked(SelectionLevel))
+		{
+			SelectionSet->ClearSelection(FTypedElementSelectionOptions());
+		}
+
 		if (bSelectionChanged)
 		{
 			if (bNotify)
