ActorMenu: Avoid loading non UClass assets

#prefight 647749180848b7126dbc1156
#rnx

diff --git a/Engine/Source/Editor/LevelEditor/Private/LevelEditorCreateActorMenu.cpp b/Engine/Source/Editor/LevelEditor/Private/LevelEditorCreateActorMenu.cpp
--- a/Engine/Source/Editor/LevelEditor/Private/LevelEditorCreateActorMenu.cpp
+++ b/Engine/Source/Editor/LevelEditor/Private/LevelEditorCreateActorMenu.cpp
@@ -99,12 +99,14 @@ static void GetMenuEntryText(const FAssetData& Asset, const TArray<FActorFactory
 	if (AssetMenuOptions.Num() == 1)
 	{
 		const FActorFactoryAssetProxy::FMenuItem& MenuItem = AssetMenuOptions[0];
-		UClass* MenuItemClass = Cast<UClass>(MenuItem.AssetData.GetAsset());
-
-		if (IsClass && MenuItemClass && MenuItemClass->IsChildOf(AActor::StaticClass()))
+		if(IsClass)
 		{
-			AActor* DefaultActor = Cast<AActor>(Cast<UClass>(MenuItem.AssetData.GetAsset())->ClassDefaultObject);
-			OutActorTypeDisplayName = FText::FromString(FName::NameToDisplayString(DefaultActor->GetClass()->GetName(), false));
+			UClass* MenuItemClass = Cast<UClass>(MenuItem.AssetData.GetAsset());
+			if (MenuItemClass && MenuItemClass->IsChildOf(AActor::StaticClass()))
+			{
+				AActor* DefaultActor = Cast<AActor>(Cast<UClass>(MenuItem.AssetData.GetAsset())->ClassDefaultObject);
+				OutActorTypeDisplayName = FText::FromString(FName::NameToDisplayString(DefaultActor->GetClass()->GetName(), false));
+			}
 		}
 
 		// If the class type name wasn't set above, then use the factory's display name
