Fix a crash in the bulk edit components submenu on Mac

diff --git a/Engine/Source/Editor/LevelEditor/Private/LevelEditorContextMenu.cpp b/Engine/Source/Editor/LevelEditor/Private/LevelEditorContextMenu.cpp
--- a/Engine/Source/Editor/LevelEditor/Private/LevelEditorContextMenu.cpp
+++ b/Engine/Source/Editor/LevelEditor/Private/LevelEditorContextMenu.cpp
@@ -1433,8 +1433,19 @@ void FLevelEditorContextMenuImpl::FillBulkEditComponentsMenu( UToolMenu* Menu )
 	}
 
 	TObjectPtr<const UTypedElementSelectionSet> Selection = LevelEditorContext->CurrentSelection;
+
+	if(!Selection)
+	{
+		return;
+	}
+	
 	TArray<AActor*> SelectedActors = Selection->GetSelectedObjects<AActor>();
 
+	if(SelectedActors.Num() == 0)
+	{
+		return;
+	}
+	
 	/* Map to keep track of components that are common between the current selection
 	 * The key represents the type of the component, while the value contains the actual components from each object in the selection
 	 */
