Unshelved from pending changelist '25239902':

PR #10401: Level editor place actors menu to respect Class DisplayName overrides
#contrib 10401

diff --git a/Engine/Source/Editor/LevelEditor/Private/LevelEditorCreateActorMenu.cpp b/Engine/Source/Editor/LevelEditor/Private/LevelEditorCreateActorMenu.cpp
--- a/Engine/Source/Editor/LevelEditor/Private/LevelEditorCreateActorMenu.cpp
+++ b/Engine/Source/Editor/LevelEditor/Private/LevelEditorCreateActorMenu.cpp
@@ -89,7 +89,7 @@ static void GetMenuEntryText(const FAssetData& Asset, const TArray<FActorFactory
 
 	if (IsClass)
 	{
-		OutAssetDisplayName = FText::FromString(FName::NameToDisplayString(Asset.AssetName.ToString(), false));
+		OutAssetDisplayName = Asset.GetClass()->GetDisplayNameText();
 	}
 	else
 	{
@@ -99,13 +99,13 @@ static void GetMenuEntryText(const FAssetData& Asset, const TArray<FActorFactory
 	if (AssetMenuOptions.Num() == 1)
 	{
 		const FActorFactoryAssetProxy::FMenuItem& MenuItem = AssetMenuOptions[0];
-		if(IsClass)
+		if (IsClass)
 		{
 			UClass* MenuItemClass = Cast<UClass>(MenuItem.AssetData.GetAsset());
 			if (MenuItemClass && MenuItemClass->IsChildOf(AActor::StaticClass()))
 			{
 				AActor* DefaultActor = Cast<AActor>(Cast<UClass>(MenuItem.AssetData.GetAsset())->ClassDefaultObject);
-				OutActorTypeDisplayName = FText::FromString(FName::NameToDisplayString(DefaultActor->GetClass()->GetName(), false));
+				OutActorTypeDisplayName = DefaultActor->GetClass()->GetDisplayNameText();
 			}
 		}
 
