Explicitly register FGlobalEditorCommonCommands to avoid implicit initialization order issues
#horde 317282

diff --git a/Engine/Source/Editor/LevelEditor/Private/LevelEditor.cpp b/Engine/Source/Editor/LevelEditor/Private/LevelEditor.cpp
--- a/Engine/Source/Editor/LevelEditor/Private/LevelEditor.cpp
+++ b/Engine/Source/Editor/LevelEditor/Private/LevelEditor.cpp
@@ -269,6 +269,7 @@ void FLevelEditorModule::StartupModule()
 	NotificationBarExtensibilityManager = MakeShared<FExtensibilityManager>();
 
 	// Note this must come before any tab spawning because that can create the SLevelEditor and attempt to map commands
+	FGlobalEditorCommonCommands::Register();
 	FEditorViewportCommands::Register();
 	FLevelViewportCommands::Register();
 	FLevelEditorCommands::Register();
