Disable the bulk edit section in the Outliner if the property matrix can't be invoked

diff --git a/Engine/Source/Editor/LevelEditor/Private/LevelEditorContextMenu.cpp b/Engine/Source/Editor/LevelEditor/Private/LevelEditorContextMenu.cpp
--- a/Engine/Source/Editor/LevelEditor/Private/LevelEditorContextMenu.cpp
+++ b/Engine/Source/Editor/LevelEditor/Private/LevelEditorContextMenu.cpp
@@ -360,30 +360,31 @@ void FLevelEditorContextMenu::RegisterActorContextMenu()
 			LevelEditorCreateActorMenu::FillAddReplaceContextMenuSections(Section, LevelEditorContext);
 		}
 		
-		{	// Options that relate to bulk editing assets
-			// In most cases, you DO NOT want to extend this section; look at ActorUETools or ActorTypeTools below
-			FToolMenuSection& Section = InMenu->AddSection("ActorBulkEdit", LOCTEXT("ActorBulkEditHeading", "Bulk Editing"));
-			
-			Section.AddMenuEntry(
-				FLevelEditorCommands::Get().OpenSelectionInPropertyMatrix,
-				TAttribute<FText>(), // use command's label
-				TAttribute<FText>(), // use command's tooltip
-				FSlateIcon(FAppStyle::GetAppStyleSetName(), "DetailsView.EditRawProperties")
-			);
+		{
+			if(FModuleManager::LoadModuleChecked<FPropertyEditorModule>( "PropertyEditor" ).GetCanUsePropertyMatrix() && NumSelectedActors >= 2)
+			{
+				// Options that relate to bulk editing assets
+				// In most cases, you DO NOT want to extend this section; look at ActorUETools or ActorTypeTools below
+				FToolMenuSection& Section = InMenu->AddSection("ActorBulkEdit", LOCTEXT("ActorBulkEditHeading", "Bulk Editing"));
+				
+				Section.AddMenuEntry(
+					FLevelEditorCommands::Get().OpenSelectionInPropertyMatrix,
+					TAttribute<FText>(), // use command's label
+					TAttribute<FText>(), // use command's tooltip
+					FSlateIcon(FAppStyle::GetAppStyleSetName(), "DetailsView.EditRawProperties")
+				);
 
-			Section.AddSubMenu(
-				"BulkEditComponentsSubmenu",
-				LOCTEXT("BulkEditComponentsSubmenuName", "Edit Components in the Property Matrix"),
-				LOCTEXT("BulkEditComponentsSubmenuTooltip", "Bulk Edit any editable components that are common between the selected actors in the Property Matrix"),
-				FNewToolMenuDelegate::CreateStatic(&FLevelEditorContextMenuImpl::FillBulkEditComponentsMenu),
-				FUIAction(FExecuteAction(), FCanExecuteAction(), FIsActionChecked(), FIsActionButtonVisible::CreateLambda([NumSelectedActors]()
-				{
-					return NumSelectedActors >= 2;
-				})),
-				EUserInterfaceActionType::Button,
-				false, // default value
-				FSlateIcon(FAppStyle::GetAppStyleSetName(), "DetailsView.EditRawProperties")
-			);
+				Section.AddSubMenu(
+					"BulkEditComponentsSubmenu",
+					LOCTEXT("BulkEditComponentsSubmenuName", "Edit Components in the Property Matrix"),
+					LOCTEXT("BulkEditComponentsSubmenuTooltip", "Bulk Edit any editable components that are common between the selected actors in the Property Matrix"),
+					FNewToolMenuDelegate::CreateStatic(&FLevelEditorContextMenuImpl::FillBulkEditComponentsMenu),
+					FUIAction(),
+					EUserInterfaceActionType::Button,
+					false, // default value
+					FSlateIcon(FAppStyle::GetAppStyleSetName(), "DetailsView.EditRawProperties")
+				);
+			}
 		}
 
 
