Make show flags persistent under "Game View".

Assumptions:
1. Selections are cleared and "Mode Widget" is disabled after switching into "Game View", based on the existing function "void FEditorViewportClient::SetGameView(bool bGameViewEnable)" in "EditorViewportClient.cpp".
2. "bAlwaysShowModeWidgetAfterSelectionChanges" being "true" under "Game View" is not the intended behavior, as "void SLevelViewport::OnActorSelectionChanged(const TArray<UObject*>& NewSelection, bool bForceRefresh)" would bring "Mode Widget" back.

diff --git a/Engine/Source/Editor/LevelEditor/Private/SLevelViewport.cpp b/Engine/Source/Editor/LevelEditor/Private/SLevelViewport.cpp
--- a/Engine/Source/Editor/LevelEditor/Private/SLevelViewport.cpp
+++ b/Engine/Source/Editor/LevelEditor/Private/SLevelViewport.cpp
@@ -1988,6 +1988,9 @@ void SLevelViewport::ToggleGameView()
 	{
 		bool bGameViewEnable = !LevelViewportClient->IsInGameView();
 
+		// "Mode Widget" should not automatically be reactivated by selecting an actor after "Game View" is enabled
+		LevelViewportClient->bAlwaysShowModeWidgetAfterSelectionChanges = bGameViewEnable ? false : true;
+
 		LevelViewportClient->SetGameView(bGameViewEnable);
 	}
 }
