Prevent LevelEditorViewportClient draw on PIE in LevelEditorViewport
* This change addesses an issue which causes the LevelEditorViewportClient to render while PIE is active on the LevelEditorViewport. This results in a frame which is never displayed and therefore reduced rendering performance.
* The issue's root cuase is due to behaviour where the GameViewportClient and EditorViewportClient swap their stat flags when PIE is started on a selected LevelEditor sceneviewport - see FSceneViewport::SwapStatCommands.  When stat flags are present on the GameViewportClient before PIE is started, the EditorViewportClient will inherit the stat flags.  However, this has the side-effect of pushing a 'true' RealTime override onto ViewportClient.  This causes ViewportClient->IsRealtime() to return true.
* UEditorEngine::UpdateSingleViewportClient() is called by the UEditorEngine::Tick() each frame,  finds the EditorViewportClient->IsRealtime() to be true and triggers LevelEditorViewportClient->Viewport->Draw() even though there is no SceneViewport to display the render.
* The fix pushes a realtime override to disable realtime mode for the LevelEditorViewportClient when a PIE session is started that takes over the LevelEditorViewport and subsequently removes the override when PIE ends.

diff --git a/Engine/Source/Editor/LevelEditor/Private/SLevelViewport.cpp b/Engine/Source/Editor/LevelEditor/Private/SLevelViewport.cpp
--- a/Engine/Source/Editor/LevelEditor/Private/SLevelViewport.cpp
+++ b/Engine/Source/Editor/LevelEditor/Private/SLevelViewport.cpp
@@ -4029,6 +4029,8 @@ void SLevelViewport::StartPlayInEditorSession(UGameViewportClient* PlayClient, c
 	ActiveViewport->SetPlayInEditorIsSimulate(bInSimulateInEditor);
 	
 	ActiveViewport->OnPlayWorldViewportSwapped( *InactiveViewport );
+	
+	LevelViewportClient->AddRealtimeOverride(false, LOCTEXT("LevelViewport_RealTimeDisableOnPie", "Disable LevelViewport Realtime for PIE"));
 
 	TSharedPtr<SWindow> ParentWindow = FSlateApplication::Get().FindWidgetWindow(AsShared());
 	PlayClient->SetViewportOverlayWidget(ParentWindow, PIEViewportOverlayWidget.ToSharedRef());
@@ -4231,6 +4233,8 @@ void SLevelViewport::EndPlayInEditorSession()
 
 	if( IsPlayInEditorViewportActive() )
 	{
+		LevelViewportClient->RemoveRealtimeOverride(LOCTEXT("LevelViewport_RealTimeDisableOnPie", "Disable LevelViewport Realtime for PIE"));
+		
 		{
 			TSharedPtr<FSceneViewport> GameViewport = ActiveViewport;
 			ActiveViewport = InactiveViewport;
