Problem:
"Mode Widgets" is by default disabled in game mode. The expected behavior in this case, according to feedback from content producers, should be as follows: "Mode Widgets" should be enabled by default, but the selection outline and transform widgets should be hidden. They should reappear once the user selects an actor.

Cause:
"EngineShowFlags.SetModeWidgets(false);" when entering game mode

Fix:
When entering game mode, temporarily set LevelViewportClient->bShowWidget to false while keeping the showflag "Mode Widgets" set to true. Transform widget is hidden in this manner when entering game mode, and selecting any actor will make it visible again because of the modification in SLevelViewport::OnActorSelectionChanged. The modification to SLevelViewport::ToggleGameView() makes sure that after leaving game mode, LevelViewportClient->bShowWidget is properly reset to true.

diff --git a/Engine/Source/Editor/LevelEditor/Private/SLevelViewport.cpp b/Engine/Source/Editor/LevelEditor/Private/SLevelViewport.cpp
--- a/Engine/Source/Editor/LevelEditor/Private/SLevelViewport.cpp
+++ b/Engine/Source/Editor/LevelEditor/Private/SLevelViewport.cpp
@@ -1992,6 +1992,13 @@ void SLevelViewport::ToggleGameView()
 		LevelViewportClient->bAlwaysShowModeWidgetAfterSelectionChanges = bGameViewEnable ? false : true;
 
 		LevelViewportClient->SetGameView(bGameViewEnable);
+
+		if (!bGameViewEnable)
+		{
+			// LevelViewportClient->bShowWidget is set to "false" when entering game mode
+			// Need to turn it back to "true" when exiting game mode
+			LevelViewportClient->ShowWidget(true);
+		}
 	}
 }
 
@@ -2763,6 +2770,9 @@ void SLevelViewport::OnActorSelectionChanged(const TArray<UObject*>& NewSelectio
 		{
 			LevelViewportClient->EngineShowFlags.SetModeWidgets(true);
 		}
+
+		// In game mode, selecting any actor should make LevelViewportClient->bShowWidget be "true"
+		LevelViewportClient->ShowWidget(true);
 		LevelViewportClient->EngineShowFlags.SetSelection(true);
 		LevelViewportClient->EngineShowFlags.SetSelectionOutline(GetDefault<ULevelEditorViewportSettings>()->bUseSelectionOutline);
 	}
diff --git a/Engine/Source/Editor/UnrealEd/Private/EditorViewportClient.cpp b/Engine/Source/Editor/UnrealEd/Private/EditorViewportClient.cpp
--- a/Engine/Source/Editor/UnrealEd/Private/EditorViewportClient.cpp
+++ b/Engine/Source/Editor/UnrealEd/Private/EditorViewportClient.cpp
@@ -6271,8 +6271,9 @@ void FEditorViewportClient::SetGameView(bool bGameViewEnable)
 	//reset game engine show flags that may have been turned on by making a selection in game view
 	if(bGameViewEnable)
 	{
-		EngineShowFlags.SetModeWidgets(false);
+		EngineShowFlags.SetModeWidgets(true); // Enable "Mode Widgets" by default when entering game mode
 		EngineShowFlags.SetSelection(false);
+		ShowWidget(false); // Hide the widget
 	}
 
 	EngineShowFlags.SetSelectionOutline(bGameViewEnable ? false : GetDefault<ULevelEditorViewportSettings>()->bUseSelectionOutline);
