Fix Pivot Update Issue When Snapping Object to View

diff --git a/Engine/Source/Editor/LevelEditor/Private/LevelEditorActions.cpp b/Engine/Source/Editor/LevelEditor/Private/LevelEditorActions.cpp
--- a/Engine/Source/Editor/LevelEditor/Private/LevelEditorActions.cpp
+++ b/Engine/Source/Editor/LevelEditor/Private/LevelEditorActions.cpp
@@ -2876,21 +2876,43 @@ void FLevelEditorActionCallbacks::SnapObjectToView_Clicked()
 	const FScopedTransaction Transaction(NSLOCTEXT("UnrealEd", "SnapObjectToView", "Snap Object to View"));
 
 	// Fires ULevel::LevelDirtiedEvent when falling out of scope.
-	FScopedLevelDirtied		LevelDirtyCallback;
+	FScopedLevelDirtied LevelDirtyCallback;
 
-	for (FSelectionIterator It(GEditor->GetSelectedActorIterator()); It; ++It)
-	{
-		AActor* Actor = Cast<AActor>(*It);
-		Actor->Modify();
-		FVector Location = GCurrentLevelEditingViewportClient->GetViewLocation();
-		FRotator Rotation = GCurrentLevelEditingViewportClient->GetViewRotation();
+	// Get the new location and rotation for the actor from the viewport client's view.
+	const FVector NewLocation = GCurrentLevelEditingViewportClient->GetViewLocation();
+	const FQuat NewRotation = GCurrentLevelEditingViewportClient->GetViewRotation().Quaternion();
 
-		Actor->SetActorLocation(Location);
-		Actor->SetActorRotation(Rotation);
-		Actor->PostEditMove(true);
+	UTypedElementSelectionSet* SelectionSet = GEditor->GetSelectedActors()->GetElementSelectionSet();
+	SelectionSet->ForEachSelectedElement<ITypedElementWorldInterface>([&NewLocation, &NewRotation, &LevelDirtyCallback](const TTypedElement<ITypedElementWorldInterface>& InElement)
+		{
+			// Get the actor's current transform.
+			FTransform CurrentTransform;
+			if (InElement.GetWorldTransform(CurrentTransform))
+			{
+				// Set new location and rotation to the current transform.
+				FTransform NewTransform = CurrentTransform;
+				NewTransform.SetLocation(NewLocation);
+				NewTransform.SetRotation(NewRotation);
 
-		LevelDirtyCallback.Request();
-	}
+				// Find a suitable transform, if the actor can't be at the exact desired transform.
+				FTransform SuitableTransform;
+				if (!InElement.FindSuitableTransformAtPoint(NewTransform, SuitableTransform))
+				{
+					SuitableTransform = NewTransform;
+				}
+
+				InElement.NotifyMovementStarted();
+				InElement.SetWorldTransform(SuitableTransform);
+				InElement.NotifyMovementEnded();
+
+				LevelDirtyCallback.Request();
+			}
+
+			return true;
+		});
+
+	GEditor->SetPivot(NewLocation, false, true); // Update the pivot location of the editor to the new actor location.
+	GEditor->RedrawLevelEditingViewports();
 }
 
 void FLevelEditorActionCallbacks::CopyActorFilePathtoClipboard_Clicked()
