Added GetLevelEditor to Level Editor Menu Context to get the (public facing) ILevelEditor, as the existing SLevelEditor is unusable outside

diff --git a/Engine/Source/Editor/LevelEditor/Private/LevelEditorMenuContext.cpp b/Engine/Source/Editor/LevelEditor/Private/LevelEditorMenuContext.cpp
--- a/Engine/Source/Editor/LevelEditor/Private/LevelEditorMenuContext.cpp
+++ b/Engine/Source/Editor/LevelEditor/Private/LevelEditorMenuContext.cpp
@@ -3,8 +3,14 @@
 #include "LevelEditorMenuContext.h"
 
 #include "Elements/Framework/TypedElementRegistry.h"
+#include "SLevelEditor.h"
+
+TWeakPtr<ILevelEditor> ULevelEditorMenuContext::GetLevelEditor() const
+{
+	return LevelEditor;
+}
 
 FScriptTypedElementHandle ULevelEditorContextMenuContext::GetScriptHitProxyElement()
 {
 	return UTypedElementRegistry::GetInstance()->CreateScriptHandle(HitProxyElement.GetId());
-}
\ No newline at end of file
+}
diff --git a/Engine/Source/Editor/LevelEditor/Public/LevelEditorMenuContext.h b/Engine/Source/Editor/LevelEditor/Public/LevelEditorMenuContext.h
--- a/Engine/Source/Editor/LevelEditor/Public/LevelEditorMenuContext.h
+++ b/Engine/Source/Editor/LevelEditor/Public/LevelEditorMenuContext.h
@@ -32,6 +32,8 @@ class LEVELEDITOR_API ULevelEditorMenuContext : public UObject
 {
 	GENERATED_BODY()
 public:
+	TWeakPtr<ILevelEditor> GetLevelEditor() const;
+
 	TWeakPtr<SLevelEditor> LevelEditor;
 };
 
