- Fix for accumulating realtime overrides after each PIE session

diff --git a/Engine/Source/Editor/LevelEditor/Private/SLevelViewport.cpp b/Engine/Source/Editor/LevelEditor/Private/SLevelViewport.cpp
--- a/Engine/Source/Editor/LevelEditor/Private/SLevelViewport.cpp
+++ b/Engine/Source/Editor/LevelEditor/Private/SLevelViewport.cpp
@@ -4241,10 +4241,11 @@ void SLevelViewport::EndPlayInEditorSession()
 
 	check( InactiveViewport.IsValid() );
 
+	const bool bCheckMissingOverride = false;
+	LevelViewportClient->RemoveRealtimeOverride(LOCTEXT("LevelViewport_RealTimeDisableOnPie", "Disable LevelViewport Realtime for PIE"), bCheckMissingOverride);
+
 	if( IsPlayInEditorViewportActive() )
 	{
-		LevelViewportClient->RemoveRealtimeOverride(LOCTEXT("LevelViewport_RealTimeDisableOnPie", "Disable LevelViewport Realtime for PIE"));
-		
 		{
 			TSharedPtr<FSceneViewport> GameViewport = ActiveViewport;
 			ActiveViewport = InactiveViewport;
