Changed a TArray to a TSet for FLevelEditorOutlinerSettings::UnsavedAssets.
Editor is now responsive when generating 32K actors through PCG.

diff --git a/Engine/Source/Editor/LevelEditor/Private/LevelEditorOutlinerSettings.cpp b/Engine/Source/Editor/LevelEditor/Private/LevelEditorOutlinerSettings.cpp
--- a/Engine/Source/Editor/LevelEditor/Private/LevelEditorOutlinerSettings.cpp
+++ b/Engine/Source/Editor/LevelEditor/Private/LevelEditorOutlinerSettings.cpp
@@ -348,7 +348,10 @@ void FLevelEditorOutlinerSettings::CreateSCCFilters()
 
 	// Unsaved Actors
 	FUnsavedAssetsTrackerModule& UnsavedAssetsTrackerModule = FUnsavedAssetsTrackerModule::Get();
-	UnsavedAssets = UnsavedAssetsTrackerModule.GetUnsavedAssets();
+	
+	UnsavedAssets.Empty();
+	UnsavedAssets.Append(UnsavedAssetsTrackerModule.GetUnsavedAssets());
+
 	UnsavedAssetsTrackerModule.OnUnsavedAssetAdded.AddSP(this, &FLevelEditorOutlinerSettings::OnUnsavedAssetAdded);
 	UnsavedAssetsTrackerModule.OnUnsavedAssetRemoved.AddSP(this, &FLevelEditorOutlinerSettings::OnUnsavedAssetRemoved);
 	
diff --git a/Engine/Source/Editor/LevelEditor/Public/LevelEditorOutlinerSettings.h b/Engine/Source/Editor/LevelEditor/Public/LevelEditorOutlinerSettings.h
--- a/Engine/Source/Editor/LevelEditor/Public/LevelEditorOutlinerSettings.h
+++ b/Engine/Source/Editor/LevelEditor/Public/LevelEditorOutlinerSettings.h
@@ -115,8 +115,8 @@ private:
 
 	// Source Control Cache
 
-	// List of currently unsaved assets
-	TArray<FString> UnsavedAssets;
+	// List of currently unsaved assets. Using a set here because it can grow a lot and impact performances when looking up in the list.
+	TSet<FString> UnsavedAssets;
 
 	// List of uncontrolled changelist state
 	TArray<TSharedRef<FUncontrolledChangelistState>> UncontrolledChangelistStates;
