Create CVars to enable\disable the card layout and if it is enabled, hide or show the component editor.

Updated CoreEntityDetailsFilter to introduce the following CVars:

If true, enables the card layout, else core entities will use Classic details panel:
CoreEntity.DetailsPanel.DetailsPanelEnableCardLayout

If in Card layout and true, shows the component editor:
CoreEntity.UI.ForceShowComponentEditor

Also updated SActorDetails.cpp to use CoreEntity.UI.ForceShowComponentEditor to show and hide the component editor.

diff --git a/Engine/Source/Editor/LevelEditor/Private/SActorDetails.cpp b/Engine/Source/Editor/LevelEditor/Private/SActorDetails.cpp
--- a/Engine/Source/Editor/LevelEditor/Private/SActorDetails.cpp
+++ b/Engine/Source/Editor/LevelEditor/Private/SActorDetails.cpp
@@ -880,7 +880,11 @@ void SActorDetails::OnNativeComponentWarningHyperlinkClicked(const FSlateHyperli
 
 EVisibility SActorDetails::GetComponentEditorVisibility() const
 {
-	const bool bHideEditorFromDetailsView = DetailsView.IsValid() &&  !DetailsView->IsDefaultStyle();
+	IConsoleVariable* ForceShow = IConsoleManager::Get().FindConsoleVariable(TEXT("CoreEntity.UI.ForceShowComponentEditor"));
+
+	// force hide it if the style is not default and the ForceShowComponentEditor CVar is not set to true
+	const bool bHideEditorFromDetailsView = (DetailsView.IsValid() &&  !DetailsView->IsDefaultStyle()) &&
+		                                    (!ForceShow || !ForceShow->GetBool());
 	return GetActorContext() && !bHideEditorFromDetailsView  ? EVisibility::Visible : EVisibility::Collapsed;
 }
 
