Create CVar to enable card layout through out the interface

This enables the card layout upon renewal of selected items in UE. I added the CVar DetailsPanel.Style.EnableCardLayout. When true throughout UE Card layout will be used, else the Classic details panel layout will be used. Also made it so the hiding\showing of the component editor is independant of the details style, by adding a ShouldHideComponentEditor to the DetailsVieiwObjectFilter, which is checked in SActorDetails.cpp to see if the Component Editor should be shown. SDetailsView has been updated to return the Card layout as the primary style key if DetailsPanel.Style.EnableCardLayout is true.

diff --git a/Engine/Source/Editor/LevelEditor/Private/SActorDetails.cpp b/Engine/Source/Editor/LevelEditor/Private/SActorDetails.cpp
--- a/Engine/Source/Editor/LevelEditor/Private/SActorDetails.cpp
+++ b/Engine/Source/Editor/LevelEditor/Private/SActorDetails.cpp
@@ -617,9 +617,10 @@ FString SActorDetails::GetReferencerName() const
 	return TEXT("SActorDetails");
 }
 
-void SActorDetails::SetActorDetailsRootCustomization(TSharedPtr<FDetailsViewObjectFilter> ActorDetailsObjectFilter, TSharedPtr<IDetailRootObjectCustomization> ActorDetailsRootCustomization)
+void SActorDetails::SetActorDetailsRootCustomization(TSharedPtr<FDetailsViewObjectFilter> InActorDetailsObjectFilter, TSharedPtr<IDetailRootObjectCustomization> ActorDetailsRootCustomization)
 {
-	DetailsView->SetObjectFilter(ActorDetailsObjectFilter);
+	ActorDetailsObjectFilter = InActorDetailsObjectFilter.ToWeakPtr();
+	DetailsView->SetObjectFilter(InActorDetailsObjectFilter);
 	DetailsView->SetRootObjectCustomizationInstance(ActorDetailsRootCustomization);
 	DetailsView->ForceRefresh();
 }
@@ -883,8 +884,9 @@ EVisibility SActorDetails::GetComponentEditorVisibility() const
 	IConsoleVariable* ForceShow = IConsoleManager::Get().FindConsoleVariable(TEXT("CoreEntity.UI.ForceShowComponentEditor"));
 
 	// force hide it if the style is not default and the ForceShowComponentEditor CVar is not set to true
-	const bool bHideEditorFromDetailsView = (DetailsView.IsValid() &&  !DetailsView->IsDefaultStyle()) &&
-		                                    (!ForceShow || !ForceShow->GetBool());
+	const bool bHideEditorFromDetailsView = (ActorDetailsObjectFilter.IsValid() &&
+							                 ActorDetailsObjectFilter.Pin()->ShouldHideComponentEditor()) &&
+		                    (!ForceShow || !ForceShow->GetBool());
 	return GetActorContext() && !bHideEditorFromDetailsView  ? EVisibility::Visible : EVisibility::Collapsed;
 }
 
diff --git a/Engine/Source/Editor/LevelEditor/Private/SActorDetails.h b/Engine/Source/Editor/LevelEditor/Private/SActorDetails.h
--- a/Engine/Source/Editor/LevelEditor/Private/SActorDetails.h
+++ b/Engine/Source/Editor/LevelEditor/Private/SActorDetails.h
@@ -71,7 +71,7 @@ public:
 	 *
 	 * @param InFilter	The filter to use or nullptr to remove the active filter
 	 */
-	void SetActorDetailsRootCustomization(TSharedPtr<FDetailsViewObjectFilter> ActorDetailsObjectFilter, TSharedPtr<class IDetailRootObjectCustomization> ActorDetailsRootCustomization);
+	void SetActorDetailsRootCustomization(TSharedPtr<FDetailsViewObjectFilter> InActorDetailsObjectFilter, TSharedPtr<class IDetailRootObjectCustomization> ActorDetailsRootCustomization);
 
 	UE_DEPRECATED(5.0, "SetSCSEditorUICustomization is deprecated, please use SetSubobjectEditorUICustomization instead.")
 	void SetSCSEditorUICustomization(TSharedPtr<ISCSEditorUICustomization> ActorDetailsSCSEditorUICustomization) { SetSubobjectEditorUICustomization(ActorDetailsSCSEditorUICustomization); }
@@ -142,4 +142,7 @@ private:
 
 	// Used to prevent reentrant changes
 	bool bSelectionGuard = false;
+
+	/** The object filter which contains relevant information about the details objects being displayed */
+	TWeakPtr<FDetailsViewObjectFilter> ActorDetailsObjectFilter;
 };
diff --git a/Engine/Source/Editor/PropertyEditor/Private/SDetailsView.cpp b/Engine/Source/Editor/PropertyEditor/Private/SDetailsView.cpp
--- a/Engine/Source/Editor/PropertyEditor/Private/SDetailsView.cpp
+++ b/Engine/Source/Editor/PropertyEditor/Private/SDetailsView.cpp
@@ -1,6 +1,5 @@
 // Copyright Epic Games, Inc. All Rights Reserved.
 
-
 #include "SDetailsView.h"
 
 #include "CategoryPropertyNode.h"
@@ -32,6 +31,12 @@
 
 #define LOCTEXT_NAMESPACE "SDetailsView"
 
+static TAutoConsoleVariable<bool> CVarDetailsPanelEnableCardLayout(
+	TEXT("DetailsPanel.Style.EnableCardLayout"),
+	false,
+	TEXT("Specifies whether the card layout is in effect for the Details View."),
+	ECVF_Default);
+
 SDetailsView::~SDetailsView()
 {
 	const FRootPropertyNodeList& RootNodes = GetRootNodes();
@@ -1500,14 +1505,17 @@ void SDetailsView::UpdateStyleKey()
 
 const FDetailsViewStyleKey& SDetailsView::GetStyleKey()
 {
-	static FDetailsViewStyleKey PrimaryKey = GetPrimaryDetailsViewStyleKey();
+	const FDetailsViewStyleKey& PrimaryKey = GetPrimaryDetailsViewStyleKey();
 	return StyleKeySP.IsValid() ? *StyleKeySP.Get() : PrimaryKey;
 }
 
 const FDetailsViewStyleKey& SDetailsView::GetPrimaryDetailsViewStyleKey()
 {
-	static FDetailsViewStyleKey PrimaryKey = FDetailsViewStyleKeys::Classic();
-	return PrimaryKey;
+	if (CVarDetailsPanelEnableCardLayout.GetValueOnAnyThread())
+	{
+		return FDetailsViewStyleKeys::Card();
+	}
+	return FDetailsViewStyleKeys::Classic();
 }
 
 #undef LOCTEXT_NAMESPACE
diff --git a/Engine/Source/Editor/PropertyEditor/Public/DetailsViewObjectFilter.h b/Engine/Source/Editor/PropertyEditor/Public/DetailsViewObjectFilter.h
--- a/Engine/Source/Editor/PropertyEditor/Public/DetailsViewObjectFilter.h
+++ b/Engine/Source/Editor/PropertyEditor/Public/DetailsViewObjectFilter.h
@@ -62,4 +62,12 @@ public:
 		static const FDetailsViewStyleKey Default = FDetailsViewStyleKeys::Default();
 		return Default;
 	}
+
+	/**
+	 * Returns a boolean indicating if the Component Editor should be hidden for these object(s)
+	 */
+	virtual bool ShouldHideComponentEditor()
+	{
+		return false;
+	}
 };
