Use correct background color for title bar project info badge

There is project information text on the right side of the title bar which uses a darker background color than the title bar background. This change fixes that by using the "Colors.Title" color as the background.

#rnx

diff --git a/Engine/Source/Editor/LevelEditor/Private/LevelEditor.cpp b/Engine/Source/Editor/LevelEditor/Private/LevelEditor.cpp
--- a/Engine/Source/Editor/LevelEditor/Private/LevelEditor.cpp
+++ b/Engine/Source/Editor/LevelEditor/Private/LevelEditor.cpp
@@ -105,7 +105,7 @@ public:
 		FString ProjectNameWatermarkPrefix;
 		GConfig->GetString(TEXT("LevelEditor"), TEXT("ProjectNameWatermarkPrefix"), /*out*/ ProjectNameWatermarkPrefix, GEditorPerProjectIni);
 
-		FSlateColor BadgeBackgroundColor = FAppStyle::Get().GetSlateColor("Colors.Foldout");
+		FSlateColor BadgeBackgroundColor = FAppStyle::Get().GetSlateColor("Colors.Title");
 
 		FColor ConfigColor;
 		if (GConfig->GetColor(TEXT("LevelEditor"), TEXT("ProjectBadgeBackgroundColor"), /*out*/ ConfigColor, GEditorPerProjectIni))
