Update CVar to show the subobject editor by default.

diff --git a/Engine/Source/Editor/LevelEditor/Private/SActorDetails.cpp b/Engine/Source/Editor/LevelEditor/Private/SActorDetails.cpp
--- a/Engine/Source/Editor/LevelEditor/Private/SActorDetails.cpp
+++ b/Engine/Source/Editor/LevelEditor/Private/SActorDetails.cpp
@@ -47,7 +47,7 @@
 
 static TAutoConsoleVariable<bool> CVarForceShowComponentEditor(
 	TEXT("CoreEntity.UI.ForceShowComponentEditor"),
-	false,
+	true,
 	TEXT("Force the component editor to show in the main details tab."));
 
 namespace UE::LevelEditor::Private
