The Prefabs editor subobject editor should behave the same as the subobject editor in the main window.

    It should always show the subobject editor when an actor is selected (this state only exists for the main window as far as I know, but it still shold be honored if it is possible)
    When a core entity\prefab is selected
        It should show the subobject editor when the CVar CoreEntity.UI.ForceShowComponentEditor = true, and this is the default behavior
        It should hide the subobject editor when the CVar CoreEntity.UI.ForceShowComponentEditor = false
        It should not have a search bar in the subobject editor

Note that there will be some differences, like the prefabs editor will not have a blueprints button, and currently the main editor has a label of the number of components\objects that I think has to be taken off, so that won't be there. See screen cap for new behavior.

diff --git a/Engine/Source/Editor/LevelEditor/Private/SActorDetails.cpp b/Engine/Source/Editor/LevelEditor/Private/SActorDetails.cpp
--- a/Engine/Source/Editor/LevelEditor/Private/SActorDetails.cpp
+++ b/Engine/Source/Editor/LevelEditor/Private/SActorDetails.cpp
@@ -45,11 +45,6 @@
 
 #define LOCTEXT_NAMESPACE "SActorDetails"
 
-static TAutoConsoleVariable<bool> CVarForceShowComponentEditor(
-	TEXT("CoreEntity.UI.ForceShowComponentEditor"),
-	true,
-	TEXT("Force the component editor to show in the main details tab."));
-
 namespace UE::LevelEditor::Private
 {
 	class SElementSelectionDetailsButtons : public SCompoundWidget
@@ -890,10 +885,9 @@ void SActorDetails::OnNativeComponentWarningHyperlinkClicked(const FSlateHyperli
 
 EVisibility SActorDetails::GetComponentEditorVisibility() const
 {
-	// force hide it if the style is not default and the ForceShowComponentEditor CVar is not set to true
+	// see if we need to hide the editor due to the current object display 
 	const bool bHideEditorFromDetailsView = (DisplayManager.IsValid() &&
-							                 DisplayManager->ShouldHideComponentEditor()) &&
-											(!CVarForceShowComponentEditor.GetValueOnAnyThread());
+							                 DisplayManager->ShouldHideComponentEditor() );
 	return GetActorContext() && !bHideEditorFromDetailsView  ? EVisibility::Visible : EVisibility::Collapsed;
 }
 
diff --git a/Engine/Source/Editor/PropertyEditor/Private/DetailsDisplayManager.cpp b/Engine/Source/Editor/PropertyEditor/Private/DetailsDisplayManager.cpp
--- a/Engine/Source/Editor/PropertyEditor/Private/DetailsDisplayManager.cpp
+++ b/Engine/Source/Editor/PropertyEditor/Private/DetailsDisplayManager.cpp
@@ -5,6 +5,11 @@
 #include "DetailsViewStyle.h"
 #include "SDetailsView.h"
 
+static TAutoConsoleVariable<bool> CVarForceShowComponentEditor(
+	TEXT("CoreEntity.UI.ForceShowComponentEditor"),
+	true,
+	TEXT("Force the component editor to show in the main details tab."));
+
 
 FDetailsDisplayManager::FDetailsDisplayManager(): bIsOuterCategory(false)
 {
@@ -18,7 +23,7 @@ FDetailsDisplayManager::~FDetailsDisplayManager()
 
 bool FDetailsDisplayManager::ShouldHideComponentEditor()
 {
-	return false;
+	return !GetForceShowSubObjectEditor();
 }
 
 bool FDetailsDisplayManager::ShouldShowCategoryMenu()
@@ -73,4 +78,10 @@ bool FDetailsDisplayManager::GetIsScrollBarNeeded() const
 void FDetailsDisplayManager::SetIsScrollBarNeeded(bool bInIsScrollBarNeeded)
 {
 	bIsScrollBarNeeded = bInIsScrollBarNeeded;
-}
\ No newline at end of file
+}
+
+bool FDetailsDisplayManager::GetForceShowSubObjectEditor()
+{
+	return CVarForceShowComponentEditor.GetValueOnAnyThread();
+}
+
diff --git a/Engine/Source/Editor/PropertyEditor/Public/DetailsDisplayManager.h b/Engine/Source/Editor/PropertyEditor/Public/DetailsDisplayManager.h
--- a/Engine/Source/Editor/PropertyEditor/Public/DetailsDisplayManager.h
+++ b/Engine/Source/Editor/PropertyEditor/Public/DetailsDisplayManager.h
@@ -25,6 +25,12 @@ public:
 	 */
 	virtual bool ShouldHideComponentEditor();
 
+	/**
+	 * Returns a boolean indicating whether the sub-object editor should show regardless of
+	 * any object filter preference settings this would override any value retrieved from ShouldHideComponentEditor.
+	 */
+	PROPERTYEDITOR_API static bool GetForceShowSubObjectEditor();
+
 	/**
 	 * Returns a @code bool @endcode indicating whether this @code DetailsViewObjectFilter @endcode instance
 	 * should show a category menu
