- Redid the quick fix for NDA PreviewPlatforms crashing Mac/Linux, as it was causing some missing menu items and/or ensures with the other fixes

diff --git a/Engine/Config/VulkanPC/DataDrivenPlatformInfo.ini b/Engine/Config/VulkanPC/DataDrivenPlatformInfo.ini
--- a/Engine/Config/VulkanPC/DataDrivenPlatformInfo.ini
+++ b/Engine/Config/VulkanPC/DataDrivenPlatformInfo.ini
@@ -274,7 +274,7 @@ bSupportsUniformBufferObjects = true
 FriendlyName=LOCTEXT("FriendlyShaderPlatformName_VulkanPC_Mobile", "PC Vulkan Mobile")
 
 [PreviewPlatform VULKAN_SM5]
-PlatformName=Linux
+PlatformName=VULKAN_SM5
 ShaderFormat=SF_VULKAN_SM5
 ShaderPlatform=VULKAN_SM5
 MenuTooltip=LOCTEXT("PreviewMenuTooltip_VulkanPC_SM5", "Linux using SM5 profile")
diff --git a/Engine/Config/Windows/DataDrivenPlatformInfo.ini b/Engine/Config/Windows/DataDrivenPlatformInfo.ini
--- a/Engine/Config/Windows/DataDrivenPlatformInfo.ini
+++ b/Engine/Config/Windows/DataDrivenPlatformInfo.ini
@@ -79,7 +79,7 @@ bSupportsDualSourceBlending=true
 bSupportsNNEShaders=true
 
 [PreviewPlatform PCD3D_SM5]
-PlatformName=Windows
+PlatformName=PCD3D_SM5
 ShaderFormat=PCD3D_SM5
 ShaderPlatform=PCD3D_SM5
 MenuTooltip=LOCTEXT("PreviewMenuTooltip_PCD3D_SM5", "DirectX 11, OpenGL 4.3+, PS4, XB1")
@@ -140,7 +140,7 @@ bSupportsNNEShaders=true
 MaxSamplers=32
 
 [PreviewPlatform PCD3D_SM6]
-PlatformName=Windows
+PlatformName=PCD3D_SM6
 ShaderFormat=PCD3D_SM6
 ShaderPlatform=PCD3D_SM6
 MenuTooltip=LOCTEXT("PreviewMenuTooltip_PCD3D_SM6", "Windows using SM6 profile")
@@ -164,7 +164,7 @@ FriendlyName=LOCTEXT("FriendlyShaderPlatformName_D3D_Mobile", "PC D3D Mobile")
 SupportsMultiViewport=RuntimeDependent
 
 [PreviewPlatform PCD3D_ES3_1]
-PlatformName=Windows
+PlatformName=PCD3D_ES3_1
 ShaderFormat=PCD3D_ES31
 ShaderPlatform=PCD3D_ES3_1
 MenuTooltip=LOCTEXT("PreviewMenuTooltip_PCD3D_Mobile", "Windows using Mobile profile.")
diff --git a/Engine/Source/Editor/LevelEditor/Private/LevelEditorActions.cpp b/Engine/Source/Editor/LevelEditor/Private/LevelEditorActions.cpp
--- a/Engine/Source/Editor/LevelEditor/Private/LevelEditorActions.cpp
+++ b/Engine/Source/Editor/LevelEditor/Private/LevelEditorActions.cpp
@@ -3930,7 +3930,15 @@ void FLevelEditorCommands::RegisterCommands()
 		if (!IsRunningCommandlet() && !GUsingNullRHI)
 		{
 			EShaderPlatform ShaderPlatform = FDataDrivenShaderPlatformInfo::GetShaderPlatformFromName(Item.PreviewShaderPlatformName);
-			FriendlyNameBuilder.AppendLine(FDataDrivenShaderPlatformInfo::GetFriendlyName(ShaderPlatform));
+			if (ShaderPlatform == SP_NumPlatforms)
+			{
+				// if the shader platform isn't compiled in, we don't have a friendly name available, so use ugly name
+				FriendlyNameBuilder.AppendLine(FText::FromName(Item.PreviewShaderPlatformName));
+			}
+			else
+			{
+				FriendlyNameBuilder.AppendLine(FDataDrivenShaderPlatformInfo::GetFriendlyName(ShaderPlatform));
+			}
 			if (FDataDrivenShaderPlatformInfo::GetShaderPlatformFromName(Item.ShaderPlatformToPreview) == GMaxRHIShaderPlatform)
 			{
 				FriendlyNameBuilder.AppendLine(NSLOCTEXT("PreviewPlatform", "PreviewMenuText_DisablePreview", "(Disable Preview)"));
diff --git a/Engine/Source/Runtime/Core/Private/Misc/DataDrivenPlatformInfoRegistry.cpp b/Engine/Source/Runtime/Core/Private/Misc/DataDrivenPlatformInfoRegistry.cpp
--- a/Engine/Source/Runtime/Core/Private/Misc/DataDrivenPlatformInfoRegistry.cpp
+++ b/Engine/Source/Runtime/Core/Private/Misc/DataDrivenPlatformInfoRegistry.cpp
@@ -277,11 +277,6 @@ static void ParsePreviewPlatforms(const FConfigFile& IniFile)
 				FName PlatformName = *GetSectionString(Section.Value, FName("PlatformName"));
 				checkf(PlatformName != NAME_None, TEXT("DataDrivenPlatformInfo section [%s] must specify a PlatformName"), *SectionName);
 				
-				if (!FDataDrivenPlatformInfoRegistry::GetPlatformInfo(PlatformName).bEnabledForUse)
-				{
-					continue;
-				}
-
 				FPreviewPlatformMenuItem Item;
 				Item.PlatformName = PlatformName;
 				Item.PreviewShaderPlatformName = PreviewPlatformName;
