Add option to create a WP level with NewLevel.

diff --git a/Engine/Source/Editor/LevelEditor/Private/LevelEditorSubsystem.cpp b/Engine/Source/Editor/LevelEditor/Private/LevelEditorSubsystem.cpp
--- a/Engine/Source/Editor/LevelEditor/Private/LevelEditorSubsystem.cpp
+++ b/Engine/Source/Editor/LevelEditor/Private/LevelEditorSubsystem.cpp
@@ -333,7 +333,7 @@ bool ULevelEditorSubsystem::GetAllowsCinematicControl(FName ViewportConfigKey)
  *
  **/
 
-bool ULevelEditorSubsystem::NewLevel(const FString& AssetPath)
+bool ULevelEditorSubsystem::NewLevel(const FString& AssetPath, bool bIsPartitionedWorld)
 {
 	TGuardValue<bool> UnattendedScriptGuard(GIsRunningUnattendedScript, true);
 
@@ -369,7 +369,7 @@ bool ULevelEditorSubsystem::NewLevel(const FString& AssetPath)
 		return false;
 	}
 
-	UWorld* World = GEditor->NewMap();
+	UWorld* World = GEditor->NewMap(bIsPartitionedWorld);
 	if (!World)
 	{
 		UE_LOG(LevelEditorSubsystem, Error, TEXT("NewLevel. Failed to create the new level."));
diff --git a/Engine/Source/Editor/LevelEditor/Public/LevelEditorSubsystem.h b/Engine/Source/Editor/LevelEditor/Public/LevelEditorSubsystem.h
--- a/Engine/Source/Editor/LevelEditor/Public/LevelEditorSubsystem.h
+++ b/Engine/Source/Editor/LevelEditor/Public/LevelEditorSubsystem.h
@@ -86,10 +86,11 @@ public:
 	 * Close the current Persistent Level (without saving it). Create a new blank Level and save it. Load the new created level.
 	 * @param	AssetPath		Asset Path of where the level will be saved.
 	 *		ie. /Game/MyFolder/MyAsset
+	 * @param   bIsPartitionedWorld	If true, new map is partitioned.
 	 * @return	True if the operation succeeds.
 	 */
 	UFUNCTION(BlueprintCallable, Category = "Editor Scripting | Level Utility")
-	bool NewLevel(const FString& AssetPath);
+	bool NewLevel(const FString& AssetPath, bool bIsPartitionedWorld = false);
 
 	/**
 	 * Close the current Persistent Level (without saving it). Create a new Level base on another level and save it. Load the new created level.
