Level editor viewport: Save level viewport infos when closing the editor
- We currently save the viewport camera infos on map change and map teardown.
- With this change, we will also save the camera infos when closing the editor, as it's equivalent to a map teardown.

diff --git a/Engine/Source/Editor/LevelEditor/Private/SLevelViewport.cpp b/Engine/Source/Editor/LevelEditor/Private/SLevelViewport.cpp
--- a/Engine/Source/Editor/LevelEditor/Private/SLevelViewport.cpp
+++ b/Engine/Source/Editor/LevelEditor/Private/SLevelViewport.cpp
@@ -118,6 +118,69 @@ namespace SLevelViewportPIEAnimation
 	float const MouseControlLabelFadeout = 5.0f;
 }
 
+namespace UE::SLevelViewport::Internal
+{
+	bool SaveViewportInfo(UWorld* World, FLevelEditorViewportClient* LevelEditorViewportClient, ULevelEditorViewportSettings* LevelEditorViewportSettings)
+	{
+		if (!World || !LevelEditorViewportClient || !LevelEditorViewportSettings)
+		{
+			return false;
+		}
+
+		// there could potentially be more than one of the same viewport type.  This effectively takes the last one of a specific type
+		World->EditorViews[LevelEditorViewportClient->ViewportType] =
+			FLevelViewportInfo(
+				LevelEditorViewportClient->GetViewLocation(),
+				LevelEditorViewportClient->GetViewRotation(),
+				LevelEditorViewportClient->GetOrthoZoom());
+
+		LevelEditorViewportSettings->EditorViews.FindOrAdd(World).LevelViewportsInfo = World->EditorViews;
+		LevelEditorViewportSettings->SaveConfig();
+
+		return true;
+	}
+
+	bool LoadViewportInfo(UWorld* World, FLevelEditorViewportClient* LevelEditorViewportClient, ULevelEditorViewportSettings* LevelEditorViewportSettings)
+	{
+		if (!World || !LevelEditorViewportClient || !LevelEditorViewportSettings)
+		{
+			return false;
+		}
+
+		if (FLevelEditorViewporEditorViews* PerUserEditorViews = LevelEditorViewportSettings->EditorViews.Find(World))
+		{
+			World->EditorViews = PerUserEditorViews->LevelViewportsInfo;
+		}
+
+		if (World->EditorViews[LevelEditorViewportClient->ViewportType].CamOrthoZoom == 0.0f)
+		{
+			World->EditorViews[LevelEditorViewportClient->ViewportType].CamOrthoZoom = DEFAULT_ORTHOZOOM;
+		}
+
+		LevelEditorViewportClient->ResetCamera();
+
+		bool bInitializedOrthoViewport = false;
+		for (int32 ViewportType = 0; ViewportType < LVT_MAX; ViewportType++)
+		{
+			if (ViewportType == LVT_Perspective || !bInitializedOrthoViewport)
+			{
+				LevelEditorViewportClient->SetInitialViewTransform(
+					static_cast<ELevelViewportType>(ViewportType),
+					World->EditorViews[ViewportType].CamPosition,
+					World->EditorViews[ViewportType].CamRotation,
+					World->EditorViews[ViewportType].CamOrthoZoom);
+
+				if (ViewportType != LVT_Perspective)
+				{
+					bInitializedOrthoViewport = true;
+				}
+			}
+		}
+
+		return true;
+	}
+}
+
 class FLevelViewportDropContextMenuImpl
 {
 public:
@@ -163,6 +226,7 @@ SLevelViewport::~SLevelViewport()
 	if(UObjectInitialized())
 	{
 		GEngine->OnLevelActorDeleted().RemoveAll(this);
+		GEngine->OnEditorClose().RemoveAll( this );
 
 		GetMutableDefault<ULevelEditorViewportSettings>()->OnSettingChanged().RemoveAll(this);
 	}
@@ -262,6 +326,7 @@ void SLevelViewport::Construct(const FArguments& InArgs, const FAssetEditorViewp
 	LevelEditor.OnMapChanged().AddRaw( this, &SLevelViewport::OnMapChanged );
 
 	GEngine->OnLevelActorDeleted().AddRaw( this, &SLevelViewport::OnLevelActorsRemoved );
+	GEngine->OnEditorClose().AddRaw( this, &SLevelViewport::OnEditorClose );
 
 	FEditorDelegates::PostPIEStarted.AddSP(this, &SLevelViewport::TransitionToPIE);
 	FEditorDelegates::PrePIEEnded.AddSP(this, &SLevelViewport::TransitionFromPIE);
@@ -1206,58 +1271,22 @@ TSharedRef< SWidget > SLevelViewport::BuildViewportDragDropContextMenu()
 
 void SLevelViewport::OnMapChanged( UWorld* World, EMapChangeType MapChangeType )
 {
+	using namespace UE::SLevelViewport::Internal;
+
 	if( World && ( ( World == GetWorld() ) || ( World->EditorViews[LevelViewportClient->ViewportType].CamUpdated ) ) )
 	{
 		if( MapChangeType == EMapChangeType::LoadMap )
 		{
-			if (FLevelEditorViewporEditorViews* PerUserEditorViews = GetMutableDefault<ULevelEditorViewportSettings>()->EditorViews.Find(World))
-			{
-				World->EditorViews = PerUserEditorViews->LevelViewportsInfo;
-			}
-
-			if (World->EditorViews[LevelViewportClient->ViewportType].CamOrthoZoom == 0.0f)
-			{
-				World->EditorViews[LevelViewportClient->ViewportType].CamOrthoZoom = DEFAULT_ORTHOZOOM;
-			}
-	
 			ResetNewLevelViewFlags();
-			LevelViewportClient->ResetCamera();
-
-			bool bInitializedOrthoViewport = false;
-			for (int32 ViewportType = 0; ViewportType < LVT_MAX; ViewportType++)
-			{
-				if (ViewportType == LVT_Perspective || !bInitializedOrthoViewport)
-				{
-					LevelViewportClient->SetInitialViewTransform(
-						static_cast<ELevelViewportType>(ViewportType),
-						World->EditorViews[ViewportType].CamPosition,
-						World->EditorViews[ViewportType].CamRotation,
-						World->EditorViews[ViewportType].CamOrthoZoom);
-
-					if (ViewportType != LVT_Perspective)
-					{
-						bInitializedOrthoViewport = true;
-					}
-				}
-			}
+			LoadViewportInfo(World, LevelViewportClient.Get(), GetMutableDefault<ULevelEditorViewportSettings>());
 		}
 		else if( (MapChangeType == EMapChangeType::SaveMap) || (MapChangeType == EMapChangeType::TearDownWorld))
 		{
-			//@todo there could potentially be more than one of the same viewport type.  This effectively takes the last one of a specific type
-			World->EditorViews[LevelViewportClient->ViewportType] = 
-				FLevelViewportInfo( 
-					LevelViewportClient->GetViewLocation(),
-					LevelViewportClient->GetViewRotation(), 
-					LevelViewportClient->GetOrthoZoom() );
-
-			GetMutableDefault<ULevelEditorViewportSettings>()->EditorViews.FindOrAdd(World).LevelViewportsInfo = World->EditorViews;
-			GetMutableDefault<ULevelEditorViewportSettings>()->SaveConfig();
+			SaveViewportInfo(World, LevelViewportClient.Get(), GetMutableDefault<ULevelEditorViewportSettings>());
 		}
 		else if( MapChangeType == EMapChangeType::NewMap )
-		{		
-		
+		{
 			ResetNewLevelViewFlags();
-
 			LevelViewportClient->ResetViewForNewMap();
 		}
 		World->EditorViews[LevelViewportClient->ViewportType].CamUpdated = false;
@@ -1268,6 +1297,11 @@ void SLevelViewport::OnMapChanged( UWorld* World, EMapChangeType MapChangeType )
 	}
 }
 
+void SLevelViewport::OnEditorClose()
+{
+	UE::SLevelViewport::Internal::SaveViewportInfo(GetWorld(), LevelViewportClient.Get(), GetMutableDefault<ULevelEditorViewportSettings>());
+}
+
 void SLevelViewport::OnLevelActorsRemoved(AActor* InActor)
 {
 	// Kill any existing actor previews that have expired
diff --git a/Engine/Source/Editor/LevelEditor/Public/SLevelViewport.h b/Engine/Source/Editor/LevelEditor/Public/SLevelViewport.h
--- a/Engine/Source/Editor/LevelEditor/Public/SLevelViewport.h
+++ b/Engine/Source/Editor/LevelEditor/Public/SLevelViewport.h
@@ -685,6 +685,8 @@ private:
 
 	/** Called in response to an actor being deleted in the level */
 	void OnLevelActorsRemoved(AActor* InActor);
+	
+	void OnEditorClose();
 
 	/** Gets whether the locked icon should be shown in the viewport because it is locked to an actor */
 	EVisibility GetLockedIconVisibility() const;
