TypedElementSelection: AllowSubRootSelection

Applies to Actor selection and allows selecting a rooted actor if it implements SupportsSubRootSelection() and returns true
Will be used in next CL for Level Instance sub selection in Outliner

#rnx

diff --git a/Engine/Source/Editor/LevelEditor/Private/Elements/Actor/ActorElementLevelEditorSelectionCustomization.cpp b/Engine/Source/Editor/LevelEditor/Private/Elements/Actor/ActorElementLevelEditorSelectionCustomization.cpp
--- a/Engine/Source/Editor/LevelEditor/Private/Elements/Actor/ActorElementLevelEditorSelectionCustomization.cpp
+++ b/Engine/Source/Editor/LevelEditor/Private/Elements/Actor/ActorElementLevelEditorSelectionCustomization.cpp
@@ -170,10 +170,13 @@ bool FActorElementLevelEditorSelectionCustomization::SelectActorElement(const TT
 		}
 	}
 
-	// If trying to select an actor, use this actors root selection actor instead (if it has one)
-	if (AActor* RootSelection = Actor->GetRootSelectionParent())
+	// If trying to select an actor, use this actors root selection actor instead (if it has one), unless actor supports being selected with a root and options allow it
+	if (!(InSelectionOptions.AllowSubRootSelection() && Actor->SupportsSubRootSelection()))
 	{
-		Actor = RootSelection;
+		if (AActor* RootSelection = Actor->GetRootSelectionParent())
+		{
+			Actor = RootSelection;
+		}
 	}
 
 	bool bSelectionChanged = false;
@@ -242,10 +245,14 @@ bool FActorElementLevelEditorSelectionCustomization::DeselectActorElement(const
 
 	bool bSelectionChanged = false;
 
-	// If trying to deselect an actor, use this actors root selection actor instead (if it has one)
-	if (AActor* RootSelection = Actor->GetRootSelectionParent())
+	// If Selection options allows selection of sub root actors and sub actor is selected directly, avoid getting root selection parent
+	if (!(InSelectionOptions.AllowSubRootSelection() && Actor->SupportsSubRootSelection() && InActorSelectionHandle.IsElementSelected(InSelectionSet, FTypedElementIsSelectedOptions().SetAllowIndirect(false))))
 	{
-		Actor = RootSelection;
+		// If trying to deselect an actor, use this actors root selection actor instead (if it has one)
+		if (AActor* RootSelection = Actor->GetRootSelectionParent())
+		{
+			Actor = RootSelection;
+		}
 	}
 
 	if (UActorGroupingUtils::IsGroupingActive() && InSelectionOptions.AllowGroups())
diff --git a/Engine/Source/Editor/UnrealEd/Private/EditorSelectUtils.cpp b/Engine/Source/Editor/UnrealEd/Private/EditorSelectUtils.cpp
--- a/Engine/Source/Editor/UnrealEd/Private/EditorSelectUtils.cpp
+++ b/Engine/Source/Editor/UnrealEd/Private/EditorSelectUtils.cpp
@@ -571,7 +571,7 @@ bool UUnrealEdEngine::CanSelectActor(AActor* Actor, bool bInSelected, bool bSele
 
 void UUnrealEdEngine::SelectActor(AActor* Actor, bool bInSelected, bool bNotify, bool bSelectEvenIfHidden, bool bForceRefresh)
 {
-	if (!Actor || Actor->GetRootSelectionParent() != nullptr)
+	if (!Actor || (Actor->GetRootSelectionParent() != nullptr && !Actor->SupportsSubRootSelection()))
 	{
 		return;
 	}
@@ -589,7 +589,8 @@ void UUnrealEdEngine::SelectActor(AActor* Actor, bool bInSelected, bool bNotify,
 		const FTypedElementSelectionOptions SelectionOptions = FTypedElementSelectionOptions()
 			.SetAllowHidden(bSelectEvenIfHidden)
 			.SetWarnIfLocked(true)
-			.SetAllowLegacyNotifications(false);
+			.SetAllowLegacyNotifications(false)
+			.SetAllowSubRootSelection(Actor->SupportsSubRootSelection());
 
 		const bool bSelectionChanged = bInSelected
 			? SelectionSet->SelectElement(UEngineElementsLibrary::AcquireEditorActorElementHandle(Actor), SelectionOptions)
diff --git a/Engine/Source/Runtime/Engine/Classes/GameFramework/Actor.h b/Engine/Source/Runtime/Engine/Classes/GameFramework/Actor.h
--- a/Engine/Source/Runtime/Engine/Classes/GameFramework/Actor.h
+++ b/Engine/Source/Runtime/Engine/Classes/GameFramework/Actor.h
@@ -2945,6 +2945,9 @@ public:
 	/** Returns top most selection parent */
 	ENGINE_API virtual AActor* GetRootSelectionParent() const;
 
+	/** Returns true if actor can be selected as a sub selection of its root selection parent */
+	ENGINE_API virtual bool SupportsSubRootSelection() const { return false; }
+
 	/** Returns if actor or selection parent is selected */
 	ENGINE_API bool IsActorOrSelectionParentSelected() const;
 
diff --git a/Engine/Source/Runtime/TypedElementRuntime/Private/Elements/Framework/TypedElementSelectionSet.cpp b/Engine/Source/Runtime/TypedElementRuntime/Private/Elements/Framework/TypedElementSelectionSet.cpp
--- a/Engine/Source/Runtime/TypedElementRuntime/Private/Elements/Framework/TypedElementSelectionSet.cpp
+++ b/Engine/Source/Runtime/TypedElementRuntime/Private/Elements/Framework/TypedElementSelectionSet.cpp
@@ -330,10 +330,13 @@ bool UTypedElementSelectionSet::ClearSelection(const FTypedElementSelectionOptio
 		// Take a copy of the currently selected elements to avoid mutating the selection set while iterating
 		TArray<FTypedElementHandle, TInlineAllocator<8>> ElementsCopy;
 		ElementList->GetElementHandles(ElementsCopy);
+		FTypedElementSelectionOptions CopySelectionOptions(InSelectionOptions);
 
+		// Always support SubRootSelection when clearing
+		CopySelectionOptions.SetAllowSubRootSelection(true);
 		for (const FTypedElementHandle& ElementHandle : ElementsCopy)
 		{
-			bSelectionChanged |= DeselectElement(ElementHandle, InSelectionOptions);
+			bSelectionChanged |= DeselectElement(ElementHandle, CopySelectionOptions);
 		}
 	}
 
diff --git a/Engine/Source/Runtime/TypedElementRuntime/Public/Elements/Interfaces/TypedElementSelectionInterface.h b/Engine/Source/Runtime/TypedElementRuntime/Public/Elements/Interfaces/TypedElementSelectionInterface.h
--- a/Engine/Source/Runtime/TypedElementRuntime/Public/Elements/Interfaces/TypedElementSelectionInterface.h
+++ b/Engine/Source/Runtime/TypedElementRuntime/Public/Elements/Interfaces/TypedElementSelectionInterface.h
@@ -72,6 +72,9 @@ public:
 	FTypedElementSelectionOptions& SetWarnIfLocked(const bool InWarnIfLocked) { bWarnIfLocked = InWarnIfLocked; return *this; }
 	bool WarnIfLocked() const { return bWarnIfLocked; }
 
+	FTypedElementSelectionOptions& SetAllowSubRootSelection(const bool InAllowSubRootSelectioin) { bAllowSubRootSelection = InAllowSubRootSelectioin; return *this; }
+	bool AllowSubRootSelection() const { return bAllowSubRootSelection; }
+
 	FTypedElementSelectionOptions& SetChildElementInclusionMethod(const ETypedElementChildInclusionMethod InChildElementInclusionMethod) { ChildElementInclusionMethod = InChildElementInclusionMethod; return *this; }
 	ETypedElementChildInclusionMethod GetChildElementInclusionMethod() const { return ChildElementInclusionMethod; }
 
@@ -91,6 +94,9 @@ private:
 	UPROPERTY(BlueprintReadWrite, Category="TypedElementInterfaces|Selection|SelectionOptions", meta=(AllowPrivateAccess=true))
 	bool bWarnIfLocked = false;
 
+	UPROPERTY(BlueprintReadWrite, Category = "TypedElementInterfaces|Selection|SelectionOptions", meta = (AllowPrivateAccess = true))
+	bool bAllowSubRootSelection = false;
+
 	UPROPERTY(BlueprintReadWrite, Category="TypedElementInterfaces|Selection|SelectionOptions", meta=(AllowPrivateAccess=true))
 	ETypedElementChildInclusionMethod ChildElementInclusionMethod = ETypedElementChildInclusionMethod::None;
 
