[Actor Context Menu] Fix an issue where snapping and transforms were not respecting lock location on the actor.

diff --git a/Engine/Source/Editor/LevelEditor/Private/LevelEditor.cpp b/Engine/Source/Editor/LevelEditor/Private/LevelEditor.cpp
--- a/Engine/Source/Editor/LevelEditor/Private/LevelEditor.cpp
+++ b/Engine/Source/Editor/LevelEditor/Private/LevelEditor.cpp
@@ -835,31 +835,31 @@ void FLevelEditorModule::BindGlobalLevelEditorCommands()
 	ActionList.MapAction(
 		Commands.SnapOriginToGrid,
 		FExecuteAction::CreateStatic(&FLevelEditorActionCallbacks::MoveElementsToGrid_Clicked, /*bAlign*/false, /*bPerElement*/false),
-		FCanExecuteAction::CreateStatic(&FLevelEditorActionCallbacks::ElementSelected_CanExecute)
+		FCanExecuteAction::CreateStatic(&FLevelEditorActionCallbacks::ElementSelected_CanExecuteMove)
 		);
 
 	ActionList.MapAction(
 		Commands.SnapOriginToGridPerActor,
 		FExecuteAction::CreateStatic(&FLevelEditorActionCallbacks::MoveElementsToGrid_Clicked, /*bAlign*/false, /*bPerElement*/true),
-		FCanExecuteAction::CreateStatic(&FLevelEditorActionCallbacks::ElementSelected_CanExecute)
+		FCanExecuteAction::CreateStatic(&FLevelEditorActionCallbacks::ElementSelected_CanExecuteMove)
 		);
 	
 	ActionList.MapAction(
 		Commands.AlignOriginToGrid,
 		FExecuteAction::CreateStatic(&FLevelEditorActionCallbacks::MoveElementsToGrid_Clicked, /*bAlign*/true, /*bPerElement*/false),
-		FCanExecuteAction::CreateStatic(&FLevelEditorActionCallbacks::ElementSelected_CanExecute)
+		FCanExecuteAction::CreateStatic(&FLevelEditorActionCallbacks::ElementSelected_CanExecuteMove)
 		);
 
 	ActionList.MapAction(
 		Commands.SnapOriginToActor,
 		FExecuteAction::CreateStatic(&FLevelEditorActionCallbacks::MoveElementsToElement_Clicked, /*bAlign*/false),
-		FCanExecuteAction::CreateStatic(&FLevelEditorActionCallbacks::ElementsSelected_CanExecute)
+		FCanExecuteAction::CreateStatic(&FLevelEditorActionCallbacks::ElementsSelected_CanExecuteMove)
 		);
 	
 	ActionList.MapAction(
 		Commands.AlignOriginToActor,
 		FExecuteAction::CreateStatic(&FLevelEditorActionCallbacks::MoveElementsToElement_Clicked, /*bAlign*/true),
-		FCanExecuteAction::CreateStatic(&FLevelEditorActionCallbacks::ElementsSelected_CanExecute)
+		FCanExecuteAction::CreateStatic(&FLevelEditorActionCallbacks::ElementsSelected_CanExecuteMove)
 		);
 
 	ActionList.MapAction(
@@ -906,96 +906,96 @@ void FLevelEditorModule::BindGlobalLevelEditorCommands()
 	ActionList.MapAction(
 		Commands.SnapToFloor,
 		FExecuteAction::CreateStatic(&FLevelEditorActionCallbacks::SnapToFloor_Clicked, /*bAlign*/false, /*bUseLineTrace*/false, /*bUseBounds*/false, /*bUsePivot*/false),
-		FCanExecuteAction::CreateStatic(&FLevelEditorActionCallbacks::ElementSelected_CanExecute)
+		FCanExecuteAction::CreateStatic(&FLevelEditorActionCallbacks::ElementSelected_CanExecuteMove)
 		);
 
 	ActionList.MapAction(
 		Commands.AlignToFloor,
 		FExecuteAction::CreateStatic(&FLevelEditorActionCallbacks::SnapToFloor_Clicked, /*bAlign*/true, /*bUseLineTrace*/false, /*bUseBounds*/false, /*bUsePivot*/false),
-		FCanExecuteAction::CreateStatic(&FLevelEditorActionCallbacks::ElementSelected_CanExecute)
+		FCanExecuteAction::CreateStatic(&FLevelEditorActionCallbacks::ElementSelected_CanExecuteMove)
 		);
 
 	ActionList.MapAction(
 		Commands.SnapPivotToFloor,
 		FExecuteAction::CreateStatic(&FLevelEditorActionCallbacks::SnapToFloor_Clicked, /*bAlign*/false, /*bUseLineTrace*/true, /*bUseBounds*/false, /*bUsePivot*/true),
-		FCanExecuteAction::CreateStatic(&FLevelEditorActionCallbacks::ElementSelected_CanExecute)
+		FCanExecuteAction::CreateStatic(&FLevelEditorActionCallbacks::ElementSelected_CanExecuteMove)
 		);
 
 	ActionList.MapAction(
 		Commands.AlignPivotToFloor,
 		FExecuteAction::CreateStatic(&FLevelEditorActionCallbacks::SnapToFloor_Clicked, /*bAlign*/true, /*bUseLineTrace*/true, /*bUseBounds*/false, /*bUsePivot*/true),
-		FCanExecuteAction::CreateStatic(&FLevelEditorActionCallbacks::ElementSelected_CanExecute)
+		FCanExecuteAction::CreateStatic(&FLevelEditorActionCallbacks::ElementSelected_CanExecuteMove)
 		);
 
 	ActionList.MapAction(
 		Commands.SnapBottomCenterBoundsToFloor,
 		FExecuteAction::CreateStatic(&FLevelEditorActionCallbacks::SnapToFloor_Clicked, /*bAlign*/false, /*bUseLineTrace*/true, /*bUseBounds*/true, /*bUsePivot*/false),
-		FCanExecuteAction::CreateStatic(&FLevelEditorActionCallbacks::ElementSelected_CanExecute)
+		FCanExecuteAction::CreateStatic(&FLevelEditorActionCallbacks::ElementSelected_CanExecuteMove)
 		);
 
 	ActionList.MapAction(
 		Commands.AlignBottomCenterBoundsToFloor,
 		FExecuteAction::CreateStatic(&FLevelEditorActionCallbacks::SnapToFloor_Clicked, /*bAlign*/true, /*bUseLineTrace*/true, /*bUseBounds*/true, /*bUsePivot*/false),
-		FCanExecuteAction::CreateStatic(&FLevelEditorActionCallbacks::ElementSelected_CanExecute)
+		FCanExecuteAction::CreateStatic(&FLevelEditorActionCallbacks::ElementSelected_CanExecuteMove)
 		);
 
 	ActionList.MapAction(
 		Commands.SnapToActor,
 		FExecuteAction::CreateStatic(&FLevelEditorActionCallbacks::SnapElementsToElement_Clicked, /*bAlign*/false, /*bUseLineTrace*/false, /*bUseBounds*/false, /*bUsePivot*/false),
-		FCanExecuteAction::CreateStatic(&FLevelEditorActionCallbacks::ElementsSelected_CanExecute)
+		FCanExecuteAction::CreateStatic(&FLevelEditorActionCallbacks::ElementsSelected_CanExecuteMove)
 		);
 
 	ActionList.MapAction(
 		Commands.AlignToActor,
 		FExecuteAction::CreateStatic(&FLevelEditorActionCallbacks::SnapElementsToElement_Clicked, /*bAlign*/true, /*bUseLineTrace*/false, /*bUseBounds*/false, /*bUsePivot*/false),
-		FCanExecuteAction::CreateStatic(&FLevelEditorActionCallbacks::ElementsSelected_CanExecute)
+		FCanExecuteAction::CreateStatic(&FLevelEditorActionCallbacks::ElementsSelected_CanExecuteMove)
 		);
 
 	ActionList.MapAction(
 		Commands.SnapPivotToActor,
 		FExecuteAction::CreateStatic(&FLevelEditorActionCallbacks::SnapElementsToElement_Clicked, /*bAlign*/false, /*bUseLineTrace*/true, /*bUseBounds*/false, /*bUsePivot*/true),
-		FCanExecuteAction::CreateStatic(&FLevelEditorActionCallbacks::ElementsSelected_CanExecute)
+		FCanExecuteAction::CreateStatic(&FLevelEditorActionCallbacks::ElementsSelected_CanExecuteMove)
 		);
 
 	ActionList.MapAction(
 		Commands.AlignPivotToActor,
 		FExecuteAction::CreateStatic(&FLevelEditorActionCallbacks::SnapElementsToElement_Clicked, /*bAlign*/true, /*bUseLineTrace*/true, /*bUseBounds*/false, /*bUsePivot*/true),
-		FCanExecuteAction::CreateStatic(&FLevelEditorActionCallbacks::ElementsSelected_CanExecute)
+		FCanExecuteAction::CreateStatic(&FLevelEditorActionCallbacks::ElementsSelected_CanExecuteMove)
 		);
 
 	ActionList.MapAction(
 		Commands.SnapBottomCenterBoundsToActor,
 		FExecuteAction::CreateStatic(&FLevelEditorActionCallbacks::SnapElementsToElement_Clicked, /*bAlign*/false, /*bUseLineTrace*/true, /*bUseBounds*/true, /*bUsePivot*/false),
-		FCanExecuteAction::CreateStatic(&FLevelEditorActionCallbacks::ElementsSelected_CanExecute)
+		FCanExecuteAction::CreateStatic(&FLevelEditorActionCallbacks::ElementsSelected_CanExecuteMove)
 		);
 
 	ActionList.MapAction(
 		Commands.AlignBottomCenterBoundsToActor,
 		FExecuteAction::CreateStatic(&FLevelEditorActionCallbacks::SnapElementsToElement_Clicked, /*bAlign*/true, /*bUseLineTrace*/true, /*bUseBounds*/true, /*bUsePivot*/false),
-		FCanExecuteAction::CreateStatic(&FLevelEditorActionCallbacks::ElementsSelected_CanExecute)
+		FCanExecuteAction::CreateStatic(&FLevelEditorActionCallbacks::ElementsSelected_CanExecuteMove)
 		);
 
 	ActionList.MapAction(
 		Commands.DeltaTransformToActors, 
 		FExecuteAction::CreateStatic( &FLevelEditorActionCallbacks::DeltaTransform ), 
-		FCanExecuteAction::CreateStatic( &FLevelEditorActionCallbacks::ElementSelected_CanExecute) );
+		FCanExecuteAction::CreateStatic( &FLevelEditorActionCallbacks::ElementSelected_CanExecuteMove) );
 
 	ActionList.MapAction(
 		Commands.MirrorActorX,
 		FExecuteAction::CreateStatic( &FLevelEditorActionCallbacks::ExecuteExecCommand, FString( TEXT("ELEMENT MIRROR X=-1") ) ),
-		FCanExecuteAction::CreateStatic(&FLevelEditorActionCallbacks::ElementSelected_CanExecute)
+		FCanExecuteAction::CreateStatic(&FLevelEditorActionCallbacks::ElementSelected_CanExecuteMove)
 		);
 
 	ActionList.MapAction(
 		Commands.MirrorActorY,
 		FExecuteAction::CreateStatic( &FLevelEditorActionCallbacks::ExecuteExecCommand, FString( TEXT("ELEMENT MIRROR Y=-1") ) ),
-		FCanExecuteAction::CreateStatic(&FLevelEditorActionCallbacks::ElementSelected_CanExecute)
+		FCanExecuteAction::CreateStatic(&FLevelEditorActionCallbacks::ElementSelected_CanExecuteMove)
 		);
 
 	ActionList.MapAction(
 		Commands.MirrorActorZ,
 		FExecuteAction::CreateStatic( &FLevelEditorActionCallbacks::ExecuteExecCommand, FString( TEXT("ELEMENT MIRROR Z=-1") ) ),
-		FCanExecuteAction::CreateStatic(&FLevelEditorActionCallbacks::ElementSelected_CanExecute)
+		FCanExecuteAction::CreateStatic(&FLevelEditorActionCallbacks::ElementSelected_CanExecuteMove)
 		);
 
 	ActionList.MapAction(
diff --git a/Engine/Source/Editor/LevelEditor/Private/LevelEditorActions.cpp b/Engine/Source/Editor/LevelEditor/Private/LevelEditorActions.cpp
--- a/Engine/Source/Editor/LevelEditor/Private/LevelEditorActions.cpp
+++ b/Engine/Source/Editor/LevelEditor/Private/LevelEditorActions.cpp
@@ -3333,8 +3333,13 @@ void FLevelEditorActionCallbacks::SnapTo2DLayer_Clicked()
 
 bool FLevelEditorActionCallbacks::CanSnapTo2DLayer()
 {
+	if (!ElementSelected_CanExecuteMove())
+	{
+		return false;
+	}
+
 	const ULevelEditor2DSettings* Settings = GetDefault<ULevelEditor2DSettings>();
-	return Settings->SnapLayers.IsValidIndex(GetDefault<ULevelEditorViewportSettings>()->ActiveSnapLayerIndex) && (GEditor->GetSelectedActorCount() > 0);
+	return Settings->SnapLayers.IsValidIndex(GetDefault<ULevelEditorViewportSettings>()->ActiveSnapLayerIndex);
 }
 
 void FLevelEditorActionCallbacks::MoveSelectionToDifferent2DLayer_Clicked(bool bGoingUp, bool bForceToTopOrBottom)
@@ -3543,6 +3548,28 @@ bool FLevelEditorActionCallbacks::ElementsSelected_CanExecute()
 	return SelectionSet && SelectionSet->GetNumSelectedElements() > 1;
 }
 
+bool FLevelEditorActionCallbacks::ElementSelected_CanExecuteMove()
+{
+	// TODO: Ideally this would come from some level editor context
+	if (GCurrentLevelEditingViewportClient)
+	{
+		return GCurrentLevelEditingViewportClient->GetElementsToManipulate(false)->Num() > 0;
+	}
+
+	return false;
+}
+
+bool FLevelEditorActionCallbacks::ElementsSelected_CanExecuteMove()
+{
+	// TODO: Ideally this would come from some level editor context
+	if (GCurrentLevelEditingViewportClient)
+	{
+		return GCurrentLevelEditingViewportClient->GetElementsToManipulate(false)->Num() > 1;
+	}
+
+	return false;
+}
+
 void FLevelEditorActionCallbacks::GeometryCollection_SelectAllGeometry()
 {
 	GEditor->Exec(GetWorld(), TEXT("GeometryCollection.SelectAllGeometry"));
diff --git a/Engine/Source/Editor/LevelEditor/Private/LevelEditorContextMenu.cpp b/Engine/Source/Editor/LevelEditor/Private/LevelEditorContextMenu.cpp
--- a/Engine/Source/Editor/LevelEditor/Private/LevelEditorContextMenu.cpp
+++ b/Engine/Source/Editor/LevelEditor/Private/LevelEditorContextMenu.cpp
@@ -1141,25 +1141,24 @@ void FLevelEditorContextMenuImpl::FillActorLevelMenu(UToolMenu* Menu)
 
 void FLevelEditorContextMenuImpl::FillTransformMenu(UToolMenu* Menu)
 {
-	if ( FLevelEditorActionCallbacks::ElementSelected_CanExecute() )
+	if (ULevelEditorContextMenuContext* LevelEditorContext = Menu->FindContext<ULevelEditorContextMenuContext>())
 	{
-		FToolMenuSection& Section = Menu->AddSection("DeltaTransformToActors");
-		Section.AddMenuEntry(FLevelEditorCommands::Get().DeltaTransformToActors);
-	}
+		if (LevelEditorContext->CurrentSelection->GetNumSelectedElements() == 0)
+		{
+			return;
+		}
 
-	{
-		FToolMenuSection& Section = Menu->AddSection("MirrorLock");
+		{
+			FToolMenuSection& Section = Menu->AddSection("DeltaTransformToActors");
+			Section.AddMenuEntry(FLevelEditorCommands::Get().DeltaTransformToActors);
+		}
 
-		// TODO: Need an element API to allow the mirror actions
-		if ( FLevelEditorActionCallbacks::ActorSelected_CanExecute() )
+		if (LevelEditorContext->CurrentSelection->HasSelectedObjects<AActor>())
 		{
+			FToolMenuSection& Section = Menu->AddSection("MirrorLock");
 			Section.AddMenuEntry(FLevelEditorCommands::Get().MirrorActorX);
 			Section.AddMenuEntry(FLevelEditorCommands::Get().MirrorActorY);
 			Section.AddMenuEntry(FLevelEditorCommands::Get().MirrorActorZ);
-		}
-
-		if ( FLevelEditorActionCallbacks::ActorSelected_CanExecute() )
-		{
 			Section.AddMenuEntry(FLevelEditorCommands::Get().LockActorMovement);
 		}
 	}
diff --git a/Engine/Source/Editor/LevelEditor/Public/LevelEditorActions.h b/Engine/Source/Editor/LevelEditor/Public/LevelEditorActions.h
--- a/Engine/Source/Editor/LevelEditor/Public/LevelEditorActions.h
+++ b/Engine/Source/Editor/LevelEditor/Public/LevelEditorActions.h
@@ -1403,14 +1403,28 @@ public:
 	 * Checks to see if at least one element is selected
 	 *	@return true if it can execute.
 	 */
+	UE_DEPRECATED(5.4, "Use ElementSelected_CanExecuteMove. If you only needed to verify if there is at least one item in the selection set, use the selection set directly from the level editor.")
 	static bool ElementSelected_CanExecute();
 
 	/**
 	 * Checks to see if multiple elements are selected
 	 *	@return true if it can execute.
 	 */
+	UE_DEPRECATED(5.4, "Use ElementsSelected_CanExecuteMove. If you only needed to verify if there is more than one item in the selection set, use the selection set directly from the level editor.")
 	static bool ElementsSelected_CanExecute();
 
+	/**
+	 * Checks to see if at least one element is selected that can be translated
+	 *	@return true if it can execute.
+	 */
+	static bool ElementSelected_CanExecuteMove();
+
+	/**
+	 * Checks to see if multiple elements are selected that can be translated
+	 *	@return true if it can execute.
+	 */
+	static bool ElementsSelected_CanExecuteMove();
+
 	/** Called when 'Open Merge Actor' is clicked */
 	static void OpenMergeActor_Clicked();
 
diff --git a/Engine/Source/Editor/UnrealEd/Private/EditorEngine.cpp b/Engine/Source/Editor/UnrealEd/Private/EditorEngine.cpp
--- a/Engine/Source/Editor/UnrealEd/Private/EditorEngine.cpp
+++ b/Engine/Source/Editor/UnrealEd/Private/EditorEngine.cpp
@@ -3280,6 +3280,14 @@ void UEditorEngine::ToggleSelectedActorMovementLock()
 		LevelDirtyCallback.Request();
 	}
 
+	// Update the editability status in the active viewport, which will update the gizmos
+	if (GCurrentLevelEditingViewportClient)
+	{
+		constexpr bool bForceCachedElementRefresh = true;
+		GCurrentLevelEditingViewportClient->GetElementsToManipulate(bForceCachedElementRefresh);
+	}
+	RedrawLevelEditingViewports(false);
+
 	bCheckForLockActors = true;
 }
 
