Fix rendering issues when switching between preview modes of the same feature level with different early-z pass settings.

In this case the value of FScene::EarlyZPassMode was not updated to the new value.

diff --git a/Engine/Source/Editor/LevelEditor/Private/SLevelEditor.cpp b/Engine/Source/Editor/LevelEditor/Private/SLevelEditor.cpp
--- a/Engine/Source/Editor/LevelEditor/Private/SLevelEditor.cpp
+++ b/Engine/Source/Editor/LevelEditor/Private/SLevelEditor.cpp
@@ -243,6 +243,11 @@ void SLevelEditor::Construct( const SLevelEditor::FArguments& InArgs)
 			World->ChangeFeatureLevel(NewFeatureLevel);
 		});
 
+	PreviewPlatformChangedHandle = GEditor->OnPreviewPlatformChanged().AddLambda([this]()
+		{
+			World->ShaderPlatformChanged();
+		});
+
 	FEditorDelegates::MapChange.AddRaw(this, &SLevelEditor::HandleEditorMapChange);
 	FEditorDelegates::OnAssetsDeleted.AddRaw(this, &SLevelEditor::HandleAssetsDeleted);
 	HandleEditorMapChange(MapChangeEventFlags::NewMap);
@@ -406,6 +411,7 @@ SLevelEditor::~SLevelEditor()
 	if (GEngine)
 	{
 		CastChecked<UEditorEngine>(GEngine)->OnPreviewFeatureLevelChanged().Remove(PreviewFeatureLevelChangedHandle);
+		CastChecked<UEditorEngine>(GEngine)->OnPreviewPlatformChanged().Remove(PreviewPlatformChangedHandle);
 	}
 
 	if (GEditor)
diff --git a/Engine/Source/Editor/LevelEditor/Private/SLevelEditor.h b/Engine/Source/Editor/LevelEditor/Private/SLevelEditor.h
--- a/Engine/Source/Editor/LevelEditor/Private/SLevelEditor.h
+++ b/Engine/Source/Editor/LevelEditor/Private/SLevelEditor.h
@@ -349,6 +349,9 @@ private:
 	/** Handle to the registered OnPreviewFeatureLevelChanged delegate. */
 	FDelegateHandle PreviewFeatureLevelChangedHandle;
 
+	/** Handle to the registered OnPreviewPlatformChanged delegate. */
+	FDelegateHandle PreviewPlatformChangedHandle;
+
 	/** Handle to the registered OnLevelActorDeleted delegate */
 	FDelegateHandle LevelActorDeletedHandle;
 
diff --git a/Engine/Source/Runtime/Engine/Classes/Engine/World.h b/Engine/Source/Runtime/Engine/Classes/Engine/World.h
--- a/Engine/Source/Runtime/Engine/Classes/Engine/World.h
+++ b/Engine/Source/Runtime/Engine/Classes/Engine/World.h
@@ -1458,6 +1458,8 @@ public:
 	/** Change the feature level that this world is current rendering with */
 	void ChangeFeatureLevel(ERHIFeatureLevel::Type InFeatureLevel, bool bShowSlowProgressDialog = true);
 
+	void ShaderPlatformChanged();
+
 	void RecreateScene(ERHIFeatureLevel::Type InFeatureLevel, bool bBroadcastChange = true);
 
 	/** Recreate the editor world's FScene with a null scene interface to drop extra GPU memory during PIE */
diff --git a/Engine/Source/Runtime/Engine/Private/World.cpp b/Engine/Source/Runtime/Engine/Private/World.cpp
--- a/Engine/Source/Runtime/Engine/Private/World.cpp
+++ b/Engine/Source/Runtime/Engine/Private/World.cpp
@@ -8967,6 +8967,14 @@ void UWorld::ChangeFeatureLevel(ERHIFeatureLevel::Type InFeatureLevel, bool bSho
 	}
 }
 
+void UWorld::ShaderPlatformChanged()
+{
+	if (Scene)
+	{
+		Scene->UpdateEarlyZPassMode();
+	}
+}
+
 void UWorld::RecreateScene(ERHIFeatureLevel::Type InFeatureLevel, bool bBroadcastChange)
 {
 	if (Scene)
diff --git a/Engine/Source/Runtime/Engine/Public/SceneInterface.h b/Engine/Source/Runtime/Engine/Public/SceneInterface.h
--- a/Engine/Source/Runtime/Engine/Public/SceneInterface.h
+++ b/Engine/Source/Runtime/Engine/Public/SceneInterface.h
@@ -633,6 +633,8 @@ public:
 
 	virtual bool IsEditorScene() const { return false; }
 
+	virtual void UpdateEarlyZPassMode() {}
+
 	ERHIFeatureLevel::Type GetFeatureLevel() const { return FeatureLevel; }
 
 	ENGINE_API EShaderPlatform GetShaderPlatform() const;
diff --git a/Engine/Source/Runtime/Renderer/Private/ScenePrivate.h b/Engine/Source/Runtime/Renderer/Private/ScenePrivate.h
--- a/Engine/Source/Runtime/Renderer/Private/ScenePrivate.h
+++ b/Engine/Source/Runtime/Renderer/Private/ScenePrivate.h
@@ -3373,6 +3373,8 @@ public:
 	virtual void UpdatePrimitiveSelectedState_RenderThread(const FPrimitiveSceneInfo* PrimitiveSceneInfo, bool bIsSelected) override;
 	virtual void UpdatePrimitiveVelocityState_RenderThread(FPrimitiveSceneInfo* PrimitiveSceneInfo, bool bIsBeingMoved) override;
 
+	virtual void UpdateEarlyZPassMode() override;
+
 	virtual void Release() override;
 	virtual UWorld* GetWorld() const override { return World; }
 
@@ -3387,8 +3389,6 @@ public:
 
 	int64 GetCachedWholeSceneShadowMapsSize() const;
 
-	void UpdateEarlyZPassMode();
-
 	/**
 	 * Get the default base pass depth stencil access
 	 */
