Editor: Add Allow Cinematic Control to FLevelEditorViewportInstanceSettings so that the state is stored/restored

diff --git a/Engine/Source/Editor/LevelEditor/Private/SLevelViewport.cpp b/Engine/Source/Editor/LevelEditor/Private/SLevelViewport.cpp
--- a/Engine/Source/Editor/LevelEditor/Private/SLevelViewport.cpp
+++ b/Engine/Source/Editor/LevelEditor/Private/SLevelViewport.cpp
@@ -570,8 +570,8 @@ void SLevelViewport::ConstructLevelEditorViewportClient(FLevelEditorViewportInst
 	{
 		LevelViewportClient->SetViewLocation( EditorViewportDefs::DefaultPerspectiveViewLocation );
 		LevelViewportClient->SetViewRotation( EditorViewportDefs::DefaultPerspectiveViewRotation );
-		LevelViewportClient->SetAllowCinematicControl(true);
 	}
+	LevelViewportClient->SetAllowCinematicControl(ViewportInstanceSettings.bAllowCinematicControl);
 	LevelViewportClient->SetRealtime(ViewportInstanceSettings.bIsRealtime);
 	LevelViewportClient->SetShowStats(ViewportInstanceSettings.bShowOnScreenStats);
 	if (ViewportInstanceSettings.bShowFPS_DEPRECATED)
@@ -2314,6 +2314,7 @@ void SLevelViewport::SaveConfig(const FString& ConfigName) const
 		ViewportInstanceSettings.GPUSkinCacheVisualizationMode = LevelViewportClient->CurrentGPUSkinCacheVisualizationMode;
 		ViewportInstanceSettings.ExposureSettings = LevelViewportClient->ExposureSettings;
 		ViewportInstanceSettings.FOVAngle = LevelViewportClient->FOVAngle;
+		ViewportInstanceSettings.bAllowCinematicControl = LevelViewportClient->AllowsCinematicControl();
 	
 		LevelViewportClient->SaveRealtimeStateToConfig(ViewportInstanceSettings.bIsRealtime);
 	
@@ -2413,6 +2414,7 @@ FLevelEditorViewportInstanceSettings SLevelViewport::LoadLegacyConfigFromIni(con
 	GConfig->GetBool(*IniSection, *(InConfigKey + TEXT(".bWantStats")), ViewportInstanceSettings.bShowOnScreenStats, GEditorPerProjectIni);
 	GConfig->GetBool(*IniSection, *(InConfigKey + TEXT(".bWantFPS")), ViewportInstanceSettings.bShowFPS_DEPRECATED, GEditorPerProjectIni);
 	GConfig->GetFloat(*IniSection, *(InConfigKey + TEXT(".FOVAngle")), ViewportInstanceSettings.FOVAngle, GEditorPerProjectIni);
+	GConfig->GetBool(*IniSection, *(InConfigKey + TEXT(".bAllowCinematicControl")), ViewportInstanceSettings.bAllowCinematicControl, GEditorPerProjectIni);
 
 	return ViewportInstanceSettings;
 }
diff --git a/Engine/Source/Editor/UnrealEd/Classes/Settings/LevelEditorViewportSettings.h b/Engine/Source/Editor/UnrealEd/Classes/Settings/LevelEditorViewportSettings.h
--- a/Engine/Source/Editor/UnrealEd/Classes/Settings/LevelEditorViewportSettings.h
+++ b/Engine/Source/Editor/UnrealEd/Classes/Settings/LevelEditorViewportSettings.h
@@ -111,6 +111,7 @@ struct FLevelEditorViewportInstanceSettings
 		// Show 'lighting needs to be rebuilt' message by default, avoids confusion when artists think lighting is built until they PIE
 		, bShowOnScreenStats(true)
 		, bShowFullToolbar(true)
+		, bAllowCinematicControl(true)
 	{ }
 
 	/** The viewport type */
@@ -202,6 +203,10 @@ struct FLevelEditorViewportInstanceSettings
 	/** When enabled, the full viewport toolbar will be shown. When disabled, a compact toolbar is used. */
 	UPROPERTY(EditAnywhere, config, Category=LookAndFeel)
 	bool bShowFullToolbar;
+
+	/** Whether or not this viewport is allowed to be possessed by cinematic/scrubbing tools. */
+	UPROPERTY(config)
+	bool bAllowCinematicControl;
 };
 
 
