Editor: Add logic for unreachable actor in CanDeselectActorElement

(Note same logic exists in CanSelectActorElement)

This fixes a bug when entering PIE and clearing selection, which attempts to clear selection on a spawnable that has been destroyed.

diff --git a/Engine/Source/Editor/LevelEditor/Private/Elements/Actor/ActorElementLevelEditorSelectionCustomization.cpp b/Engine/Source/Editor/LevelEditor/Private/Elements/Actor/ActorElementLevelEditorSelectionCustomization.cpp
--- a/Engine/Source/Editor/LevelEditor/Private/Elements/Actor/ActorElementLevelEditorSelectionCustomization.cpp
+++ b/Engine/Source/Editor/LevelEditor/Private/Elements/Actor/ActorElementLevelEditorSelectionCustomization.cpp
@@ -77,6 +77,34 @@ void FActorElementLevelEditorSelectionCustomization::GetNormalizedElements(const
 	FActorElementLevelEditorSelectionCustomization::AppendNormalizedActors(Actor, InSelectionSet, InNormalizationOptions, OutNormalizedElements);
 }
 
+namespace LevelEditorSelectionHelpers
+{
+
+bool IsActorReachable(const AActor* Actor)
+{
+	if (!Actor)
+	{
+		return false;
+	}
+
+	// Ensure that neither the level nor the actor is being destroyed or is unreachable
+	const EObjectFlags InvalidSelectableFlags = RF_BeginDestroyed;
+	if (Actor->GetLevel()->HasAnyFlags(InvalidSelectableFlags) || (!GIsTransacting && (!IsValidChecked(Actor->GetLevel()) || Actor->GetLevel()->IsUnreachable())))
+	{
+		UE_LOG(LogActorLevelEditorSelection, Warning, TEXT("SelectActor: %s (%s)"), TEXT("The requested operation could not be completed because the level has invalid flags."), *Actor->GetActorLabel());
+		return false;
+	}
+	if (Actor->HasAnyFlags(InvalidSelectableFlags) || (!GIsTransacting && (!IsValidChecked(Actor) || Actor->IsUnreachable())))
+	{
+		UE_LOG(LogActorLevelEditorSelection, Warning, TEXT("SelectActor: %s (%s)"), TEXT("The requested operation could not be completed because the actor has invalid flags."), *Actor->GetActorLabel());
+		return false;
+	}
+
+	return true;
+}
+
+}
+
 bool FActorElementLevelEditorSelectionCustomization::CanSelectActorElement(const TTypedElement<ITypedElementSelectionInterface>& InActorSelectionHandle, const FTypedElementSelectionOptions& InSelectionOptions) const
 {
 	AActor* Actor = ActorElementDataUtil::GetActorFromHandleChecked(InActorSelectionHandle);
@@ -93,16 +121,8 @@ bool FActorElementLevelEditorSelectionCustomization::CanSelectActorElement(const
 		return false;
 	}
 
-	// Ensure that neither the level nor the actor is being destroyed or is unreachable
-	const EObjectFlags InvalidSelectableFlags = RF_BeginDestroyed;
-	if (Actor->GetLevel()->HasAnyFlags(InvalidSelectableFlags) || (!GIsTransacting && (!IsValidChecked(Actor->GetLevel()) || Actor->GetLevel()->IsUnreachable())))
+	if (!LevelEditorSelectionHelpers::IsActorReachable(Actor))
 	{
-		UE_LOG(LogActorLevelEditorSelection, Warning, TEXT("SelectActor: %s (%s)"), TEXT("The requested operation could not be completed because the level has invalid flags."), *Actor->GetActorLabel());
-		return false;
-	}
-	if (Actor->HasAnyFlags(InvalidSelectableFlags) || (!GIsTransacting && (!IsValidChecked(Actor) || Actor->IsUnreachable())))
-	{
-		UE_LOG(LogActorLevelEditorSelection, Warning, TEXT("SelectActor: %s (%s)"), TEXT("The requested operation could not be completed because the actor has invalid flags."), *Actor->GetActorLabel());
 		return false;
 	}
 
@@ -147,6 +167,11 @@ bool FActorElementLevelEditorSelectionCustomization::CanDeselectActorElement(con
 		return false;
 	}
 
+	if (!LevelEditorSelectionHelpers::IsActorReachable(Actor))
+	{
+		return false;
+	}
+
 	if (const IToolkitHost* ToolkitHostPtr = GetToolkitHost())
 	{
 		// Allow active modes to determine whether the deselection is allowed
