[Keyboard Shortcuts] Deprecate AActor::EditorKeyPressed and convert current usage to keybinds

diff --git a/Engine/Source/Editor/LevelEditor/Private/SLevelViewport.cpp b/Engine/Source/Editor/LevelEditor/Private/SLevelViewport.cpp
--- a/Engine/Source/Editor/LevelEditor/Private/SLevelViewport.cpp
+++ b/Engine/Source/Editor/LevelEditor/Private/SLevelViewport.cpp
@@ -1523,7 +1523,8 @@ void SLevelViewport::BindOptionCommands( FUICommandList& OutCommandList )
 	{
 		OutCommandList.MapAction( 
 			ViewportActions.JumpToBookmarkCommands[BookmarkIndex],
-			FExecuteAction::CreateSP( this, &SLevelViewport::OnJumpToBookmark, BookmarkIndex )
+			FExecuteAction::CreateSP( this, &SLevelViewport::OnJumpToBookmark, BookmarkIndex ),
+			FCanExecuteAction::CreateSP(this, &SLevelViewport::OnHasBookmarkSet, BookmarkIndex)
 			);
 
 		OutCommandList.MapAction( 
@@ -2430,6 +2431,11 @@ void SLevelViewport::OnJumpToBookmark( int32 BookmarkIndex )
 	IBookmarkTypeTools::Get().JumpToBookmark( BookmarkIndex, TSharedPtr<struct FBookmarkBaseJumpToSettings>(), LevelViewportClient.Get() );
 }
 
+bool SLevelViewport::OnHasBookmarkSet(int32 BookmarkIndex)
+{
+	return IBookmarkTypeTools::Get().CheckBookmark(BookmarkIndex, LevelViewportClient.Get());
+}
+
 void SLevelViewport::OnClearBookmark(int32 BookmarkIndex)
 {
 	IBookmarkTypeTools::Get().ClearBookmark(BookmarkIndex, LevelViewportClient.Get());
diff --git a/Engine/Source/Editor/LevelEditor/Public/SLevelViewport.h b/Engine/Source/Editor/LevelEditor/Public/SLevelViewport.h
--- a/Engine/Source/Editor/LevelEditor/Public/SLevelViewport.h
+++ b/Engine/Source/Editor/LevelEditor/Public/SLevelViewport.h
@@ -409,6 +409,13 @@ public:
 	 */
 	void OnSetBookmark( int32 BookmarkIndex );
 
+	/**
+	 * Called to check if a bookmark is set
+	 *
+	 * @param BookmarkIndex	The index of the bookmark to check
+	 */
+	bool OnHasBookmarkSet(int32 BookmarkIndex);
+
 	/**
 	 * Called to jump to a bookmark
 	 *
diff --git a/Engine/Source/Editor/UnrealEd/Private/EditorModeManager.cpp b/Engine/Source/Editor/UnrealEd/Private/EditorModeManager.cpp
--- a/Engine/Source/Editor/UnrealEd/Private/EditorModeManager.cpp
+++ b/Engine/Source/Editor/UnrealEd/Private/EditorModeManager.cpp
@@ -1382,17 +1382,19 @@ bool FEditorModeTools::InputKey(FEditorViewportClient* InViewportClient, FViewpo
 		}
 
 		return true;
-	});
+		});
 
-	// Finally, pass input to selected actors if nothing else handled the input
+PRAGMA_DISABLE_DEPRECATION_WARNINGS // Begin AActor::EditorKeyPressed
+	// Finally, pass input to selected actors if nothing else handled the input (Deprecated in 5.4)
 	if (!bHandled)
 	{
 		GetEditorSelectionSet()->ForEachSelectedObject<AActor>([Key, Event](AActor* ActorPtr)
-		{
-			ActorPtr->EditorKeyPressed(Key, Event);
-			return true;
-		});
+			{
+				ActorPtr->EditorKeyPressed(Key, Event);
+				return true;
+			});
 	}
+PRAGMA_ENABLE_DEPRECATION_WARNINGS // End AActor::EditorKeyPressed
 	return bHandled;
 }
 
diff --git a/Engine/Source/Runtime/Engine/Classes/GameFramework/Actor.h b/Engine/Source/Runtime/Engine/Classes/GameFramework/Actor.h
--- a/Engine/Source/Runtime/Engine/Classes/GameFramework/Actor.h
+++ b/Engine/Source/Runtime/Engine/Classes/GameFramework/Actor.h
@@ -2486,6 +2486,7 @@ public:
 	virtual bool ShouldImport(FStringView ActorPropString, bool IsMovingLevel) { return true; }
 
 	/** Called by InputKey when an unhandled key is pressed with a selected actor */
+	UE_DEPRECATED(5.4, "Please use UI Commands and process your custom actors as necessary by extending the level editor modules' global command list.")
 	virtual void EditorKeyPressed(FKey Key, EInputEvent Event) {}
 
 	/** Called by ReplaceSelectedActors to allow a new actor to copy properties from an old actor when it is replaced */
