PR #10115:  Call OpenEditorForAsset instead of LoadMap when opening a map using the "File" menu to match how the CB opens levels

diff --git a/Engine/Source/Editor/LevelEditor/Private/LevelEditorActions.cpp b/Engine/Source/Editor/LevelEditor/Private/LevelEditorActions.cpp
--- a/Engine/Source/Editor/LevelEditor/Private/LevelEditorActions.cpp
+++ b/Engine/Source/Editor/LevelEditor/Private/LevelEditorActions.cpp
@@ -313,12 +313,7 @@ void FLevelEditorActionCallbacks::OpenRecentFile( int32 RecentFileIndex )
 			// Prompt the user to save any outstanding changes.
 			if( FEditorFileUtils::SaveDirtyPackages(true, true, false) )
 			{
-				FString NewFilename;
-				if (FPackageName::TryConvertLongPackageNameToFilename(NewPackageName, NewFilename, FPackageName::GetMapPackageExtension()))
-				{
-					// Load the requested level.
-					FEditorFileUtils::LoadMap(NewFilename);
-				}
+				GEditor->GetEditorSubsystem<UAssetEditorSubsystem>()->OpenEditorForAsset(NewPackageName, EAssetTypeActivationOpenedMethod::Edit);
 			}
 		}
 	}
@@ -346,13 +341,8 @@ void FLevelEditorActionCallbacks::OpenFavoriteFile( int32 FavoriteFileIndex )
 		// Prompt the user to save any outstanding changes
 		if( FEditorFileUtils::SaveDirtyPackages(true, true, false) )
 		{
-			FString FileName;
-			if (FPackageName::TryConvertLongPackageNameToFilename(PackageName, FileName, FPackageName::GetMapPackageExtension()))
-			{
-				// Load the requested level.
-				FEditorFileUtils::LoadMap(FileName);
-			}
-
+			GEditor->GetEditorSubsystem<UAssetEditorSubsystem>()->OpenEditorForAsset(PackageName, EAssetTypeActivationOpenedMethod::Edit);
+			
 			// Move the item to the head of the list
 			MRUFavoritesList->MoveFavoritesItemToHead(PackageName);
 		}
