World Partition - HLOD: Refresh viewport when dragging the HLOD in editor min/max draw distance sliders

diff --git a/Engine/Source/Editor/LevelEditor/Private/SLevelViewportToolBar.cpp b/Engine/Source/Editor/LevelEditor/Private/SLevelViewportToolBar.cpp
--- a/Engine/Source/Editor/LevelEditor/Private/SLevelViewportToolBar.cpp
+++ b/Engine/Source/Editor/LevelEditor/Private/SLevelViewportToolBar.cpp
@@ -1587,7 +1587,16 @@ void SLevelViewportToolBar::FillShowHLODsMenu(UToolMenu* Menu) const
 			.IsEnabled(bHLODInEditorAllowed)
 			.Value(this, &SLevelViewportToolBar::OnGetHLODInEditorMinDrawDistanceValue)
 			.OnValueChanged(this, &SLevelViewportToolBar::OnHLODInEditorMinDrawDistanceValueChanged)
-			.ToolTipText(bHLODInEditorAllowed ? LOCTEXT("HLODsInEditor_MinDrawDistance_Tooltip", "Sets the minimum distance at which HLOD will be rendered") : HLODInEditorDisallowedReason);
+			.ToolTipText(bHLODInEditorAllowed ? LOCTEXT("HLODsInEditor_MinDrawDistance_Tooltip", "Sets the minimum distance at which HLOD will be rendered") : HLODInEditorDisallowedReason)
+			.OnBeginSliderMovement_Lambda([this]()
+			{
+				// Disable Slate throttling during slider drag to ensure immediate updates while moving the slider.
+				FSlateThrottleManager::Get().DisableThrottle(true);
+			})
+			.OnEndSliderMovement_Lambda([this](float)
+			{
+				FSlateThrottleManager::Get().DisableThrottle(false);
+			});
 
 		TSharedRef<SSpinBox<double>> MaxDrawDistanceSpinBox = SNew(SSpinBox<double>)
 			.MinValue(MaxDrawDistanceMinValue)
@@ -1595,7 +1604,16 @@ void SLevelViewportToolBar::FillShowHLODsMenu(UToolMenu* Menu) const
 			.IsEnabled(bHLODInEditorAllowed)
 			.Value(this, &SLevelViewportToolBar::OnGetHLODInEditorMaxDrawDistanceValue)
 			.OnValueChanged(this, &SLevelViewportToolBar::OnHLODInEditorMaxDrawDistanceValueChanged)
-			.ToolTipText(bHLODInEditorAllowed ? LOCTEXT("HLODsInEditor_MaxDrawDistance_Tooltip", "Sets the maximum distance at which HLODs will be rendered") : HLODInEditorDisallowedReason);
+			.ToolTipText(bHLODInEditorAllowed ? LOCTEXT("HLODsInEditor_MaxDrawDistance_Tooltip", "Sets the maximum distance at which HLODs will be rendered") : HLODInEditorDisallowedReason)
+			.OnBeginSliderMovement_Lambda([this]()
+			{
+				// Disable Slate throttling during slider drag to ensure immediate updates while moving the slider.
+				FSlateThrottleManager::Get().DisableThrottle(true);
+			})
+			.OnEndSliderMovement_Lambda([this](float)
+			{
+				FSlateThrottleManager::Get().DisableThrottle(false);
+			});
 
 		auto CreateDrawDistanceWidget = [](TSharedRef<SSpinBox<double>> InSpinBoxWidget)
 		{
