Actor Coloration Framework
- Level Coloration show flag is now replaced by the generic Actor Coloration show flag, and is now hidden in the show flags menu.
- Added a dropdown Actor Coloration menu that is dynamically populated with the current registered Actor Coloration handlers.
- Color handlers can now return different colors for the same actor, but different primitive components.

This way, if a plugin registers a new Actor Coloration mode, the menu will be automatically populated with that option until the plugin is unregistered.

diff --git a/Engine/Source/Editor/LevelEditor/Private/SLevelViewportToolBar.cpp b/Engine/Source/Editor/LevelEditor/Private/SLevelViewportToolBar.cpp
--- a/Engine/Source/Editor/LevelEditor/Private/SLevelViewportToolBar.cpp
+++ b/Engine/Source/Editor/LevelEditor/Private/SLevelViewportToolBar.cpp
@@ -17,6 +17,7 @@
 #include "Styling/AppStyle.h"
 #include "Camera/CameraActor.h"
 #include "Misc/ConfigCacheIni.h"
+#include "GameFramework/ActorPrimitiveColorHandler.h"
 #include "GameFramework/WorldSettings.h"
 #include "EngineUtils.h"
 #include "LevelEditor.h"
@@ -1812,47 +1813,141 @@ void SLevelViewportToolBar::FillViewMenu(UToolMenu* Menu)
 	}
 
 	{
-		FToolMenuSection& Section = Menu->AddSection("LevelViewportLandscape", LOCTEXT("LandscapeHeader", "Landscape"), InsertPosition);
-
-		struct Local
+		auto BuildActorColorationMenu = [this](UToolMenu* Menu, SLevelViewportToolBar* Toolbar)
 		{
-			static void BuildLandscapeLODMenu(UToolMenu* Menu, SLevelViewportToolBar* Toolbar)
+			FToolMenuSection& SubMenuSection = Menu->AddSection("LevelViewportActorColoration", LOCTEXT("ActorColorationHeader", "Actor Coloration"));
+
+			TArray<FActorPrimitiveColorHandler::FPrimitiveColorHandler> PrimitiveColorHandlers;
+			FActorPrimitiveColorHandler::Get().GetRegisteredPrimitiveColorHandlers(PrimitiveColorHandlers);
+
+			for (const FActorPrimitiveColorHandler::FPrimitiveColorHandler& PrimitiveColorHandler : PrimitiveColorHandlers)
 			{
+				SubMenuSection.AddMenuEntry(
+					NAME_None,
+					PrimitiveColorHandler.HandlerText,
+					FText(),
+					FSlateIcon(),
+					FUIAction(
+						FExecuteAction::CreateLambda([this, PrimitiveColorHandler]()
+						{
+							if (FLevelEditorViewportClient* ViewportClient = GetLevelViewportClient())
+							{
+								const bool bActorColorationEnabled = ViewportClient->HandleIsShowFlagEnabled(FEngineShowFlags::EShowFlag::SF_ActorColoration);
+
+								if (PrimitiveColorHandler.HandlerName.IsNone())
+								{
+									if (bActorColorationEnabled)
+									{
+										ViewportClient->HandleToggleShowFlag(FEngineShowFlags::EShowFlag::SF_ActorColoration);
+									}
+								}
+								else
+								{
+									if (!bActorColorationEnabled)
+									{
+										ViewportClient->HandleToggleShowFlag(FEngineShowFlags::EShowFlag::SF_ActorColoration);
+									}
+
+									FActorPrimitiveColorHandler::Get().SetActivePrimitiveColorHandler(PrimitiveColorHandler.HandlerName, GWorld);
+								}
+							}
+						}),
+						FCanExecuteAction::CreateLambda([this]()
+						{
+							if (FLevelEditorViewportClient* ViewportClient = GetLevelViewportClient())
+							{
+								return true;
+							}
+							return false;
+						}),
+						FGetActionCheckState::CreateLambda([this, PrimitiveColorHandler]()
+						{
+							if (FLevelEditorViewportClient* ViewportClient = GetLevelViewportClient())
+							{
+								const bool bActorColorationEnabled = ViewportClient->HandleIsShowFlagEnabled(FEngineShowFlags::EShowFlag::SF_ActorColoration);
+
+								if (PrimitiveColorHandler.HandlerName.IsNone())
+								{
+									return bActorColorationEnabled ? ECheckBoxState::Unchecked : ECheckBoxState::Checked;
+								}
+								else
+								{
+									if (bActorColorationEnabled)
+									{
+										return FActorPrimitiveColorHandler::Get().GetActivePrimitiveColorHandler() == PrimitiveColorHandler.HandlerName ? ECheckBoxState::Checked : ECheckBoxState::Unchecked;
+									}
+								}								
+							}
+
+							return ECheckBoxState::Unchecked;
+						})
+					),
+					EUserInterfaceActionType::RadioButton
+				);
+			}
+		};
+
+		FToolMenuSection& Section = Menu->FindOrAddSection("ViewMode");
+		Section.AddSubMenu(
+			"ActorColoration",
+			LOCTEXT("ActorColorationDisplayName", "Actor Coloration"),
+			LOCTEXT("ActorColorationMenu_ToolTip", "Override Actor Coloration mode"),
+			FNewToolMenuDelegate::CreateLambda(BuildActorColorationMenu, this),
+			FUIAction(
+				FExecuteAction(),
+				FCanExecuteAction(),
+				FIsActionChecked::CreateLambda([this]()
 				{
-					FToolMenuSection& SubMenuSection = Menu->AddSection("LevelViewportLandScapeLOD", LOCTEXT("LandscapeLODHeader", "Landscape LOD"));
-
-					SubMenuSection.AddMenuEntry(
-						"LandscapeLODAuto",
-						LOCTEXT("LandscapeLODAuto", "Auto"),
-						FText(),
-						FSlateIcon(),
-						FUIAction(
-							FExecuteAction::CreateSP(Toolbar, &SLevelViewportToolBar::OnLandscapeLODChanged, -1),
-							FCanExecuteAction(),
-							FIsActionChecked::CreateSP(Toolbar, &SLevelViewportToolBar::IsLandscapeLODSettingChecked, -1)
-						),
-						EUserInterfaceActionType::RadioButton
-					);
-
-					SubMenuSection.AddSeparator("LandscapeLODSeparator");
-
-					static const FText FormatString = LOCTEXT("LandscapeLODFixed", "Fixed at {0}");
-					for (int32 i = 0; i < 8; ++i)
+					if (FLevelEditorViewportClient* ViewportClient = GetLevelViewportClient())
 					{
-						SubMenuSection.AddMenuEntry(
-							NAME_None,
-							FText::Format(FormatString, FText::AsNumber(i)),
-							FText(),
-							FSlateIcon(),
-							FUIAction(
-								FExecuteAction::CreateSP(Toolbar, &SLevelViewportToolBar::OnLandscapeLODChanged, i),
-								FCanExecuteAction(),
-								FIsActionChecked::CreateSP(Toolbar, &SLevelViewportToolBar::IsLandscapeLODSettingChecked, i)
-							),
-							EUserInterfaceActionType::RadioButton
-						);
+						return ViewportClient->HandleIsShowFlagEnabled(FEngineShowFlags::EShowFlag::SF_ActorColoration);
 					}
-				}
+					return false;
+				})
+			),
+			EUserInterfaceActionType::RadioButton,
+			/*bInOpenSubMenuOnClick=*/ false,
+			FSlateIcon(FAppStyle::GetAppStyleSetName(), "EditorViewport.LODColorationMode")
+		);
+	}
+
+	{
+		FToolMenuSection& Section = Menu->AddSection("LevelViewportLandscape", LOCTEXT("LandscapeHeader", "Landscape"), InsertPosition);
+
+		auto BuildLandscapeLODMenu = [](UToolMenu* Menu, SLevelViewportToolBar* Toolbar)
+		{
+			FToolMenuSection& SubMenuSection = Menu->AddSection("LevelViewportLandScapeLOD", LOCTEXT("LandscapeLODHeader", "Landscape LOD"));
+
+			SubMenuSection.AddMenuEntry(
+				"LandscapeLODAuto",
+				LOCTEXT("LandscapeLODAuto", "Auto"),
+				FText(),
+				FSlateIcon(),
+				FUIAction(
+					FExecuteAction::CreateSP(Toolbar, &SLevelViewportToolBar::OnLandscapeLODChanged, -1),
+					FCanExecuteAction(),
+					FIsActionChecked::CreateSP(Toolbar, &SLevelViewportToolBar::IsLandscapeLODSettingChecked, -1)
+				),
+				EUserInterfaceActionType::RadioButton
+			);
+
+			SubMenuSection.AddSeparator("LandscapeLODSeparator");
+
+			static const FText FormatString = LOCTEXT("LandscapeLODFixed", "Fixed at {0}");
+			for (int32 i = 0; i < 8; ++i)
+			{
+				SubMenuSection.AddMenuEntry(
+					NAME_None,
+					FText::Format(FormatString, FText::AsNumber(i)),
+					FText(),
+					FSlateIcon(),
+					FUIAction(
+						FExecuteAction::CreateSP(Toolbar, &SLevelViewportToolBar::OnLandscapeLODChanged, i),
+						FCanExecuteAction(),
+						FIsActionChecked::CreateSP(Toolbar, &SLevelViewportToolBar::IsLandscapeLODSettingChecked, i)
+					),
+					EUserInterfaceActionType::RadioButton
+				);
 			}
 		};
 
@@ -1860,7 +1955,7 @@ void SLevelViewportToolBar::FillViewMenu(UToolMenu* Menu)
 			"LandscapeLOD",
 			LOCTEXT("LandscapeLODDisplayName", "LOD"),
 			LOCTEXT("LandscapeLODMenu_ToolTip", "Override Landscape LOD in this viewport"),
-			FNewToolMenuDelegate::CreateStatic(&Local::BuildLandscapeLODMenu, this),
+			FNewToolMenuDelegate::CreateLambda(BuildLandscapeLODMenu, this),
 			/*bInOpenSubMenuOnClick=*/ false,
 			FSlateIcon(FAppStyle::GetAppStyleSetName(), "EditorViewport.LOD")
 		);
diff --git a/Engine/Source/Runtime/Engine/Classes/GameFramework/ActorPrimitiveColorHandler.h b/Engine/Source/Runtime/Engine/Classes/GameFramework/ActorPrimitiveColorHandler.h
--- a/Engine/Source/Runtime/Engine/Classes/GameFramework/ActorPrimitiveColorHandler.h
+++ b/Engine/Source/Runtime/Engine/Classes/GameFramework/ActorPrimitiveColorHandler.h
@@ -3,9 +3,12 @@
 #pragma once
 
 #include "UObject/ObjectMacros.h"
+#include "Components/PrimitiveComponent.h"
 
-class AActor;
 class UWorld;
+class UPrimitiveComponent;
+
+#define ENABLE_ACTOR_PRIMITIVE_COLOR_HANDLER !(UE_BUILD_SHIPPING || UE_BUILD_TEST)
 
 /**
  * FActorPrimitiveColorHandler is a simple mechanism for custom actor coloration registration. Once an actor color
@@ -13,19 +16,36 @@ class UWorld;
  */
 class ENGINE_API FActorPrimitiveColorHandler
 {
-	using FPrimitiveColorHandler = TFunction<FLinearColor(AActor*)>;
+	using FFunc = TFunction<FLinearColor(const UPrimitiveComponent*)>;
 
 public:
+	struct FPrimitiveColorHandler
+	{
+		FPrimitiveColorHandler(FName InHandlerName, FText InHandlerText, const FFunc& InHandlerFunc)
+			: HandlerName(InHandlerName)
+			, HandlerText(InHandlerText)
+			, HandlerFunc(InHandlerFunc)
+		{}
+
+		FName HandlerName;
+		FText HandlerText;
+		FFunc HandlerFunc;
+	};	
+
 	FActorPrimitiveColorHandler();
 	static FActorPrimitiveColorHandler& Get();
-	void RegisterPrimitiveColorHandler(FName InHandlerName, const FPrimitiveColorHandler& InHandler);
+	void RegisterPrimitiveColorHandler(FName InHandlerName, FText InHandlerText, const FFunc& InHandlerFunc);
 	void UnregisterPrimitiveColorHandler(FName InHandlerName);
 	bool SetActivePrimitiveColorHandler(FName InHandlerName, UWorld* InWorld);
-	FLinearColor GetPrimitiveColor(AActor* InActor);
+	void RefreshPrimitiveColorHandler(FName InHandlerName, UWorld* InWorld);
+	FName GetActivePrimitiveColorHandler() const;
+	void GetRegisteredPrimitiveColorHandlers(TArray<FPrimitiveColorHandler>& OutPrimitiveColorHandlers) const;
+	FLinearColor GetPrimitiveColor(const UPrimitiveComponent* InPrimitiveComponent) const;
 
 private:
-#if !(UE_BUILD_SHIPPING || UE_BUILD_TEST)
+#if ENABLE_ACTOR_PRIMITIVE_COLOR_HANDLER
 	FName ActivePrimitiveColorHandlerName;
+	FText ActivePrimitiveColorHandlerText;
 	FPrimitiveColorHandler* ActivePrimitiveColorHandler;
 	TMap<FName, FPrimitiveColorHandler> Handlers;
 #endif
diff --git a/Engine/Source/Runtime/Engine/Private/ActorPrimitiveColorHandler.cpp b/Engine/Source/Runtime/Engine/Private/ActorPrimitiveColorHandler.cpp
--- a/Engine/Source/Runtime/Engine/Private/ActorPrimitiveColorHandler.cpp
+++ b/Engine/Source/Runtime/Engine/Private/ActorPrimitiveColorHandler.cpp
@@ -2,14 +2,15 @@
 
 #include "GameFramework/ActorPrimitiveColorHandler.h"
 #include "GameFramework/Actor.h"
-#include "Components/PrimitiveComponent.h"
 #include "Misc/LazySingleton.h"
 #include "EngineUtils.h"
 
+#define LOCTEXT_NAMESPACE "ActorColoration"
+
 FActorPrimitiveColorHandler::FActorPrimitiveColorHandler()
 {
-#if !(UE_BUILD_SHIPPING || UE_BUILD_TEST)
-	RegisterPrimitiveColorHandler(NAME_None, [](const AActor*)
+#if ENABLE_ACTOR_PRIMITIVE_COLOR_HANDLER
+	RegisterPrimitiveColorHandler(NAME_None, LOCTEXT("Disable", "Disable"), [](const UPrimitiveComponent*)
 	{
 		return FLinearColor::White;
 	});
@@ -24,19 +25,18 @@ FActorPrimitiveColorHandler& FActorPrimitiveColorHandler::Get()
 	return TLazySingleton<FActorPrimitiveColorHandler>::Get();
 }
 
-void FActorPrimitiveColorHandler::RegisterPrimitiveColorHandler(FName InHandlerName, const FPrimitiveColorHandler& InHandler)
+void FActorPrimitiveColorHandler::RegisterPrimitiveColorHandler(FName InHandlerName, FText InHandlerText, const FFunc& InHandlerFunc)
 {
-#if !(UE_BUILD_SHIPPING || UE_BUILD_TEST)
+#if ENABLE_ACTOR_PRIMITIVE_COLOR_HANDLER
 	check(!Handlers.Contains(InHandlerName));
-	Handlers.Add(InHandlerName, InHandler);
-	
-	ActivePrimitiveColorHandler = Handlers.Find(ActivePrimitiveColorHandlerName);
+	Handlers.Add(InHandlerName, { InHandlerName, InHandlerText, InHandlerFunc });	
+	ActivePrimitiveColorHandler = &Handlers.FindChecked(ActivePrimitiveColorHandlerName);
 #endif
 }
 
 void FActorPrimitiveColorHandler::UnregisterPrimitiveColorHandler(FName InHandlerName)
 {
-#if !(UE_BUILD_SHIPPING || UE_BUILD_TEST)
+#if ENABLE_ACTOR_PRIMITIVE_COLOR_HANDLER
 	check(!InHandlerName.IsNone());
 	check(Handlers.Contains(InHandlerName));
 	Handlers.Remove(InHandlerName);
@@ -52,12 +52,24 @@ void FActorPrimitiveColorHandler::UnregisterPrimitiveColorHandler(FName InHandle
 
 bool FActorPrimitiveColorHandler::SetActivePrimitiveColorHandler(FName InHandlerName, UWorld* InWorld)
 {
-#if !(UE_BUILD_SHIPPING || UE_BUILD_TEST)
+#if ENABLE_ACTOR_PRIMITIVE_COLOR_HANDLER
 	if (FPrimitiveColorHandler* NewActivePrimitiveColorHandler = Handlers.Find(InHandlerName); NewActivePrimitiveColorHandler && (NewActivePrimitiveColorHandler != ActivePrimitiveColorHandler))
 	{
 		ActivePrimitiveColorHandlerName = InHandlerName;
 		ActivePrimitiveColorHandler = NewActivePrimitiveColorHandler;
+		RefreshPrimitiveColorHandler(InHandlerName, InWorld);
+		return true;
+	}
+#endif
 
+	return false;
+}
+
+void FActorPrimitiveColorHandler::RefreshPrimitiveColorHandler(FName InHandlerName, UWorld* InWorld)
+{
+#if ENABLE_ACTOR_PRIMITIVE_COLOR_HANDLER
+	if (ActivePrimitiveColorHandlerName == InHandlerName)
+	{
 		for (TActorIterator<AActor> It(InWorld); It; ++It)
 		{
 			TInlineComponentArray<UPrimitiveComponent*> PrimitiveComponents;
@@ -67,23 +79,37 @@ bool FActorPrimitiveColorHandler::SetActivePrimitiveColorHandler(FName InHandler
 			{
 				if (PrimitiveComponent->IsRegistered())
 				{
-					PrimitiveComponent->PushPrimitiveColorToProxy(GetPrimitiveColor(*It));
+					PrimitiveComponent->PushPrimitiveColorToProxy(GetPrimitiveColor(PrimitiveComponent));
 				}
 			}
 		}
-
-		return true;
 	}
 #endif
+}
 
-	return false;
+FName FActorPrimitiveColorHandler::GetActivePrimitiveColorHandler() const
+{
+#if ENABLE_ACTOR_PRIMITIVE_COLOR_HANDLER
+	return ActivePrimitiveColorHandlerName;
+#else
+	return NAME_None;
+#endif
 }
 
-FLinearColor FActorPrimitiveColorHandler::GetPrimitiveColor(AActor* InActor)
+void FActorPrimitiveColorHandler::GetRegisteredPrimitiveColorHandlers(TArray<FPrimitiveColorHandler>& OutPrimitiveColorHandlers) const
 {
-#if !(UE_BUILD_SHIPPING || UE_BUILD_TEST)
-	return (*ActivePrimitiveColorHandler)(InActor);
+#if ENABLE_ACTOR_PRIMITIVE_COLOR_HANDLER
+	Handlers.GenerateValueArray(OutPrimitiveColorHandlers);
+#endif
+}
+
+FLinearColor FActorPrimitiveColorHandler::GetPrimitiveColor(const UPrimitiveComponent* InPrimitiveComponent) const
+{
+#if ENABLE_ACTOR_PRIMITIVE_COLOR_HANDLER
+	return ActivePrimitiveColorHandler->HandlerFunc(InPrimitiveComponent);
 #else
 	return FLinearColor::White;
 #endif
-}
\ No newline at end of file
+}
+
+#undef LOCTEXT_NAMESPACE
\ No newline at end of file
diff --git a/Engine/Source/Runtime/Engine/Private/Components/BrushComponent.cpp b/Engine/Source/Runtime/Engine/Private/Components/BrushComponent.cpp
--- a/Engine/Source/Runtime/Engine/Private/Components/BrushComponent.cpp
+++ b/Engine/Source/Runtime/Engine/Private/Components/BrushComponent.cpp
@@ -131,7 +131,7 @@ public:
 			}
 			else
 			{
-				ActorColor = FActorPrimitiveColorHandler::Get().GetPrimitiveColor(Owner);
+				ActorColor = FActorPrimitiveColorHandler::Get().GetPrimitiveColor(Component);
 			}
 		}
 
diff --git a/Engine/Source/Runtime/Engine/Private/LevelStreaming.cpp b/Engine/Source/Runtime/Engine/Private/LevelStreaming.cpp
--- a/Engine/Source/Runtime/Engine/Private/LevelStreaming.cpp
+++ b/Engine/Source/Runtime/Engine/Private/LevelStreaming.cpp
@@ -398,12 +398,12 @@ ULevelStreaming::ULevelStreaming(const FObjectInitializer& ObjectInitializer)
 	bSkipClientUseMakingVisibleTransactionRequest = false;
 	bGarbageCollectionClusteringEnabled = true;
 
-#if !(UE_BUILD_SHIPPING || UE_BUILD_TEST)
+#if ENABLE_ACTOR_PRIMITIVE_COLOR_HANDLER
 	if (HasAnyFlags(RF_ClassDefaultObject) && ExactCast<ULevelStreaming>(this))
 	{
-		FActorPrimitiveColorHandler::Get().RegisterPrimitiveColorHandler(TEXT("LevelColor"), [](const AActor* InActor)
+		FActorPrimitiveColorHandler::Get().RegisterPrimitiveColorHandler(TEXT("LevelColor"), LOCTEXT("LevelColor", "Level Color"), [](const UPrimitiveComponent* InPrimitiveComponent)
 		{
-			if (ULevel* Level = InActor ? InActor->GetLevel() : nullptr)
+			if (ULevel* Level = InPrimitiveComponent ? InPrimitiveComponent->GetComponentLevel() : nullptr)
 			{
 				if (ULevelStreaming* LevelStreaming = FLevelUtils::FindStreamingLevel(Level))
 				{
diff --git a/Engine/Source/Runtime/Engine/Private/PrimitiveSceneProxy.cpp b/Engine/Source/Runtime/Engine/Private/PrimitiveSceneProxy.cpp
--- a/Engine/Source/Runtime/Engine/Private/PrimitiveSceneProxy.cpp
+++ b/Engine/Source/Runtime/Engine/Private/PrimitiveSceneProxy.cpp
@@ -555,9 +555,9 @@ FPrimitiveSceneProxy::FPrimitiveSceneProxy(const FPrimitiveSceneProxyDesc& InPro
 		}
 
 #if !(UE_BUILD_SHIPPING || UE_BUILD_TEST)
-		if (AActor* ActorOwner = Cast<AActor>(InProxyDesc.Owner))
+		if (UPrimitiveComponent* PrimitiveComponent = Cast<UPrimitiveComponent>(InProxyDesc.Component))
 		{
-			SetPrimitiveColor(FActorPrimitiveColorHandler::Get().GetPrimitiveColor(ActorOwner));
+			SetPrimitiveColor(FActorPrimitiveColorHandler::Get().GetPrimitiveColor(PrimitiveComponent));
 		}
 #endif
 	}
diff --git a/Engine/Source/Runtime/Engine/Public/ShowFlagsValues.inl b/Engine/Source/Runtime/Engine/Public/ShowFlagsValues.inl
--- a/Engine/Source/Runtime/Engine/Public/ShowFlagsValues.inl
+++ b/Engine/Source/Runtime/Engine/Public/ShowFlagsValues.inl
@@ -277,8 +277,8 @@ SHOWFLAG_ALWAYS_ACCESSIBLE(Fog, SFG_Normal, NSLOCTEXT("UnrealEd", "FogSF", "Fog"
 SHOWFLAG_FIXED_IN_SHIPPING(0, Volumes, SFG_Advanced, NSLOCTEXT("UnrealEd", "VolumesSF", "Volumes"))
 /** if this is a game viewport, needed? */
 SHOWFLAG_ALWAYS_ACCESSIBLE(Game, SFG_Hidden, NSLOCTEXT("UnrealEd", "GameSF", "Game"))
-/** Render objects with colors based on what the actors coloring handlers provides (keep level coloration here for viewports backward compatibility) */
-SHOWFLAG_FIXED_IN_SHIPPING(0, ActorColoration, SFG_Advanced, NSLOCTEXT("UnrealEd", "LevelColorationSF", "Level Coloration"))
+/** Render objects with colors based on what the actors coloring handlers provides */
+SHOWFLAG_FIXED_IN_SHIPPING(0, ActorColoration, SFG_Hidden, NSLOCTEXT("UnrealEd", "ActorColorationSF", "Actor Coloration"))
 /** Draws BSP brushes (in game or editor textured triangles usually with lightmaps), for now SHOWFLAG_ALWAYS_ACCESSIBLE because it's exposed in SceneCapture */
 SHOWFLAG_ALWAYS_ACCESSIBLE(BSP, SFG_Normal, NSLOCTEXT("UnrealEd", "BSPSF", "BSP"))
 /** Collision drawing */
