Add names to ShowFlags menu entries so that we can allow list them.

diff --git a/Engine/Source/Editor/LevelEditor/Private/SLevelViewportToolBar.cpp b/Engine/Source/Editor/LevelEditor/Private/SLevelViewportToolBar.cpp
--- a/Engine/Source/Editor/LevelEditor/Private/SLevelViewportToolBar.cpp
+++ b/Engine/Source/Editor/LevelEditor/Private/SLevelViewportToolBar.cpp
@@ -102,8 +102,16 @@ static void FillShowMenuStatic(UToolMenu* Menu, TArray< FLevelViewportCommands::
 	// Assumption: the first 'n' entries types like 'Show All' and 'Hide All' buttons, so insert a separator after them
 	for (int32 EntryIndex = 0; EntryIndex < MenuCommands.Num(); ++EntryIndex)
 	{
+		FName EntryName = NAME_None;
+
+		if (MenuCommands[EntryIndex].ShowMenuItem)
+		{
+			EntryName = MenuCommands[EntryIndex].ShowMenuItem->GetCommandName();
+			ensure(Section.FindEntry(EntryName) == nullptr);
+		}
+
 		Section.AddMenuEntry(
-			NAME_None,
+			EntryName,
 			MenuCommands[EntryIndex].ShowMenuItem,
 			MenuCommands[EntryIndex].LabelOverride
 		);
