Pilot toolbar label converted to a button that selects the actor

diff --git a/Engine/Source/Editor/LevelEditor/Private/LevelViewportActions.cpp b/Engine/Source/Editor/LevelEditor/Private/LevelViewportActions.cpp
--- a/Engine/Source/Editor/LevelEditor/Private/LevelViewportActions.cpp
+++ b/Engine/Source/Editor/LevelEditor/Private/LevelViewportActions.cpp
@@ -44,6 +44,7 @@ void FLevelViewportCommands::RegisterCommands()
 	
 	UI_COMMAND( UseDefaultShowFlags, "Use Defaults", "Resets all show flags to default", EUserInterfaceActionType::Button, FInputChord() );
 
+	UI_COMMAND( SelectPilotedActor, "Select Piloted Actor", "Select the currently piloted actor.", EUserInterfaceActionType::Button, FInputChord() );
 	UI_COMMAND( PilotSelectedActor, "Pilot Selected Actor", "Move the selected actor around using the viewport controls, and bind the viewport to the actor's location and orientation.", EUserInterfaceActionType::Button, FInputChord( EModifierKey::Control | EModifierKey::Shift, EKeys::P ) );
 	UI_COMMAND( EjectActorPilot, "Eject from Actor Pilot", "Stop piloting an actor with the current viewport. Unlocks the viewport's position and orientation from the actor the viewport is currently piloting.", EUserInterfaceActionType::Button, FInputChord() );
 	UI_COMMAND( ToggleActorPilotCameraView, "Actor Pilot Camera View", "Toggles showing the exact camera view when using the viewport to pilot a camera", EUserInterfaceActionType::ToggleButton, FInputChord( EModifierKey::Control | EModifierKey::Shift, EKeys::C ) );
diff --git a/Engine/Source/Editor/LevelEditor/Private/SActorPilotViewportToolbar.cpp b/Engine/Source/Editor/LevelEditor/Private/SActorPilotViewportToolbar.cpp
new file mode 100644
--- /dev/null
+++ b/Engine/Source/Editor/LevelEditor/Private/SActorPilotViewportToolbar.cpp
@@ -0,0 +1,112 @@
+// Copyright Epic Games, Inc. All Rights Reserved.
+
+#include "SActorPilotViewportToolbar.h"
+
+#include "Framework/MultiBox/MultiBoxBuilder.h"
+#include "Framework/MultiBox/MultiBoxDefs.h"
+#include "Layout/Margin.h"
+#include "Layout/Visibility.h"
+#include "LevelEditorViewport.h"
+#include "LevelViewportActions.h"
+#include "Misc/Attribute.h"
+#include "Styling/AppStyle.h"
+#include "Textures/SlateIcon.h"
+#include "Widgets/DeclarativeSyntaxSupport.h"
+#include "Widgets/Images/SImage.h"
+#include "Widgets/Layout/SBorder.h"
+#include "Widgets/Layout/SBox.h"
+#include "Widgets/SBoxPanel.h"
+#include "Widgets/Text/STextBlock.h"
+
+#define LOCTEXT_NAMESPACE "SActorPilotViewportToolbar"
+
+void SActorPilotViewportToolbar::Construct(const FArguments& InArgs)
+{
+	SViewportToolBar::Construct(SViewportToolBar::FArguments());
+
+	Viewport = InArgs._Viewport;
+
+	FLevelViewportCommands& ViewportCommands = FLevelViewportCommands::Get();
+
+	TSharedPtr<const FUICommandList> CommandList = InArgs._Viewport->GetCommandList();
+	FSlimHorizontalToolBarBuilder ToolbarBuilder(CommandList, FMultiBoxCustomization::None);
+
+	// Use a custom style
+	FName ToolBarStyle = "EditorViewportToolBar";
+	ToolbarBuilder.SetStyle(&FAppStyle::Get(), ToolBarStyle);
+	ToolbarBuilder.SetIsFocusable(false);
+
+	ToolbarBuilder.BeginSection("ActorPilot");
+	{
+		ToolbarBuilder.BeginBlockGroup();
+		{
+			static const FName PilotedActorName = "PilotedActor";
+			ToolbarBuilder.AddToolBarButton(
+				ViewportCommands.SelectPilotedActor,
+				NAME_None,
+				TAttribute<FText>::Create([this]()
+					{
+						return GetActiveText();
+					}),
+				TAttribute<FText>(),
+				FSlateIcon(FAppStyle::Get().GetStyleSetName(), "LevelViewport.PilotSelectedActor"),
+				PilotedActorName);
+	
+			static const FName EjectActorPilotName = "EjectActorPilot";
+			ToolbarBuilder.AddToolBarButton(
+				ViewportCommands.EjectActorPilot,
+				NAME_None,
+				FText::GetEmpty(),
+				TAttribute<FText>(),
+				TAttribute<FSlateIcon>(),
+				EjectActorPilotName);
+		
+			static FName ToggleActorPilotCameraViewName = "ToggleActorPilotCameraView";
+			ToolbarBuilder.AddToolBarButton(
+				ViewportCommands.ToggleActorPilotCameraView,
+				NAME_None,
+				FText::GetEmpty(),
+				TAttribute<FText>(),
+				TAttribute<FSlateIcon>(),
+				ToggleActorPilotCameraViewName);
+		}
+		ToolbarBuilder.EndBlockGroup();
+	}
+	ToolbarBuilder.EndSection();
+
+	ChildSlot
+	[
+		SNew( SBorder )
+		.BorderImage( FAppStyle::GetBrush("NoBorder") )
+		.Padding(FMargin(4.f, 0.f))
+		[
+			ToolbarBuilder.MakeWidget()
+		]
+	];
+}
+
+FText SActorPilotViewportToolbar::GetActiveText() const
+{
+	const AActor* Pilot = nullptr;
+	const TSharedPtr<SLevelViewport> ViewportPtr = Viewport.Pin();
+	if (ViewportPtr.IsValid())
+	{
+		Pilot = ViewportPtr->GetLevelViewportClient().GetActiveActorLock().Get();
+	}
+
+	return Pilot ? FText::Format(LOCTEXT("ActiveText", "Pilot Actor: {0} "), FText::FromString(Pilot->GetActorLabel())) : FText();
+}
+
+EVisibility SActorPilotViewportToolbar::GetLockedTextVisibility() const
+{
+	const AActor* Pilot = nullptr;
+	const TSharedPtr<SLevelViewport> ViewportPtr = Viewport.Pin();
+	if (ViewportPtr.IsValid())
+	{
+		Pilot = ViewportPtr->GetLevelViewportClient().GetActiveActorLock().Get();
+	}
+
+	return Pilot && Pilot->IsLockLocation() ? EVisibility::Visible : EVisibility::Collapsed;
+}
+
+#undef LOCTEXT_NAMESPACE
diff --git a/Engine/Source/Editor/LevelEditor/Private/SActorPilotViewportToolbar.h b/Engine/Source/Editor/LevelEditor/Private/SActorPilotViewportToolbar.h
--- a/Engine/Source/Editor/LevelEditor/Private/SActorPilotViewportToolbar.h
+++ b/Engine/Source/Editor/LevelEditor/Private/SActorPilotViewportToolbar.h
@@ -2,26 +2,8 @@
 
 #pragma once
 
-#include "CoreMinimal.h"
-#include "Misc/Attribute.h"
-#include "Layout/Visibility.h"
-#include "Widgets/DeclarativeSyntaxSupport.h"
-#include "Textures/SlateIcon.h"
-#include "Layout/Margin.h"
-#include "Widgets/Images/SImage.h"
-#include "Widgets/SBoxPanel.h"
 #include "SLevelViewport.h"
-#include "Widgets/Text/STextBlock.h"
-#include "Framework/MultiBox/MultiBoxDefs.h"
-#include "Framework/MultiBox/MultiBoxBuilder.h"
-#include "Widgets/Layout/SBorder.h"
-#include "Widgets/Layout/SBox.h"
-#include "Styling/AppStyle.h"
-#include "LevelEditorViewport.h"
 #include "SViewportToolBar.h"
-#include "LevelViewportActions.h"
-
-#define LOCTEXT_NAMESPACE "SActorPilotViewportToolbar"
 
 class SActorPilotViewportToolbar : public SViewportToolBar
 {
@@ -30,106 +12,14 @@ public:
 		SLATE_ARGUMENT(TSharedPtr<SLevelViewport>, Viewport)
 	SLATE_END_ARGS()
 
-	FText GetActiveText() const
-	{
-		const AActor* Pilot = nullptr;
-		auto ViewportPtr = Viewport.Pin();
-		if (ViewportPtr.IsValid())
-		{
-			Pilot = ViewportPtr->GetLevelViewportClient().GetActiveActorLock().Get();
-		}
-		
-		return Pilot ? FText::Format(LOCTEXT("ActiveText", "Pilot Actor: {0} "), FText::FromString(Pilot->GetActorLabel())) : FText();
-	}
-
-	EVisibility GetLockedTextVisibility() const
-	{
-		const AActor* Pilot = nullptr;
-		auto ViewportPtr = Viewport.Pin();
-		if (ViewportPtr.IsValid())
-		{
-			Pilot = ViewportPtr->GetLevelViewportClient().GetActiveActorLock().Get();
-		}
-
-		return Pilot && Pilot->IsLockLocation() ? EVisibility::Visible : EVisibility::Collapsed;
-	}
-
-	void Construct(const FArguments& InArgs)
-	{
-		SViewportToolBar::Construct(SViewportToolBar::FArguments());
-
-		Viewport = InArgs._Viewport;
+	void Construct(const FArguments& InArgs);
 
-		auto& ViewportCommands = FLevelViewportCommands::Get();
+	FText GetActiveText() const;
 
-		FSlimHorizontalToolBarBuilder ToolbarBuilder(InArgs._Viewport->GetCommandList(), FMultiBoxCustomization::None);
-
-		// Use a custom style
-		FName ToolBarStyle = "EditorViewportToolBar";
-		ToolbarBuilder.SetStyle(&FAppStyle::Get(), ToolBarStyle);
-		ToolbarBuilder.SetLabelVisibility(EVisibility::Collapsed);
-		
-		ToolbarBuilder.BeginSection("ActorPilot");
-		{
-			ToolbarBuilder.BeginBlockGroup();
-
-			ToolbarBuilder.AddWidget(
-			SNew(SBorder)
-			.BorderImage(FAppStyle::Get().GetBrush("EditorViewportToolBar.StartToolbarImage"))
-			[
-				SNew(SHorizontalBox)
-				+ SHorizontalBox::Slot()
-				.VAlign(VAlign_Center)
-				.AutoWidth()
-				.Padding(5.0f, 0.0f)
-				[
-						SNew(SImage)
-						.Image(FAppStyle::Get()
-						.GetBrush("LevelViewport.PilotSelectedActor"))
-				]
-				+ SHorizontalBox::Slot()
-				.VAlign(VAlign_Center)
-				[
-					SNew(STextBlock)
-					.Text(this, &SActorPilotViewportToolbar::GetActiveText)
-				]
-				+ SHorizontalBox::Slot()
-				.Padding(5.0f, 0.0f)
-				.VAlign(VAlign_Center)
-				.AutoWidth()
-				[
-					SNew(STextBlock)
-					.Text(LOCTEXT("ActorLockedText", "(Locked)"))
-					.ToolTipText(LOCTEXT("ActorLockedToolTipText", "This actor has locked movement so it will not be updated based on camera position"))
-					.Visibility(this, &SActorPilotViewportToolbar::GetLockedTextVisibility)
-				]
-			]
-		);
-			
-			static const FName EjectActorPilotName = "EjectActorPilot";
-			ToolbarBuilder.AddToolBarButton(ViewportCommands.EjectActorPilot, NAME_None, TAttribute<FText>(), TAttribute<FText>(), TAttribute<FSlateIcon>(), EjectActorPilotName);
-			
-			static FName ToggleActorPilotCameraViewName = "ToggleActorPilotCameraView";
-			ToolbarBuilder.AddToolBarButton(ViewportCommands.ToggleActorPilotCameraView, NAME_None, TAttribute<FText>(), TAttribute<FText>(), TAttribute<FSlateIcon>(), ToggleActorPilotCameraViewName);
-			ToolbarBuilder.EndBlockGroup();
-		}
-		ToolbarBuilder.EndSection();
-
-		ChildSlot
-		[
-			SNew( SBorder )
-			.BorderImage( FAppStyle::GetBrush("NoBorder") )
-			.Padding(FMargin(4.f, 0.f))
-			[
-				ToolbarBuilder.MakeWidget()
-			]
-		];
-	}
+	EVisibility GetLockedTextVisibility() const;
 
 private:
 
 	/** The viewport that we are in */
 	TWeakPtr<SLevelViewport> Viewport;
 };
-
-#undef LOCTEXT_NAMESPACE
diff --git a/Engine/Source/Editor/LevelEditor/Private/SLevelViewport.cpp b/Engine/Source/Editor/LevelEditor/Private/SLevelViewport.cpp
--- a/Engine/Source/Editor/LevelEditor/Private/SLevelViewport.cpp
+++ b/Engine/Source/Editor/LevelEditor/Private/SLevelViewport.cpp
@@ -179,6 +179,23 @@ namespace UE::SLevelViewport::Internal
 
 		return true;
 	}
+
+	// Clears existing selection, then selects the specified actor 
+	void SelectActor(AActor* InActor)
+	{
+		check(InActor);
+
+		const FLevelEditorModule& LevelEditorModule = FModuleManager::LoadModuleChecked<FLevelEditorModule>(LevelEditorName);
+		FEditorModeTools& EditorModeManager = LevelEditorModule.GetFirstLevelEditor()->GetEditorModeManager();
+		
+		// Deselect any currently selected actors
+		EditorModeManager.SelectNone();
+		EditorModeManager.GetSelectedActors()->DeselectAll();
+		EditorModeManager.GetSelectedObjects()->DeselectAll();
+		EditorModeManager.GetSelectedActors()->Select(InActor, true);
+
+		EditorModeManager.ActorSelectionChangeNotify();
+	}
 }
 
 class FLevelViewportDropContextMenuImpl
@@ -1572,6 +1589,12 @@ void SLevelViewport::BindViewCommands( FUICommandList& OutCommandList )
 		FCanExecuteAction::CreateSP( this, &SLevelViewport::CanFindSelectedInLevelScript ) 
 		);
 
+	OutCommandList.MapAction(
+		ViewportActions.SelectPilotedActor,
+		FExecuteAction::CreateSP( this, &SLevelViewport::OnSelectLockedActor ),
+		FCanExecuteAction::CreateSP( this, &SLevelViewport::CanExecuteSelectLockedActor )
+		);
+
 	OutCommandList.MapAction(
 		ViewportActions.EjectActorPilot,
 		FExecuteAction::CreateSP( this, &SLevelViewport::OnActorUnlock ),
@@ -1992,13 +2015,8 @@ void SLevelViewport::OnCreateCameraActor(UClass* InClass)
 	pNewCamera->SetActorRotation( ViewportClient->GetViewRotation() );
 	pNewCamera->GetCameraComponent()->SetFieldOfView( ViewportClient->ViewFOV );
 
-	// Deselect any currently selected actors
-	GUnrealEd->SelectNone( false, true );
-	GEditor->GetSelectedActors()->DeselectAll();
-	GEditor->GetSelectedObjects()->DeselectAll();
-
-	// Select newly created Camera
-	GEditor->SelectActor( pNewCamera, true, true );
+	// Deselect any currently selected actors, then select newly created camera
+	UE::SLevelViewport::Internal::SelectActor(pNewCamera);
 
 	// Send notification about actors that may have changed
 	ULevel::LevelDirtiedEvent.Broadcast();
@@ -2553,6 +2571,25 @@ bool SLevelViewport::CanFindSelectedInLevelScript() const
 	return (Actor != nullptr);
 }
 
+void SLevelViewport::OnSelectLockedActor()
+{
+	if (AActor* LockedActor = LevelViewportClient->GetActiveActorLock().Get())
+	{
+		// Deselect any currently selected actors, then select the locked/piloted actor
+		UE::SLevelViewport::Internal::SelectActor(LockedActor);
+	}
+}
+
+bool SLevelViewport::CanExecuteSelectLockedActor() const
+{
+	if (const AActor* LockedActor = LevelViewportClient->GetActiveActorLock().Get())
+	{
+		return LockedActor->IsSelectable();
+	}
+
+	return false;
+}
+
 void SLevelViewport::OnActorUnlock()
 {
 	if (AActor* LockedActor = LevelViewportClient->GetActiveActorLock().Get())
diff --git a/Engine/Source/Editor/LevelEditor/Public/LevelViewportActions.h b/Engine/Source/Editor/LevelEditor/Public/LevelViewportActions.h
--- a/Engine/Source/Editor/LevelEditor/Public/LevelViewportActions.h
+++ b/Engine/Source/Editor/LevelEditor/Public/LevelViewportActions.h
@@ -171,6 +171,7 @@ PRAGMA_ENABLE_DEPRECATION_WARNINGS
 	TSharedPtr< FUICommandInfo > ClearAllBookmarks;
 
 	/** Actor pilot commands */
+	TSharedPtr< FUICommandInfo > SelectPilotedActor;
 	TSharedPtr< FUICommandInfo > EjectActorPilot;
 	TSharedPtr< FUICommandInfo > PilotSelectedActor;
 
diff --git a/Engine/Source/Editor/LevelEditor/Public/SLevelViewport.h b/Engine/Source/Editor/LevelEditor/Public/SLevelViewport.h
--- a/Engine/Source/Editor/LevelEditor/Public/SLevelViewport.h
+++ b/Engine/Source/Editor/LevelEditor/Public/SLevelViewport.h
@@ -589,6 +589,14 @@ private:
 	/** Can we find the currently selected actor in the level script. */
 	bool CanFindSelectedInLevelScript() const;
 
+	/** Called to select the currently locked actor */
+	void OnSelectLockedActor();
+
+	/**
+	 * @return true if the currently locked actor is selectable
+	 */
+	bool CanExecuteSelectLockedActor() const;
+
 	/** Called to clear the current actor lock */
 	void OnActorUnlock();
 
