LevelEditor: Refactor ULevelViewportToolBarContext so that the ViewportClient is properly accessible through GetViewportClient().

#rnx

diff --git a/Engine/Source/Editor/LevelEditor/Private/SLevelViewportToolBar.cpp b/Engine/Source/Editor/LevelEditor/Private/SLevelViewportToolBar.cpp
--- a/Engine/Source/Editor/LevelEditor/Private/SLevelViewportToolBar.cpp
+++ b/Engine/Source/Editor/LevelEditor/Private/SLevelViewportToolBar.cpp
@@ -825,12 +825,12 @@ TSharedRef<SWidget> SLevelViewportToolBar::GenerateDevicePreviewMenu() const
 		Menu->AddDynamicSection("DynamicSection", FNewToolMenuDelegate::CreateLambda([](UToolMenu* InMenu)
 		{
 			ULevelViewportToolBarContext* Context = InMenu->FindContext<ULevelViewportToolBarContext>();
-			Context->LevelViewportToolBarWidgetConst.Pin()->FillDevicePreviewMenu(InMenu);
+			Context->LevelViewportToolBarWidget.Pin()->FillDevicePreviewMenu(InMenu);
 		}));
 	}
 
 	ULevelViewportToolBarContext* ContextObject = NewObject<ULevelViewportToolBarContext>();
-	ContextObject->LevelViewportToolBarWidgetConst = SharedThis(this);
+	ContextObject->LevelViewportToolBarWidget = ConstCastSharedRef<SLevelViewportToolBar>(SharedThis(this));
 
 	FToolMenuContext MenuContext(Viewport.Pin()->GetCommandList(), TSharedPtr<FExtender>(), ContextObject);
 	return UToolMenus::Get()->GenerateWidget(MenuName, MenuContext);
@@ -1001,14 +1001,14 @@ TSharedRef<SWidget> SLevelViewportToolBar::GenerateCameraMenu() const
 		Menu->AddDynamicSection("DynamicSection", FNewToolMenuDelegate::CreateLambda([](UToolMenu* InMenu)
 		{
 			ULevelViewportToolBarContext* Context = InMenu->FindContext<ULevelViewportToolBarContext>();
-			Context->LevelViewportToolBarWidgetConst.Pin()->FillCameraMenu(InMenu);
+			Context->LevelViewportToolBarWidget.Pin()->FillCameraMenu(InMenu);
 		}));
 	}
 
 	Viewport.Pin()->OnFloatingButtonClicked();
 
 	ULevelViewportToolBarContext* ContextObject = NewObject<ULevelViewportToolBarContext>();
-	ContextObject->LevelViewportToolBarWidgetConst = SharedThis(this);
+	ContextObject->LevelViewportToolBarWidget = ConstCastSharedRef<SLevelViewportToolBar>(SharedThis(this));
 
 	FToolMenuContext MenuContext(Viewport.Pin()->GetCommandList(), TSharedPtr<FExtender>(), ContextObject);
 	return UToolMenus::Get()->GenerateWidget(MenuName, MenuContext);
@@ -1199,7 +1199,7 @@ TSharedRef<SWidget> SLevelViewportToolBar::GenerateShowMenu() const
 		Menu->AddDynamicSection("LevelDynamicSection", FNewToolMenuDelegate::CreateLambda([](UToolMenu* InMenu)
 		{
 			ULevelViewportToolBarContext* Context = InMenu->FindContext<ULevelViewportToolBarContext>();
-			Context->LevelViewportToolBarWidgetConst.Pin()->FillShowMenu(InMenu);
+			Context->LevelViewportToolBarWidget.Pin()->FillShowMenu(InMenu);
 		}));
 	}
 
@@ -1212,7 +1212,7 @@ TSharedRef<SWidget> SLevelViewportToolBar::GenerateShowMenu() const
 	TSharedPtr<FExtender> MenuExtender = LevelEditorModule.AssembleExtenders(CommandList, LevelEditorModule.GetAllLevelViewportShowMenuExtenders());
 
 	ULevelViewportToolBarContext* ContextObject = NewObject<ULevelViewportToolBarContext>();
-	ContextObject->LevelViewportToolBarWidgetConst = SharedThis(this);
+	ContextObject->LevelViewportToolBarWidget = ConstCastSharedRef<SLevelViewportToolBar>(SharedThis(this));
 
 	FToolMenuContext MenuContext(CommandList, MenuExtender, ContextObject);
 	return UToolMenus::Get()->GenerateWidget(MenuName, MenuContext);
diff --git a/Engine/Source/Editor/LevelEditor/Public/LevelEditorMenuContext.h b/Engine/Source/Editor/LevelEditor/Public/LevelEditorMenuContext.h
--- a/Engine/Source/Editor/LevelEditor/Public/LevelEditorMenuContext.h
+++ b/Engine/Source/Editor/LevelEditor/Public/LevelEditorMenuContext.h
@@ -86,7 +86,6 @@ class LEVELEDITOR_API ULevelViewportToolBarContext : public UObject
 	GENERATED_BODY()
 public:
 	TWeakPtr<SLevelViewportToolBar> LevelViewportToolBarWidget;
-	TWeakPtr<const SLevelViewportToolBar> LevelViewportToolBarWidgetConst;
 
 	FLevelEditorViewportClient* GetLevelViewportClient();
 };
