Fix crash in level editor context menu for delta transform to actors

diff --git a/Engine/Source/Editor/LevelEditor/Private/LevelEditorContextMenu.cpp b/Engine/Source/Editor/LevelEditor/Private/LevelEditorContextMenu.cpp
--- a/Engine/Source/Editor/LevelEditor/Private/LevelEditorContextMenu.cpp
+++ b/Engine/Source/Editor/LevelEditor/Private/LevelEditorContextMenu.cpp
@@ -1143,7 +1143,7 @@ void FLevelEditorContextMenuImpl::FillTransformMenu(UToolMenu* Menu)
 {
 	if (ULevelEditorContextMenuContext* LevelEditorContext = Menu->FindContext<ULevelEditorContextMenuContext>())
 	{
-		if (LevelEditorContext->CurrentSelection->GetNumSelectedElements() == 0)
+		if (!LevelEditorContext->CurrentSelection || LevelEditorContext->CurrentSelection->GetNumSelectedElements() == 0)
 		{
 			return;
 		}
