Fixing Preview modes
- Certain "SM5" platforms had FL SM6 requirements when being previewed.
- Adding a "PreviewFeatureLevel" to shader PreviewPlatform's to override the preview's feature level requirement.
- When changing between shader previews of the same feature level, the scene and all associated work should be reset just like when we change to a preview with a different feature level.
- During scene destruction, the VSM cache shouldn't need to process invalidations, especially since the preview shader platform has already been changed.
- Removing code that prevented switching between preview platforms that had different VSM settings.

diff --git a/Engine/Source/Editor/LevelEditor/Private/LevelEditor.cpp b/Engine/Source/Editor/LevelEditor/Private/LevelEditor.cpp
--- a/Engine/Source/Editor/LevelEditor/Private/LevelEditor.cpp
+++ b/Engine/Source/Editor/LevelEditor/Private/LevelEditor.cpp
@@ -1814,11 +1814,20 @@ void FLevelEditorModule::BindGlobalLevelEditorCommands()
 		
 		if (ShaderPlatform < SP_NumPlatforms)
 		{
-			const ERHIFeatureLevel::Type FeatureLevel = GetMaxSupportedFeatureLevel(ShaderPlatform);
-
 			const bool bIsDefaultShaderPlatform = FDataDrivenShaderPlatformInfo::GetPreviewShaderPlatformParent(ShaderPlatform) == GMaxRHIShaderPlatform;
 
-			FPreviewPlatformInfo PreviewFeatureLevelInfo(bIsDefaultShaderPlatform ? GMaxRHIFeatureLevel : FeatureLevel, bIsDefaultShaderPlatform ? GMaxRHIShaderPlatform : (EShaderPlatform)ShaderPlatform, bIsDefaultShaderPlatform ? NAME_None : Item.PlatformName, bIsDefaultShaderPlatform ? NAME_None : Item.ShaderFormat, bIsDefaultShaderPlatform ? NAME_None : Item.DeviceProfileName, true, bIsDefaultShaderPlatform ? NAME_None : Item.PreviewShaderPlatformName);
+			auto GetPreviewFeatureLevelInfo = [&]()
+			{
+				if (bIsDefaultShaderPlatform)
+				{
+					return FPreviewPlatformInfo(GMaxRHIFeatureLevel, GMaxRHIShaderPlatform, NAME_None, NAME_None, NAME_None, true, NAME_None);
+				}
+
+				const ERHIFeatureLevel::Type FeatureLevel = GetMaxSupportedFeatureLevel(ShaderPlatform);
+				return FPreviewPlatformInfo(FeatureLevel, ShaderPlatform, Item.PlatformName, Item.ShaderFormat, Item.DeviceProfileName, true, Item.PreviewShaderPlatformName);
+			};
+
+			FPreviewPlatformInfo PreviewFeatureLevelInfo = GetPreviewFeatureLevelInfo();
 
 			ActionList.MapAction(
 				Commands.PreviewPlatformOverrides[Index],
diff --git a/Engine/Source/Editor/LevelEditor/Private/LevelEditorActions.cpp b/Engine/Source/Editor/LevelEditor/Private/LevelEditorActions.cpp
--- a/Engine/Source/Editor/LevelEditor/Private/LevelEditorActions.cpp
+++ b/Engine/Source/Editor/LevelEditor/Private/LevelEditorActions.cpp
@@ -692,44 +692,28 @@ bool FLevelEditorActionCallbacks::CanExecutePreviewPlatform(FPreviewPlatformInfo
 		return false;
 	}
 
-	// TODO: Prevent switching from a platform with VSM to a platform without VSM at the same FeatureLevel until all issues are resolved.
-	// (for instance, GlobalShaderMap does not contain the shaders necessary for FVirtualShadowMapArrayCacheManager::ProcessInvalidations anymore)
-	// Currently this is mostly meant to isolate going to/from VULKAN_SM5 which is wildly different and causes issues with preview mechanics.
-	if (NewPreviewPlatform.PreviewFeatureLevel == GMaxRHIFeatureLevel)
-	{
-		if (FDataDrivenShaderPlatformInfo::IsValid(NewPreviewPlatform.ShaderPlatform) &&
-			FDataDrivenShaderPlatformInfo::GetIsPreviewPlatform(NewPreviewPlatform.ShaderPlatform))
-		{
-			const EShaderPlatform ParentShaderPlatform = FDataDrivenShaderPlatformInfo::GetPreviewShaderPlatformParent(NewPreviewPlatform.ShaderPlatform);
-			if (DoesPlatformSupportVirtualShadowMaps(GMaxRHIShaderPlatform) != DoesPlatformSupportVirtualShadowMaps(ParentShaderPlatform))
-			{
-				return false;
-			}
-		}
-	}
+	const EShaderPlatform PreviewShaderPlatform = NewPreviewPlatform.ShaderPlatform;
 
-	// TODO: Prevent previewing VULKAN_SM5 with a D3D renderer until all issues are resolved.
-	if (FDataDrivenShaderPlatformInfo::IsValid(NewPreviewPlatform.ShaderPlatform) &&
-		FDataDrivenShaderPlatformInfo::GetIsPreviewPlatform(NewPreviewPlatform.ShaderPlatform) &&
-		FDataDrivenShaderPlatformInfo::GetIsLanguageD3D(GMaxRHIShaderPlatform))
+	if (FDataDrivenShaderPlatformInfo::IsValid(PreviewShaderPlatform) && FDataDrivenShaderPlatformInfo::GetIsPreviewPlatform(PreviewShaderPlatform))
 	{
-		const EShaderPlatform ParentShaderPlatform = FDataDrivenShaderPlatformInfo::GetPreviewShaderPlatformParent(NewPreviewPlatform.ShaderPlatform);
-		if (FDataDrivenShaderPlatformInfo::GetIsLanguageVulkan(ParentShaderPlatform) && IsFeatureLevelSupported(ParentShaderPlatform, ERHIFeatureLevel::SM5))
+		const EShaderPlatform RealShaderPlatform = FDataDrivenShaderPlatformInfo::GetPreviewShaderPlatformParent(PreviewShaderPlatform);
+
+		// TODO: Prevent previewing VULKAN_SM5 with a D3D renderer until all issues are resolved.
+		// We have to use the real platform here since these fields are overridden in preview.
+		if (GDynamicRHI
+			&& (GDynamicRHI->GetInterfaceType() == ERHIInterfaceType::D3D12 || GDynamicRHI->GetInterfaceType() == ERHIInterfaceType::D3D11)
+			&& IsVulkanPlatform(RealShaderPlatform)
+			&& IsFeatureLevelSupported(RealShaderPlatform, ERHIFeatureLevel::SM5))
 		{
 			return false;
 		}
-	}
 
-	// When the preview platform's DDSPI MaxSamplers is > 16 and the current RHI device has support
-	// for > 16 samplers we rely on the shader compiler being able to choose an appropriate profile for the
-	// preview feature level that supports > 16 samplers. On D3D12 SM5 the D3D shader compiler will use Dxc and
-	// sm6.0. Vulkan SM5 also appears to handle > 16 samplers fine.
-	// TODO: Look into support for Metal (MacOS).
-	if (FDataDrivenShaderPlatformInfo::IsValid(NewPreviewPlatform.ShaderPlatform) &&
-		FDataDrivenShaderPlatformInfo::GetIsPreviewPlatform(NewPreviewPlatform.ShaderPlatform))
-	{
-		const EShaderPlatform ParentShaderPlatform = FDataDrivenShaderPlatformInfo::GetPreviewShaderPlatformParent(NewPreviewPlatform.ShaderPlatform);
-		if (GMaxTextureSamplers < (int32)FDataDrivenShaderPlatformInfo::GetMaxSamplers(ParentShaderPlatform))
+		// When the preview platform's DDSPI MaxSamplers is > 16 and the current RHI device has support
+		// for > 16 samplers we rely on the shader compiler being able to choose an appropriate profile for the
+		// preview feature level that supports > 16 samplers. On D3D12 SM5 the D3D shader compiler will use Dxc and
+		// sm6.0. Vulkan SM5 also appears to handle > 16 samplers fine.
+		// TODO: Look into support for Metal (MacOS).
+		if ((int32)FDataDrivenShaderPlatformInfo::GetMaxSamplers(PreviewShaderPlatform) > GRHIGlobals.MaxTextureSamplers)
 		{
 			return false;
 		}
diff --git a/Engine/Source/Editor/LevelEditor/Private/SLevelEditor.cpp b/Engine/Source/Editor/LevelEditor/Private/SLevelEditor.cpp
--- a/Engine/Source/Editor/LevelEditor/Private/SLevelEditor.cpp
+++ b/Engine/Source/Editor/LevelEditor/Private/SLevelEditor.cpp
@@ -250,7 +250,7 @@ void SLevelEditor::Construct( const SLevelEditor::FArguments& InArgs)
 			{
 				GEditor->BuildReflectionCaptures();
 			}
-			World->ChangeFeatureLevel(NewFeatureLevel);
+			World->ChangeFeatureLevel(NewFeatureLevel, true, true);
 		});
 
 	PreviewPlatformChangedHandle = GEditor->OnPreviewPlatformChanged().AddLambda([this]()
diff --git a/Engine/Source/Runtime/Core/Private/Misc/DataDrivenPlatformInfoRegistry.cpp b/Engine/Source/Runtime/Core/Private/Misc/DataDrivenPlatformInfoRegistry.cpp
--- a/Engine/Source/Runtime/Core/Private/Misc/DataDrivenPlatformInfoRegistry.cpp
+++ b/Engine/Source/Runtime/Core/Private/Misc/DataDrivenPlatformInfoRegistry.cpp
@@ -291,6 +291,8 @@ static void ParsePreviewPlatforms(const FConfigFile& IniFile)
 				Item.InactiveIconPath = GetSectionString(Section.Value, FName("InactiveIconPath"));
 				Item.InactiveIconName = *GetSectionString(Section.Value, FName("InactiveIconName"));
 				Item.ShaderPlatformToPreview = *GetSectionString(Section.Value, FName("ShaderPlatform"));
+				Item.PreviewFeatureLevelName = *GetSectionString(Section.Value, FName("PreviewFeatureLevel"));
+
 				checkf(Item.ShaderPlatformToPreview != NAME_None, TEXT("DataDrivenPlatformInfo section [PreviewPlatform %s] must specify a ShaderPlatform"), *SectionName);
 				FTextStringHelper::ReadFromBuffer(*GetSectionString(Section.Value, FName("MenuTooltip")), Item.MenuTooltip);
 				FTextStringHelper::ReadFromBuffer(*GetSectionString(Section.Value, FName("IconText")), Item.IconText);
diff --git a/Engine/Source/Runtime/Core/Public/Misc/DataDrivenPlatformInfoRegistry.h b/Engine/Source/Runtime/Core/Public/Misc/DataDrivenPlatformInfoRegistry.h
--- a/Engine/Source/Runtime/Core/Public/Misc/DataDrivenPlatformInfoRegistry.h
+++ b/Engine/Source/Runtime/Core/Public/Misc/DataDrivenPlatformInfoRegistry.h
@@ -101,6 +101,7 @@ struct FPreviewPlatformMenuItem
 	FName ShaderFormat;
 	FName ShaderPlatformToPreview;
 	FName PreviewShaderPlatformName;
+	FName PreviewFeatureLevelName;
 	FString ActiveIconPath;
 	FName ActiveIconName;
 	FString InactiveIconPath;
diff --git a/Engine/Source/Runtime/Engine/Classes/Engine/World.h b/Engine/Source/Runtime/Engine/Classes/Engine/World.h
--- a/Engine/Source/Runtime/Engine/Classes/Engine/World.h
+++ b/Engine/Source/Runtime/Engine/Classes/Engine/World.h
@@ -1377,7 +1377,7 @@ public:
 #if WITH_EDITOR
 
 	/** Change the feature level that this world is current rendering with */
-	void ChangeFeatureLevel(ERHIFeatureLevel::Type InFeatureLevel, bool bShowSlowProgressDialog = true);
+	void ChangeFeatureLevel(ERHIFeatureLevel::Type InFeatureLevel, bool bShowSlowProgressDialog = true, bool bForceUpdate = false);
 
 	void ShaderPlatformChanged();
 
diff --git a/Engine/Source/Runtime/Engine/Private/World.cpp b/Engine/Source/Runtime/Engine/Private/World.cpp
--- a/Engine/Source/Runtime/Engine/Private/World.cpp
+++ b/Engine/Source/Runtime/Engine/Private/World.cpp
@@ -8928,11 +8928,11 @@ void UWorld::DuplicateRequestedLevels(const FName MapName)
 }
 
 #if WITH_EDITOR
-void UWorld::ChangeFeatureLevel(ERHIFeatureLevel::Type InFeatureLevel, bool bShowSlowProgressDialog )
+void UWorld::ChangeFeatureLevel(ERHIFeatureLevel::Type InFeatureLevel, bool bShowSlowProgressDialog, bool bForceUpdate)
 {
-	if (InFeatureLevel != GetFeatureLevel())
+	if (InFeatureLevel != GetFeatureLevel() || bForceUpdate)
 	{
-		UE_LOG(LogWorld, Log, TEXT("Changing Feature Level (Enum) from %i to %i"), (int)GetFeatureLevel(), (int)InFeatureLevel);
+		UE_LOG(LogWorld, Log, TEXT("Changing Feature Level (Enum) from %i to %i%s"), (int)GetFeatureLevel(), (int)InFeatureLevel, bForceUpdate ? TEXT(" (forced)") : TEXT(""));
 		FScopedSlowTask SlowTask(100.f, NSLOCTEXT("Engine", "ChangingPreviewRenderingLevelMessage", "Changing Preview Rendering Level"), bShowSlowProgressDialog);
 		SlowTask.MakeDialog();
 		{
diff --git a/Engine/Source/Runtime/RHI/Private/DataDrivenShaderPlatformInfo.cpp b/Engine/Source/Runtime/RHI/Private/DataDrivenShaderPlatformInfo.cpp
--- a/Engine/Source/Runtime/RHI/Private/DataDrivenShaderPlatformInfo.cpp
+++ b/Engine/Source/Runtime/RHI/Private/DataDrivenShaderPlatformInfo.cpp
@@ -343,6 +343,7 @@ void FGenericDataDrivenShaderPlatformInfo::Initialize()
 			if (Section.Key.StartsWith(TEXT("ShaderPlatform ")))
 			{
 				const FString& SectionName = Section.Key;
+				const FConfigSection& SectionSettings = Section.Value;
 
 				// get enum value for the string name
 				const EShaderPlatform ShaderPlatform = ParseShaderPlatform(*SectionName.Mid(15));
@@ -360,7 +361,7 @@ void FGenericDataDrivenShaderPlatformInfo::Initialize()
 				// at this point, we can start pulling information out
 				Infos[ShaderPlatform].Name = *SectionName.Mid(15);
 				PlatformNameToShaderPlatformMap.FindOrAdd(Infos[ShaderPlatform].Name) = ShaderPlatform;
-				ParseDataDrivenShaderInfo(Section.Value, ShaderPlatform);
+				ParseDataDrivenShaderInfo(SectionSettings, ShaderPlatform);
 				Infos[ShaderPlatform].bContainsValidPlatformInfo = true;
 
 #if WITH_EDITOR
@@ -375,7 +376,7 @@ void FGenericDataDrivenShaderPlatformInfo::Initialize()
 							const EShaderPlatform PreviewShaderPlatform = EShaderPlatform(CustomShaderPlatform++);
 							FGenericDataDrivenShaderPlatformInfo& PreviewInfo = Infos[PreviewShaderPlatform];
 							PreviewInfo.Name = Item.PreviewShaderPlatformName;
-							ParseDataDrivenShaderInfo(Section.Value, PreviewShaderPlatform);
+							ParseDataDrivenShaderInfo(SectionSettings, PreviewShaderPlatform);
 							PreviewInfo.bIsPreviewPlatform = true;
 							PreviewInfo.bContainsValidPlatformInfo = true;
 
@@ -386,6 +387,12 @@ void FGenericDataDrivenShaderPlatformInfo::Initialize()
 								PreviewEditorInfo.FriendlyName = Item.OptionalFriendlyNameOverride;
 							}
 
+							ERHIFeatureLevel::Type PreviewFeatureLevel = ERHIFeatureLevel::Num;
+							if (GetFeatureLevelFromName(Item.PreviewFeatureLevelName, PreviewFeatureLevel))
+							{
+								PreviewInfo.MaxFeatureLevel = PreviewFeatureLevel;
+							}
+
 							PlatformNameToShaderPlatformMap.FindOrAdd(PreviewInfo.Name) = PreviewShaderPlatform;
 						}
 					}
diff --git a/Engine/Source/Runtime/Renderer/Private/RendererScene.cpp b/Engine/Source/Runtime/Renderer/Private/RendererScene.cpp
--- a/Engine/Source/Runtime/Renderer/Private/RendererScene.cpp
+++ b/Engine/Source/Runtime/Renderer/Private/RendererScene.cpp
@@ -4598,7 +4598,16 @@ void FScene::Release()
 		[Scene](FRHICommandListImmediate& RHICmdList)
 		{
 			// Flush any remaining batched primitive update commands before deleting the scene.
-			Scene->UpdateAllPrimitiveSceneInfos(RHICmdList);
+			FUpdateParameters UpdateParameters;
+			UpdateParameters.bDestruction = true;
+
+			// Scope required so that the GraphBuilder is destructed before this Scene
+			{
+				FRDGBuilder GraphBuilder(RHICmdList, FRDGEventName(TEXT("UpdateAllPrimitiveSceneInfos")));
+				Scene->Update(GraphBuilder, UpdateParameters);
+				GraphBuilder.Execute();
+			}
+
 			delete Scene;
 		});
 }
@@ -5530,18 +5539,21 @@ void FScene::Update(FRDGBuilder& GraphBuilder, const FUpdateParameters& Paramete
 	RemovedPrimitiveSceneInfos.Empty();
 	bool bAnySceneUpdatesQueued = RemovedLocalPrimitiveSceneInfos.Num() + AddedPrimitiveSceneInfos.Num() + UpdatedTransforms.Num() + UpdatedInstances.Num() > 0;
 	RemovedLocalPrimitiveSceneInfos.Sort(FPrimitiveArraySortKey());
+
 	GPUScene.OnPreSceneUpdate(GraphBuilder, SceneUpdateChangeSetStorage.GetPreUpdateSet());
 
 	// Create a SceneUB that permits access to the scene for invalidation processing.
 	FSceneUniformBuffer SceneUB;
 	GPUScene.FillSceneUniformBuffer(GraphBuilder, SceneUB);
 
-	auto &SceneCullingUpdater = SceneCulling->BeginUpdate(GraphBuilder, SceneUB, bAnySceneUpdatesQueued);
+	FSceneCulling::FUpdater& SceneCullingUpdater = SceneCulling->BeginUpdate(GraphBuilder, SceneUB, bAnySceneUpdatesQueued);
 	SceneCullingUpdater.OnPreSceneUpdate(GraphBuilder, SceneUpdateChangeSetStorage.GetPreUpdateSet());
-	
-	auto& SceneExtensionsUpdaters = *GraphBuilder.AllocObject<FSceneExtensionsUpdaters>(*this);
+
+	FSceneExtensionsUpdaters& SceneExtensionsUpdaters = *GraphBuilder.AllocObject<FSceneExtensionsUpdaters>(*this);
 	SceneExtensionsUpdaters.PreSceneUpdate(GraphBuilder, SceneUpdateChangeSetStorage.GetPreUpdateSet());
 
+	// Don't queue VSM invalidations when being destroyed. This avoids issues on preview mode change when the new preview platform doesn't use VSM.
+	if (!Parameters.bDestruction)
 	{
 		SCOPED_NAMED_EVENT(FScene_VirtualShadowCacheUpdate, FColor::Orange);
 		FVirtualShadowMapArrayCacheManager* CacheManager = GetVirtualShadowMapCache();
diff --git a/Engine/Source/Runtime/Renderer/Private/SceneCulling/SceneCulling.cpp b/Engine/Source/Runtime/Renderer/Private/SceneCulling/SceneCulling.cpp
--- a/Engine/Source/Runtime/Renderer/Private/SceneCulling/SceneCulling.cpp
+++ b/Engine/Source/Runtime/Renderer/Private/SceneCulling/SceneCulling.cpp
@@ -2204,6 +2204,12 @@ public:
 
 	void UploadToGPU(FRDGBuilder& GraphBuilder, FSceneUniformBuffer& SceneUniformBuffer)
 	{
+		// Early out if for some reason the platform has changed from underneath us.
+		if (!UseSceneCulling(SceneCulling.Scene.GetShaderPlatform()))
+		{
+			return;
+		}
+
 		BUILDER_LOG("UploadToGPU %d", ItemChunkUploader.GetNumScatters());
 
 		bool bValidToCapture = CellHeaderUploader.GetNumScatters() > 0;
diff --git a/Engine/Source/Runtime/Renderer/Private/ScenePrivate.h b/Engine/Source/Runtime/Renderer/Private/ScenePrivate.h
--- a/Engine/Source/Runtime/Renderer/Private/ScenePrivate.h
+++ b/Engine/Source/Runtime/Renderer/Private/ScenePrivate.h
@@ -3236,6 +3236,7 @@ public:
 	{
 		EUpdateAllPrimitiveSceneInfosAsyncOps AsyncOps = EUpdateAllPrimitiveSceneInfosAsyncOps::None;
 		UE::Tasks::FTask GPUSceneUpdateTaskPrerequisites;
+		bool bDestruction = false;
 
 		struct
 		{
