Remove Mac-only special case for how "Select an object to view actions" appears in the Actor menu. It is not needed anymore now that the main menu appears in-editor on the Mac.

#tests Ran on a remote Mac with and without change.

diff --git a/Engine/Source/Editor/LevelEditor/Private/LevelEditorContextMenu.cpp b/Engine/Source/Editor/LevelEditor/Private/LevelEditorContextMenu.cpp
--- a/Engine/Source/Editor/LevelEditor/Private/LevelEditorContextMenu.cpp
+++ b/Engine/Source/Editor/LevelEditor/Private/LevelEditorContextMenu.cpp
@@ -669,15 +669,6 @@ void FLevelEditorContextMenu::RegisterMenuBarEmptyContextMenu()
 
 	const FText EmptySelectionInformationalMessage = LOCTEXT("EmptySelectionInformationalMessage", "Select an object to view actions.");
 
-#if PLATFORM_MAC
-	// Can't include arbitrary widgets in a main menu on Mac, so display the informational message using a disabled entry.
-	Section.AddMenuEntry(
-		NAME_None,
-		EmptySelectionInformationalMessage,
-		TAttribute<FText>(),
-		TAttribute<FSlateIcon>(),
-		FUIAction(FExecuteAction(), FCanExecuteAction::CreateLambda([]() { return false; })));
-#else
 	Section.AddEntry(FToolMenuEntry::InitWidget(
 		NAME_None,
 		SNew(SBox)
@@ -688,7 +679,6 @@ void FLevelEditorContextMenu::RegisterMenuBarEmptyContextMenu()
 			.TextStyle(FAppStyle::Get(), "HintText")
 		],
 		FText::GetEmpty(), /*bNoIndent*/ true, /*bSearchable*/ false));
-#endif
 }
 
 void FLevelEditorContextMenu::RegisterEmptySelectionContextMenu()
