Support asset tiles in main menu on the Mac

Since the main menu moved into the editor window on the Mac, we can support custom Slate widgets in the menus. This change removes a special case where asset tiles were disabled on the Mac in certain menus, for example under the "RECENTLY PLACED" section in "Actor > PlaceActor" in the level editor main menu.

diff --git a/Engine/Source/Editor/LevelEditor/Private/LevelEditorCreateActorMenu.cpp b/Engine/Source/Editor/LevelEditor/Private/LevelEditorCreateActorMenu.cpp
--- a/Engine/Source/Editor/LevelEditor/Private/LevelEditorCreateActorMenu.cpp
+++ b/Engine/Source/Editor/LevelEditor/Private/LevelEditorCreateActorMenu.cpp
@@ -344,17 +344,6 @@ static void BuildSingleAssetAddReplaceActorMenu(FToolMenuSection& Section, const
 		return;
 	}
 
-#if PLATFORM_MAC
-	// Cannot use asset tile if this is being shown in the Mac global menu bar, force a normal menu entry
-	if (ULevelEditorContextMenuContext* Context = Section.FindContext<ULevelEditorContextMenuContext>())
-	{
-		if (Context->ContextType == ELevelEditorMenuContext::MainMenu)
-		{
-			bUseAssetTile = false;
-		}
-	}
-#endif
-
 	if ( AssetMenuOptions.Num() == 1 )
 	{
 		const FActorFactoryAssetProxy::FMenuItem& MenuItem = AssetMenuOptions[0];
