Expose realtime toggle to level editor subsystem
This allows utilities to disable realtime for better performance during intenstive tasks

#tests Ran function to enable / disable realtime via EUW.

diff --git a/Engine/Source/Editor/LevelEditor/Private/LevelEditorSubsystem.cpp b/Engine/Source/Editor/LevelEditor/Private/LevelEditorSubsystem.cpp
--- a/Engine/Source/Editor/LevelEditor/Private/LevelEditorSubsystem.cpp
+++ b/Engine/Source/Editor/LevelEditor/Private/LevelEditorSubsystem.cpp
@@ -239,6 +239,26 @@ void ULevelEditorSubsystem::EditorInvalidateViewports()
 	}
 }
 
+void ULevelEditorSubsystem::EditorSetViewportRealtime(bool bInRealtime, FName ViewportConfigKey)
+{
+	TSharedPtr<SLevelViewport> LevelViewport = InternalEditorLevelLibrary::GetLevelViewport(ViewportConfigKey);
+	if (LevelViewport.IsValid())
+	{
+		FLevelEditorViewportClient& LevelViewportClient = LevelViewport->GetLevelViewportClient();
+
+		FText RealTimeOverrideSystemName = LOCTEXT("LevelEditorSubsystemRealtimeOverride", "Level Editor Subsystem Realtime Override");
+
+		if (bInRealtime)
+		{
+			LevelViewportClient.RemoveRealtimeOverride(RealTimeOverrideSystemName);
+		}
+		else
+		{
+			LevelViewportClient.AddRealtimeOverride(false, RealTimeOverrideSystemName);
+		}
+	}
+}
+
 void ULevelEditorSubsystem::EditorSetGameView(bool bGameView, FName ViewportConfigKey)
 {
 	TSharedPtr<SLevelViewport> LevelViewport = InternalEditorLevelLibrary::GetLevelViewport(ViewportConfigKey);
diff --git a/Engine/Source/Editor/LevelEditor/Public/LevelEditorSubsystem.h b/Engine/Source/Editor/LevelEditor/Public/LevelEditorSubsystem.h
--- a/Engine/Source/Editor/LevelEditor/Public/LevelEditorSubsystem.h
+++ b/Engine/Source/Editor/LevelEditor/Public/LevelEditorSubsystem.h
@@ -58,6 +58,9 @@ public:
 	UFUNCTION(BlueprintCallable, Category = "Editor Scripting | Level Utility", meta = (DevelopmentOnly))
 	void EditorInvalidateViewports();
 
+	UFUNCTION(BlueprintCallable, Category = "Editor Scripting | Level Utility", meta = (DevelopmentOnly))
+	void EditorSetViewportRealtime(bool bInRealtime, FName ViewportConfigKey = NAME_None);
+
 	UFUNCTION(BlueprintCallable, Category = "Editor Scripting | Level Utility", meta = (DevelopmentOnly))
 	void EditorSetGameView(bool bGameView, FName ViewportConfigKey = NAME_None);
 
