Make a few menus non-searchable to hide their menu search fields

This change removes duplicate search bars in the following menus
* Actor > Attach To" (main menu, viewport right click, details panel right click)
* Level Editor Toolbar > Blueprint Icon > Open Blueprint Class...
* Level Editor Toolbar > Blueprint Icon > GameMode Override > Create
* Level Editor Toolbar > Blueprint Icon > GameMode Override > Select ... Class
* Level Editor Toolbar > Media Plate Icon > Select a Media Profile
* Viewport Hamburger Icon > Layouts
* Viewport > Lit > OCIO Display

diff --git a/Engine/Plugins/Compositing/OpenColorIO/Source/OpenColorIOEditor/Private/OpenColorIOEditorModule.cpp b/Engine/Plugins/Compositing/OpenColorIO/Source/OpenColorIOEditor/Private/OpenColorIOEditorModule.cpp
--- a/Engine/Plugins/Compositing/OpenColorIO/Source/OpenColorIOEditor/Private/OpenColorIOEditorModule.cpp
+++ b/Engine/Plugins/Compositing/OpenColorIO/Source/OpenColorIOEditor/Private/OpenColorIOEditorModule.cpp
@@ -172,6 +172,10 @@ void FOpenColorIOEditorModule::UnregisterViewMenuExtension()
 
 void FOpenColorIOEditorModule::AddOpenColorIODisplaySubMenu(UToolMenu* Menu)
 {
+	// Disable searching in this parent menu because the contents of this menu is added as a
+	// menu-in-a-menu which leads to double search fields if the parent is searchable too.
+	Menu->bSearchable = false;
+
 	//Viewport menu was clicked, get which one was hit
 	FViewport* CurrentViewport = GEditor->GetActiveViewport();
 
diff --git a/Engine/Plugins/Media/MediaFrameworkUtilities/Source/MediaFrameworkUtilitiesEditor/Private/UI/MediaProfileMenuEntry.cpp b/Engine/Plugins/Media/MediaFrameworkUtilities/Source/MediaFrameworkUtilitiesEditor/Private/UI/MediaProfileMenuEntry.cpp
--- a/Engine/Plugins/Media/MediaFrameworkUtilities/Source/MediaFrameworkUtilitiesEditor/Private/UI/MediaProfileMenuEntry.cpp
+++ b/Engine/Plugins/Media/MediaFrameworkUtilities/Source/MediaFrameworkUtilitiesEditor/Private/UI/MediaProfileMenuEntry.cpp
@@ -181,6 +181,10 @@ struct FMediaProfileMenuEntryImpl
 
 	void AddObjectSubMenu(FMenuBuilder& MenuBuilder)
 	{
+		// The contents of this menu are added as a custom widget with its own search field so we
+		// disable searching in this parent menu to avoid displaying two search fields to the user.
+		MenuBuilder.SetSearchable(false);
+
 		UMediaProfile* CurrentMediaProfile = GetCurrentProfile();
 		FAssetData CurrentAssetData = CurrentMediaProfile ? FAssetData(CurrentMediaProfile) : FAssetData();
 
diff --git a/Engine/Source/Editor/LevelEditor/Private/LevelEditorContextMenu.cpp b/Engine/Source/Editor/LevelEditor/Private/LevelEditorContextMenu.cpp
--- a/Engine/Source/Editor/LevelEditor/Private/LevelEditorContextMenu.cpp
+++ b/Engine/Source/Editor/LevelEditor/Private/LevelEditorContextMenu.cpp
@@ -1180,6 +1180,10 @@ private:
 
 void FLevelEditorContextMenuImpl::FillActorMenu(UToolMenu* Menu)
 {
+	// The contents of this menu are added as a custom widget with its own search field so we
+	// disable searching in this parent menu to avoid displaying two search fields to the user.
+	Menu->bSearchable = false;
+
 	struct Local
 	{
 		static FReply OnInteractiveActorPickerClicked()
diff --git a/Engine/Source/Editor/LevelEditor/Private/LevelEditorToolBar.cpp b/Engine/Source/Editor/LevelEditor/Private/LevelEditorToolBar.cpp
--- a/Engine/Source/Editor/LevelEditor/Private/LevelEditorToolBar.cpp
+++ b/Engine/Source/Editor/LevelEditor/Private/LevelEditorToolBar.cpp
@@ -388,6 +388,10 @@ namespace LevelEditorActionHelpers
 	 */
 	void GetSelectSettingsClassSubMenu(UToolMenu* InMenu, UClass* InRootClass, FOnClassPicked InOnClassPicked)
 	{
+		// The contents of this menu are added as a custom widget with its own search field so we
+		// disable searching in this parent menu to avoid displaying two search fields to the user.
+		InMenu->bSearchable = false;
+
 		FClassViewerInitializationOptions Options;
 		Options.Mode = EClassViewerMode::ClassPicker;
 		Options.DisplayMode = EClassViewerDisplayMode::ListView;
@@ -421,6 +425,10 @@ namespace LevelEditorActionHelpers
 	 */
 	void GetCreateSettingsClassSubMenu(UToolMenu* InMenu, UClass* InRootClass, FOnClassPicked InOnClassPicked)
 	{
+		// The contents of this menu are added as a custom widget with its own search field so we
+		// disable searching in this parent menu to avoid displaying two search fields to the user.
+		InMenu->bSearchable = false;
+
 		FClassViewerInitializationOptions Options;
 		Options.Mode = EClassViewerMode::ClassPicker;
 		Options.DisplayMode = EClassViewerDisplayMode::ListView;
@@ -605,8 +613,8 @@ namespace LevelEditorActionHelpers
 #define LOCTEXT_NAMESPACE "LevelToolBarViewMenu"
 		Section.AddDynamicEntry(InName, FNewToolMenuSectionDelegate::CreateLambda([=](FToolMenuSection& InSection)
 		{
-			ULevelEditorMenuContext* Context = InSection.FindContext<ULevelEditorMenuContext>();
-			if (Context && Context->LevelEditor.IsValid())
+			if (ULevelEditorMenuContext* const Context = InSection.FindContext<ULevelEditorMenuContext>();
+			Context && Context->LevelEditor.IsValid())
 			{
 				LevelEditorActionHelpers::FBlueprintMenuSettings GameModeMenuSettings;
 				GameModeMenuSettings.EditCommand =
@@ -1929,6 +1937,10 @@ void FLevelEditorToolBar::RegisterOpenBlueprintMenu()
 		/** Generates 'open blueprint' sub-menu */
 		static void MakeOpenBPClassMenu(UToolMenu* InMenu)
 		{
+			// The contents of this menu are added as a custom widget with its own search field so we
+			// disable searching in this parent menu to avoid displaying two search fields to the user.
+			InMenu->bSearchable = false;
+
 			FContentBrowserModule& ContentBrowserModule = FModuleManager::Get().LoadModuleChecked<FContentBrowserModule>(TEXT("ContentBrowser"));
 
 			// Configure filter for asset picker
diff --git a/Engine/Source/Editor/LevelEditor/Private/SLevelViewportToolBar.cpp b/Engine/Source/Editor/LevelEditor/Private/SLevelViewportToolBar.cpp
--- a/Engine/Source/Editor/LevelEditor/Private/SLevelViewportToolBar.cpp
+++ b/Engine/Source/Editor/LevelEditor/Private/SLevelViewportToolBar.cpp
@@ -1081,6 +1081,10 @@ void SLevelViewportToolBar::GenerateViewportConfigsMenu(UToolMenu* Menu) const
 	check (Viewport.IsValid());
 	TSharedPtr<FUICommandList> CommandList = Viewport.Pin()->GetCommandList();
 
+	// Disable searching in this menu because it only contains visual representations of
+	// viewport layouts without any searchable text.
+	Menu->bSearchable = false;
+
 	{
 		FToolMenuSection& Section = Menu->AddSection("LevelViewportOnePaneConfigs", LOCTEXT("OnePaneConfigHeader", "One Pane"));
 
