Adds SActionableMessageViewportWidget in Viewport's toolbar.
This widget will host warnings or messages with the ability to automate actions for the user.

Current actionable messages:
* Invalid Landscape grassmaps, with associated rebuild action
* Invalid Landscape physical materials, with associated rebuild action
* Invalid Landscape nanite meshes with associated rebuild action
* Invalid grassmaps, physical materials, nanite meshes will be agregated to only one message if multiple rebuild are necessary
* Landscape proxy packages needs to be updated, with associated mark dirty action

diff --git a/Engine/Source/Editor/EditorStyle/Private/StarshipStyle.cpp b/Engine/Source/Editor/EditorStyle/Private/StarshipStyle.cpp
--- a/Engine/Source/Editor/EditorStyle/Private/StarshipStyle.cpp
+++ b/Engine/Source/Editor/EditorStyle/Private/StarshipStyle.cpp
@@ -6557,6 +6557,26 @@ void FStarshipEditorStyle::FStyle::SetupLandscapeEditorStyle()
 			.SetForegroundColor(FSlateColor::UseForeground())
 			.SetArrowsImage(FSlateNoResource())
 		);
+
+		const FSlateColorBrush ActionableListViewBrush = FSlateColorBrush(FLinearColor(0.f, 0.f, 0.f, 0.f));
+		FLinearColor ToolbarBackgroundColor = FStyleColors::Dropdown.GetSpecifiedColor();
+		ToolbarBackgroundColor.A = .80f;
+		
+		Set("LandscapeEditor.ActionableMessage.Border", new FSlateRoundedBoxBrush(ToolbarBackgroundColor, 4.0f, FLinearColor(0,0,0,.8), 1.0));
+		Set("LandscapeEditor.ActionableMessage.Warning", new IMAGE_BRUSH_SVG(TEXT("Icons/Landscape/alert-solid"), Icon16x16, FStyleColors::Warning));
+		Set("LandscapeEditor.ActionableMessage.Update", new IMAGE_BRUSH_SVG(TEXT("Icons/Landscape/update"), Icon16x16));
+		Set("LandscapeEditor.ActionableMessage.ListView", FTableViewStyle().SetBackgroundBrush(ActionableListViewBrush));
+		Set("LandscapeEditor.ActionableMessage.ListViewRow",FTableRowStyle()
+			.SetEvenRowBackgroundBrush(ActionableListViewBrush)
+			.SetEvenRowBackgroundHoveredBrush(ActionableListViewBrush)
+			.SetOddRowBackgroundBrush(ActionableListViewBrush)
+			.SetOddRowBackgroundHoveredBrush(ActionableListViewBrush)
+			.SetSelectorFocusedBrush(ActionableListViewBrush)
+			.SetActiveBrush(ActionableListViewBrush)
+			.SetActiveHoveredBrush(ActionableListViewBrush)
+			.SetInactiveBrush(ActionableListViewBrush)
+			.SetInactiveHoveredBrush(ActionableListViewBrush)
+		);
 }
 
 #endif
diff --git a/Engine/Source/Editor/LandscapeEditor/Private/LandscapeEditorModule.cpp b/Engine/Source/Editor/LandscapeEditor/Private/LandscapeEditorModule.cpp
--- a/Engine/Source/Editor/LandscapeEditor/Private/LandscapeEditorModule.cpp
+++ b/Engine/Source/Editor/LandscapeEditor/Private/LandscapeEditorModule.cpp
@@ -133,13 +133,13 @@ public:
 		{
 			FToolMenuSection& Section = BuildMenu->FindOrAddSection("LevelEditorLandscape");
 
-			FUIAction ActionBuildGrassMaps(FExecuteAction::CreateStatic(&BuildGrassMaps), FCanExecuteAction());
+			FUIAction ActionBuildGrassMaps(FExecuteAction::CreateStatic(&UE::Landscape::BuildGrassMaps), FCanExecuteAction());
 			Section.AddMenuEntry(TEXT("BuildGrassMapsOnly"), LOCTEXT("BuildGrassMapsOnly", "Build Grass Maps Only"), LOCTEXT("BuildLandscapeGrassMaps", "Build landscape grass maps"), TAttribute<FSlateIcon>(), ActionBuildGrassMaps, EUserInterfaceActionType::Button);
 
-			FUIAction ActionBuildPhysicalMaterial(FExecuteAction::CreateStatic(&BuildPhysicalMaterial), FCanExecuteAction());
+			FUIAction ActionBuildPhysicalMaterial(FExecuteAction::CreateStatic(&UE::Landscape::BuildPhysicalMaterial), FCanExecuteAction());
 			Section.AddMenuEntry(TEXT("BuildPhysicalMaterialOnly"), LOCTEXT("BuildPhysicalMaterialOnly", "Build Physical Material Only"), LOCTEXT("BuildLandscapePhysicalMaterial", "Build landscape physical material"), TAttribute<FSlateIcon>(), ActionBuildPhysicalMaterial, EUserInterfaceActionType::Button);
 		
-			FUIAction ActionBuildNanite(FExecuteAction::CreateStatic(&BuildNanite), FCanExecuteAction());
+			FUIAction ActionBuildNanite(FExecuteAction::CreateStatic(&UE::Landscape::BuildNanite), FCanExecuteAction());
 			Section.AddMenuEntry(NAME_None, LOCTEXT("BuildNaniteOnly", "Build Nanite Only"), LOCTEXT("BuildLandscapeNanite", "Build Nanite representation"), TAttribute<FSlateIcon>(), ActionBuildNanite, EUserInterfaceActionType::Button);
 
 			FUIAction ActionSaveModifiedLandscapes(FExecuteAction::CreateStatic(&SaveModifiedLandscapes), FCanExecuteAction::CreateStatic(&HasModifiedLandscapes));
@@ -230,39 +230,6 @@ public:
 		);
 	}
 
-	static void BuildGrassMaps()
-	{
-		if (UWorld* World = GEditor->GetEditorWorldContext().World())
-		{
-			if (ULandscapeSubsystem* LandscapeSubsystem = World->GetSubsystem<ULandscapeSubsystem>())
-			{
-				LandscapeSubsystem->BuildGrassMaps();
-			}
-		}
-	}
-
-	static void BuildPhysicalMaterial()
-	{
-		if (UWorld* World = GEditor->GetEditorWorldContext().World())
-		{
-			if (ULandscapeSubsystem* LandscapeSubsystem = World->GetSubsystem<ULandscapeSubsystem>())
-			{
-				LandscapeSubsystem->BuildPhysicalMaterial();
-			}
-		}
-	}
-
-	static void BuildNanite()
-	{
-		if (UWorld* World = GEditor->GetEditorWorldContext().World())
-		{
-			if (ULandscapeSubsystem* LandscapeSubsystem = World->GetSubsystem<ULandscapeSubsystem>())
-			{
-				LandscapeSubsystem->BuildNanite();
-			}
-		}
-	}
-
 	static bool HasModifiedLandscapes()
 	{
 		if (UWorld* World = GEditor->GetEditorWorldContext().World())
diff --git a/Engine/Source/Editor/LandscapeEditor/Private/SActionableMessageViewportWidget.cpp b/Engine/Source/Editor/LandscapeEditor/Private/SActionableMessageViewportWidget.cpp
new file mode 100644
--- /dev/null
+++ b/Engine/Source/Editor/LandscapeEditor/Private/SActionableMessageViewportWidget.cpp
@@ -0,0 +1,220 @@
+// Copyright Epic Games, Inc. All Rights Reserved.
+
+#include "SActionableMessageViewportWidget.h"
+
+#include "ActionableMessageSubsystem.h"
+
+#include "Editor.h"
+#include "Engine/World.h"
+#include "Framework/Docking/TabManager.h"
+#include "Styling/AppStyle.h"
+#include "Styling/StyleColors.h"
+#include "Modules/ModuleManager.h"
+#include "Widgets/Input/SButton.h"
+#include "Widgets/Images/SImage.h"
+#include "Widgets/Text/STextBlock.h"
+#include "Widgets/SBoxPanel.h"
+#include "Widgets/Views/STableViewBase.h"
+#include "Widgets/Views/STableRow.h"
+#include "Widgets/Views/SListView.h"
+
+#define LOCTEXT_NAMESPACE "ActionableMessageViewportWidget"
+
+void SActionableMessageEntry::Construct(const FArguments& InArgs)
+{
+	ActionableMessage = InArgs._ActionableMessage;
+	OnClicked = InArgs._OnClicked;
+	
+	ChildSlot
+	[
+		SNew(SButton)
+		.ButtonStyle(&FCoreStyle::Get().GetWidgetStyle<FButtonStyle>("HoverHintOnly"))
+		.IsEnabled(true)
+		.OnClicked_Lambda([this]()
+		{
+			if (OnClicked.IsBound())
+			{
+				return OnClicked.Execute();
+			}
+
+			return FReply::Handled();
+		})
+		.ToolTipText_Lambda([this]() { return ActionableMessage.IsValid() ? ActionableMessage->Tooltip : FText::GetEmpty(); })
+		.ContentPadding(FMargin(0.f, 0.f, 0.f, -1.f))
+		[
+			SNew(SHorizontalBox)
+			+ SHorizontalBox::Slot()
+			.AutoWidth()
+			.VAlign(VAlign_Center)
+			[
+				SNew(SImage)
+				.Image(FAppStyle::Get().GetBrush("LandscapeEditor.ActionableMessage.Warning"))
+				.DesiredSizeOverride(FVector2D(16, 16))
+			]
+			+ SHorizontalBox::Slot()
+			.Padding(FMargin(4, 0, 8, 0))
+			.VAlign(VAlign_Center)
+			.HAlign(HAlign_Left)
+			.FillWidth(1)
+			[
+				SNew(STextBlock)
+				.Text_Lambda([this]() { return ActionableMessage.IsValid() ? ActionableMessage->Message : FText::GetEmpty(); })
+				.ColorAndOpacity(FStyleColors::Foreground)
+			]
+			+ SHorizontalBox::Slot()
+			.VAlign(VAlign_Center)
+			.HAlign(HAlign_Right)
+			.AutoWidth()
+			[
+				SNew(SButton)
+				.Visibility_Lambda([this]()
+				{
+					if (!ActionableMessage.IsValid() || ActionableMessage->ActionMessage.IsEmpty())
+					{
+						return EVisibility::Collapsed;
+					}
+
+					return EVisibility::Visible;
+				})
+				.ButtonStyle(&FCoreStyle::Get().GetWidgetStyle<FButtonStyle>("Button"))
+				.ContentPadding(FMargin(-4.f, 2.f))
+				.OnClicked_Lambda([this]()
+				{
+					if (ActionableMessage.IsValid() && (ActionableMessage->ActionCallback != nullptr))
+					{
+						ActionableMessage->ActionCallback();
+					}
+
+					return FReply::Handled();
+				})
+				[
+					SNew(SHorizontalBox)
+					+ SHorizontalBox::Slot()
+					.AutoWidth()
+					.Padding(0.f, 0.f, 4.f, 0.f)
+					.VAlign(VAlign_Center)
+					.HAlign(HAlign_Left)
+					[
+						SNew(SImage)
+						.Image(FAppStyle::Get().GetBrush("LandscapeEditor.ActionableMessage.Update"))
+						.DesiredSizeOverride(FVector2D(12, 12))
+					]
+					+ SHorizontalBox::Slot()
+					.AutoWidth()
+					.VAlign(VAlign_Center)
+					.HAlign(HAlign_Right)
+					[
+						SNew(STextBlock)
+						.Text_Lambda([this]() { return ActionableMessage.IsValid() ? ActionableMessage->ActionMessage : FText::GetEmpty(); })
+						.ColorAndOpacity(FStyleColors::White)
+					]
+				]
+			]
+		]
+	];
+}
+
+void SActionableMessageEntry::SetActionableMessage(TSharedPtr<FActionableMessage> InActionableMessage)
+{
+	ActionableMessage = InActionableMessage;
+}
+
+void SActionableMessageViewportWidget::Construct(const FArguments& InArgs)
+{
+	ChildSlot
+	[
+		SNew(SVerticalBox)
+		+ SVerticalBox::Slot()
+		.HAlign(HAlign_Right)
+		.VAlign(VAlign_Top)
+		[
+			SNew(SImage)
+			.Image_Lambda([this]() { return bExpanded ? FAppStyle::GetBrush("ContentBrowser.SortDown") : FAppStyle::GetBrush("ContentBrowser.SortUp"); })
+			.Visibility_Lambda([this]() { return ActionableMessages.IsEmpty() ? EVisibility::Collapsed : EVisibility::Visible; })
+			.ColorAndOpacity(FSlateColor::UseForeground())
+		]
+		+ SVerticalBox::Slot()
+		.Padding(8.f, 4.f)
+		.AutoHeight()
+		[
+			SAssignNew(MainActionableMessage, SActionableMessageEntry)
+			.OnClicked_Lambda([this]()
+			{
+				bExpanded = !bExpanded;
+				return FReply::Handled();
+			})				
+		]
+		+ SVerticalBox::Slot()
+		.AutoHeight()
+		[
+			SAssignNew(ActionableMessageList, SListView<TSharedPtr<FActionableMessage>>)
+			.ListViewStyle(&FAppStyle::Get().GetWidgetStyle<FTableViewStyle>("LandscapeEditor.ActionableMessage.ListView"))
+			.ItemHeight(48)
+			.ListItemsSource(&ActionableMessages)
+			.OnGenerateRow(this, &SActionableMessageViewportWidget::OnGenerateRow)
+			.Visibility_Lambda([this]()
+			{
+				return bExpanded ? EVisibility::Visible : EVisibility::Collapsed;
+			})
+		]
+	];
+}
+
+EVisibility SActionableMessageViewportWidget::GetVisibility()
+{
+	if (UWorld* World = GEditor->GetEditorWorldContext().World())
+	{
+		const UActionableMessageSubsystem* ActionableMessageSubsystem = World->GetSubsystem<UActionableMessageSubsystem>();
+		
+		if (ActionableMessageSubsystem == nullptr)
+		{
+			return EVisibility::Collapsed;
+		}
+
+		const uint32 SubsystemStateID = ActionableMessageSubsystem->GetStateID();
+
+		if (CachedStateID == SubsystemStateID)
+		{
+			return bMainMessageFilled ? EVisibility::Visible : EVisibility::Collapsed;
+		}
+		
+		const TMap<FName, TSharedPtr<FActionableMessage>>& ActionableMessageSource = ActionableMessageSubsystem->GetActionableMessages();
+
+		bMainMessageFilled = false;
+		ActionableMessages.Empty();
+		
+		for (const TPair<FName, TSharedPtr<FActionableMessage>>& ActionableMessagePair : ActionableMessageSource)
+		{
+			if ((!bMainMessageFilled) && (MainActionableMessage.IsValid()))
+			{
+				bMainMessageFilled = true;
+				MainActionableMessage->SetActionableMessage(ActionableMessagePair.Value);
+			}
+			else
+			{
+				ActionableMessages.Emplace(ActionableMessagePair.Value);
+			}
+		}
+
+		CachedStateID = SubsystemStateID;
+		
+		return bMainMessageFilled ? EVisibility::Visible : EVisibility::Collapsed;
+	}
+
+	return EVisibility::Collapsed;
+}
+
+TSharedRef<ITableRow> SActionableMessageViewportWidget::OnGenerateRow(TSharedPtr<FActionableMessage> InActionableMessage, const TSharedRef<STableViewBase>& OwnerTable)
+{
+	check(InActionableMessage.IsValid());
+
+	return SNew(STableRow<TSharedPtr<FActionableMessage>>, OwnerTable)
+	.Style(&FAppStyle::Get().GetWidgetStyle<FTableRowStyle>("LandscapeEditor.ActionableMessage.ListViewRow"))
+	.Padding(FMargin(8.f, 0.f, 8.f, 4.f))
+	[
+		SNew(SActionableMessageEntry)
+		.ActionableMessage(InActionableMessage)
+	];
+}
+
+#undef LOCTEXT_NAMESPACE
diff --git a/Engine/Source/Editor/LandscapeEditor/Public/SActionableMessageViewportWidget.h b/Engine/Source/Editor/LandscapeEditor/Public/SActionableMessageViewportWidget.h
new file mode 100644
--- /dev/null
+++ b/Engine/Source/Editor/LandscapeEditor/Public/SActionableMessageViewportWidget.h
@@ -0,0 +1,46 @@
+// Copyright Epic Games, Inc. All Rights Reserved.
+
+#pragma once
+
+#include "ActionableMessageSubsystem.h"
+
+#include "Widgets/SCompoundWidget.h"
+#include "Widgets/Views/STableViewBase.h"
+#include "Widgets/Views/STableRow.h"
+#include "Widgets/Views/SListView.h"
+
+class LANDSCAPEEDITOR_API SActionableMessageEntry : public SCompoundWidget
+{
+public:
+	SLATE_BEGIN_ARGS(SActionableMessageEntry) {}
+	SLATE_ARGUMENT(TSharedPtr<FActionableMessage>, ActionableMessage)
+	SLATE_EVENT(FOnClicked, OnClicked)
+	SLATE_END_ARGS()
+
+	void Construct(const FArguments& InArgs);
+	void SetActionableMessage(TSharedPtr<FActionableMessage> InActionableMessage);
+
+private:
+	TSharedPtr<FActionableMessage> ActionableMessage;
+	FOnClicked OnClicked;
+};
+
+class LANDSCAPEEDITOR_API SActionableMessageViewportWidget : public SCompoundWidget
+{
+public:
+	SLATE_BEGIN_ARGS(SActionableMessageViewportWidget) {}
+	SLATE_END_ARGS()
+
+	void Construct(const FArguments& InArgs);
+	EVisibility GetVisibility();
+
+	TSharedRef<ITableRow> OnGenerateRow(TSharedPtr<FActionableMessage> InActionableMessage, const TSharedRef<STableViewBase>& OwnerTable);
+
+private:
+	TSharedPtr<SActionableMessageEntry> MainActionableMessage;
+	TSharedPtr<SListView<TSharedPtr<FActionableMessage>>> ActionableMessageList;
+	TArray<TSharedPtr<FActionableMessage>> ActionableMessages;
+	uint32 CachedStateID = 0;
+	bool bMainMessageFilled = false;
+	bool bExpanded = false;
+};
diff --git a/Engine/Source/Editor/LevelEditor/LevelEditor.Build.cs b/Engine/Source/Editor/LevelEditor/LevelEditor.Build.cs
--- a/Engine/Source/Editor/LevelEditor/LevelEditor.Build.cs
+++ b/Engine/Source/Editor/LevelEditor/LevelEditor.Build.cs
@@ -89,7 +89,8 @@ public class LevelEditor : ModuleRules
 				"UnsavedAssetsTracker",
 				"UncontrolledChangelists",
 				"RenderCore",
-				"DeveloperSettings"
+				"DeveloperSettings",
+				"LandscapeEditor"
 			}
 		);
 
diff --git a/Engine/Source/Editor/LevelEditor/Private/SLevelViewportToolBar.cpp b/Engine/Source/Editor/LevelEditor/Private/SLevelViewportToolBar.cpp
--- a/Engine/Source/Editor/LevelEditor/Private/SLevelViewportToolBar.cpp
+++ b/Engine/Source/Editor/LevelEditor/Private/SLevelViewportToolBar.cpp
@@ -55,6 +55,7 @@
 #include "SortHelper.h"
 #include "Interfaces/IMainFrameModule.h"
 #include "SCommonEditorViewportToolbarBase.h"
+#include "SActionableMessageViewportWidget.h"
 
 #define LOCTEXT_NAMESPACE "LevelViewportToolBar"
 
@@ -201,178 +202,200 @@ void SLevelViewportToolBar::Construct( const FArguments& InArgs )
 		.BorderImage(FAppStyle::Get().GetBrush("EditorViewportToolBar.Background"))
 		.Cursor(EMouseCursor::Default)
 		[
-			SNew( SHorizontalBox )
-			+ SHorizontalBox::Slot()
-			.AutoWidth()
+			SNew( SVerticalBox )
+			+ SVerticalBox::Slot()
+			.AutoHeight()
 			.Padding(ToolbarSlotPadding)
-			[
-				SNew( SEditorViewportToolbarMenu )
-				.ParentToolBar( SharedThis( this ) )
-				.Visibility(Viewport.Pin().Get(), &SLevelViewport::GetToolbarVisibility)
-				.Image("EditorViewportToolBar.OptionsDropdown")
-				.AddMetaData<FTagMetaData>(FTagMetaData(TEXT("EditorViewportToolBar.MenuDropdown")))
-				.OnGetMenuContent( this, &SLevelViewportToolBar::GenerateOptionsMenu )
-			]
-			+ SHorizontalBox::Slot()
 			[
 				SNew( SHorizontalBox )
-				.Visibility(Viewport.Pin().Get(), &SLevelViewport::GetFullToolbarVisibility)
 				+ SHorizontalBox::Slot()
 				.AutoWidth()
-				.Padding( ToolbarSlotPadding )
-				[
-					SNew( SEditorViewportToolbarMenu )
-					.ParentToolBar( SharedThis( this ) )
-					.Label( this, &SLevelViewportToolBar::GetCameraMenuLabel )
-					.LabelIcon( this, &SLevelViewportToolBar::GetCameraMenuLabelIcon )
-					.AddMetaData<FTagMetaData>(FTagMetaData(TEXT("EditorViewportToolBar.CameraMenu")))
-					.OnGetMenuContent( this, &SLevelViewportToolBar::GenerateCameraMenu ) 
-				]
-				+ SHorizontalBox::Slot()
-				.AutoWidth()
-				.Padding( ToolbarSlotPadding )
-				[
-					SNew( SLevelEditorViewportViewMenu, ViewportRef, SharedThis(this) )
-					.MenuExtenders(GetViewMenuExtender())
-					.AddMetaData<FTagMetaData>(FTagMetaData(TEXT("ViewMenuButton")))
-				]
-				+ SHorizontalBox::Slot()
-				.AutoWidth()
-				.Padding( ToolbarSlotPadding )
-				[
-					SNew( SEditorViewportToolbarMenu )
-					.Label( LOCTEXT("ShowMenuTitle", "Show") )
-					.ParentToolBar( SharedThis( this ) )
-					.AddMetaData<FTagMetaData>(FTagMetaData(TEXT("EditorViewportToolBar.ShowMenu")))
-					.OnGetMenuContent( this, &SLevelViewportToolBar::GenerateShowMenu ) 
-				]
-				+ SHorizontalBox::Slot()
-				.AutoWidth()
-				.Padding( ToolbarSlotPadding )
-				[
-					SNew( SEditorViewportToolbarMenu )
-					.Label( this, &SLevelViewportToolBar::GetViewModeOptionsMenuLabel )
-					.ParentToolBar( SharedThis( this ) )
-					.Visibility( this, &SLevelViewportToolBar::GetViewModeOptionsVisibility )
-					.AddMetaData<FTagMetaData>(FTagMetaData(TEXT("EditorViewportToolBar.ViewModeOptions")))
-					.OnGetMenuContent( this, &SLevelViewportToolBar::GenerateViewModeOptionsMenu ) 
-				]
-				+ SHorizontalBox::Slot()
-				.AutoWidth()
-				.Padding( ToolbarSlotPadding )
+				.Padding(ToolbarSlotPadding)
 				[
 					SNew( SEditorViewportToolbarMenu )
 					.ParentToolBar( SharedThis( this ) )
-					.Label( this, &SLevelViewportToolBar::GetDevicePreviewMenuLabel )
-					.LabelIcon( this, &SLevelViewportToolBar::GetDevicePreviewMenuLabelIcon )
-					.OnGetMenuContent( this, &SLevelViewportToolBar::GenerateDevicePreviewMenu )
-					//@todo rendering: mobile preview in view port is not functional yet - remove this once it is.
-					.Visibility(EVisibility::Collapsed)
-				]
-				+ SHorizontalBox::Slot()
-				.Padding(ToolbarSlotPadding)
-				.AutoWidth()
-				.HAlign(HAlign_Left)
-				.VAlign(VAlign_Fill)
-				[
-					SNew(SExtensionPanel)
-					.ExtensionPanelID("LevelViewportToolBar.LeftExtension")
-					.ExtensionContext(ExtensionContextObject)
+					.Visibility(Viewport.Pin().Get(), &SLevelViewport::GetToolbarVisibility)
+					.Image("EditorViewportToolBar.OptionsDropdown")
+					.AddMetaData<FTagMetaData>(FTagMetaData(TEXT("EditorViewportToolBar.MenuDropdown")))
+					.OnGetMenuContent( this, &SLevelViewportToolBar::GenerateOptionsMenu )
 				]
 				+ SHorizontalBox::Slot()
-				.AutoWidth()
-				.Padding(ToolbarSlotPadding)
 				[
-					// Button to show that realtime is off
-					SNew(SEditorViewportToolBarButton)	
-					.ButtonType(EUserInterfaceActionType::Button)
-					.ButtonStyle(&FAppStyle::Get().GetWidgetStyle<FButtonStyle>("EditorViewportToolBar.WarningButton"))
-					.OnClicked(this, &SLevelViewportToolBar::OnRealtimeWarningClicked)
-					.Visibility(this, &SLevelViewportToolBar::GetRealtimeWarningVisibility)
-					.ToolTipText(LOCTEXT("RealtimeOff_ToolTip", "This viewport is not updating in realtime.  Click to turn on realtime mode."))
-					.Content()
+					SNew( SHorizontalBox )
+					.Visibility(Viewport.Pin().Get(), &SLevelViewport::GetFullToolbarVisibility)
+					+ SHorizontalBox::Slot()
+					.AutoWidth()
+					.Padding( ToolbarSlotPadding )
 					[
-						SNew(STextBlock)
-						.TextStyle(&FAppStyle::Get().GetWidgetStyle<FTextBlockStyle>("SmallText"))
-						.Text(LOCTEXT("RealtimeOff", "Realtime Off"))
+						SNew( SEditorViewportToolbarMenu )
+						.ParentToolBar( SharedThis( this ) )
+						.Label( this, &SLevelViewportToolBar::GetCameraMenuLabel )
+						.LabelIcon( this, &SLevelViewportToolBar::GetCameraMenuLabelIcon )
+						.AddMetaData<FTagMetaData>(FTagMetaData(TEXT("EditorViewportToolBar.CameraMenu")))
+						.OnGetMenuContent( this, &SLevelViewportToolBar::GenerateCameraMenu ) 
+					]
+					+ SHorizontalBox::Slot()
+					.AutoWidth()
+					.Padding( ToolbarSlotPadding )
+					[
+						SNew( SLevelEditorViewportViewMenu, ViewportRef, SharedThis(this) )
+						.MenuExtenders(GetViewMenuExtender())
+						.AddMetaData<FTagMetaData>(FTagMetaData(TEXT("ViewMenuButton")))
+					]
+					+ SHorizontalBox::Slot()
+					.AutoWidth()
+					.Padding( ToolbarSlotPadding )
+					[
+						SNew( SEditorViewportToolbarMenu )
+						.Label( LOCTEXT("ShowMenuTitle", "Show") )
+						.ParentToolBar( SharedThis( this ) )
+						.AddMetaData<FTagMetaData>(FTagMetaData(TEXT("EditorViewportToolBar.ShowMenu")))
+						.OnGetMenuContent( this, &SLevelViewportToolBar::GenerateShowMenu ) 
+					]
+					+ SHorizontalBox::Slot()
+					.AutoWidth()
+					.Padding( ToolbarSlotPadding )
+					[
+						SNew( SEditorViewportToolbarMenu )
+						.Label( this, &SLevelViewportToolBar::GetViewModeOptionsMenuLabel )
+						.ParentToolBar( SharedThis( this ) )
+						.Visibility( this, &SLevelViewportToolBar::GetViewModeOptionsVisibility )
+						.AddMetaData<FTagMetaData>(FTagMetaData(TEXT("EditorViewportToolBar.ViewModeOptions")))
+						.OnGetMenuContent( this, &SLevelViewportToolBar::GenerateViewModeOptionsMenu ) 
+					]
+					+ SHorizontalBox::Slot()
+					.AutoWidth()
+					.Padding( ToolbarSlotPadding )
+					[
+						SNew( SEditorViewportToolbarMenu )
+						.ParentToolBar( SharedThis( this ) )
+						.Label( this, &SLevelViewportToolBar::GetDevicePreviewMenuLabel )
+						.LabelIcon( this, &SLevelViewportToolBar::GetDevicePreviewMenuLabelIcon )
+						.OnGetMenuContent( this, &SLevelViewportToolBar::GenerateDevicePreviewMenu )
+						//@todo rendering: mobile preview in view port is not functional yet - remove this once it is.
+						.Visibility(EVisibility::Collapsed)
+					]
+					+ SHorizontalBox::Slot()
+					.Padding(ToolbarSlotPadding)
+					.AutoWidth()
+					.HAlign(HAlign_Left)
+					.VAlign(VAlign_Fill)
+					[
+						SNew(SExtensionPanel)
+						.ExtensionPanelID("LevelViewportToolBar.LeftExtension")
+						.ExtensionContext(ExtensionContextObject)
+					]
+					+ SHorizontalBox::Slot()
+					.AutoWidth()
+					.Padding(ToolbarSlotPadding)
+					[
+						// Button to show that realtime is off
+						SNew(SEditorViewportToolBarButton)	
+						.ButtonType(EUserInterfaceActionType::Button)
+						.ButtonStyle(&FAppStyle::Get().GetWidgetStyle<FButtonStyle>("EditorViewportToolBar.WarningButton"))
+						.OnClicked(this, &SLevelViewportToolBar::OnRealtimeWarningClicked)
+						.Visibility(this, &SLevelViewportToolBar::GetRealtimeWarningVisibility)
+						.ToolTipText(LOCTEXT("RealtimeOff_ToolTip", "This viewport is not updating in realtime.  Click to turn on realtime mode."))
+						.Content()
+						[
+							SNew(STextBlock)
+							.TextStyle(&FAppStyle::Get().GetWidgetStyle<FTextBlockStyle>("SmallText"))
+							.Text(LOCTEXT("RealtimeOff", "Realtime Off"))
+						]
+					]
+					+ SHorizontalBox::Slot()
+					.AutoWidth()
+					.Padding(ToolbarSlotPadding)
+					[
+						// Button to show scalability warnings
+						SNew(SEditorViewportToolbarMenu)
+						.ParentToolBar(SharedThis(this))
+						.Label(this, &SLevelViewportToolBar::GetScalabilityWarningLabel)
+						.MenuStyle(&FAppStyle::Get().GetWidgetStyle<FButtonStyle>("EditorViewportToolBar.WarningButton"))
+						.OnGetMenuContent(this, &SLevelViewportToolBar::GetScalabilityWarningMenuContent)
+						.Visibility(this, &SLevelViewportToolBar::GetScalabilityWarningVisibility)
+						.ToolTipText(LOCTEXT("ScalabilityWarning_ToolTip", "Non-default scalability settings could be affecting what is shown in this viewport.\nFor example you may experience lower visual quality, reduced particle counts, and other artifacts that don't match what the scene would look like when running outside of the editor. Click to make changes."))
+					]
+					+ SHorizontalBox::Slot()
+					.Padding(ToolbarSlotPadding)
+					.HAlign(HAlign_Center)
+					.VAlign(VAlign_Fill)
+					[
+						SNew(SExtensionPanel)
+						.ExtensionPanelID("LevelViewportToolBar.MiddleExtension")
+						.ExtensionContext(ExtensionContextObject)
+					]
+					+ SHorizontalBox::Slot()
+					.Padding(ToolbarSlotPadding)
+					.AutoWidth()
+					.HAlign(HAlign_Right)
+					.VAlign(VAlign_Fill)
+					[
+						SNew(SExtensionPanel)
+						.ExtensionPanelID("LevelViewportToolBar.RightExtension")
+						.ExtensionContext(ExtensionContextObject)
+					]
+					+ SHorizontalBox::Slot()
+					.AutoWidth()
+					.MaxWidth(TAttribute<float>::CreateSP(this, &SLevelViewportToolBar::GetTransformToolbarWidth))
+					.Padding(ToolbarSlotPadding)
+					.HAlign(HAlign_Right)
+					[
+						SAssignNew(TransformToolbar, STransformViewportToolBar)
+						.Viewport(ViewportRef)
+						.CommandList(ViewportRef->GetCommandList())
+						.Extenders(LevelEditorModule.GetToolBarExtensibilityManager()->GetAllExtenders())
+						.Visibility(ViewportRef, &SLevelViewport::GetTransformToolbarVisibility)
+					]
+					+ SHorizontalBox::Slot()
+					.HAlign(HAlign_Right)
+					.AutoWidth()
+					.Padding(ToolbarSlotPadding)
+					[
+						//The Maximize/Minimize button is only displayed when not in Immersive mode.
+						SNew(SEditorViewportToolBarButton)
+						.ButtonType(EUserInterfaceActionType::ToggleButton)
+						.CheckBoxStyle(&FAppStyle::Get().GetWidgetStyle<FCheckBoxStyle>("EditorViewportToolBar.MaximizeRestoreButton"))
+						.IsChecked(ViewportRef, &SLevelViewport::IsMaximized)
+						.OnClicked(ViewportRef, &SLevelViewport::OnToggleMaximize)
+						.Visibility(ViewportRef, &SLevelViewport::GetMaximizeToggleVisibility)
+						.Image("EditorViewportToolBar.Maximize")
+						.ToolTipText(LOCTEXT("Maximize_ToolTip", "Maximizes or restores this viewport"))
+					]
+					+ SHorizontalBox::Slot()
+					.HAlign(HAlign_Right)
+					.AutoWidth()
+					.Padding(ToolbarSlotPadding)
+					[
+						//The Restore from Immersive' button is only displayed when the editor is in Immersive mode.
+						SNew(SEditorViewportToolBarButton)
+						.ButtonType(EUserInterfaceActionType::Button)
+						.OnClicked(ViewportRef, &SLevelViewport::OnToggleMaximize)
+						.Visibility(ViewportRef, &SLevelViewport::GetCloseImmersiveButtonVisibility)
+						.Image("EditorViewportToolBar.RestoreFromImmersive.Normal")
+						.ToolTipText(LOCTEXT("RestoreFromImmersive_ToolTip", "Restore from Immersive"))
 					]
 				]
-				+ SHorizontalBox::Slot()
-				.AutoWidth()
-				.Padding(ToolbarSlotPadding)
-				[
-					// Button to show scalability warnings
-					SNew(SEditorViewportToolbarMenu)
-					.ParentToolBar(SharedThis(this))
-					.Label(this, &SLevelViewportToolBar::GetScalabilityWarningLabel)
-					.MenuStyle(&FAppStyle::Get().GetWidgetStyle<FButtonStyle>("EditorViewportToolBar.WarningButton"))
-					.OnGetMenuContent(this, &SLevelViewportToolBar::GetScalabilityWarningMenuContent)
-					.Visibility(this, &SLevelViewportToolBar::GetScalabilityWarningVisibility)
-					.ToolTipText(LOCTEXT("ScalabilityWarning_ToolTip", "Non-default scalability settings could be affecting what is shown in this viewport.\nFor example you may experience lower visual quality, reduced particle counts, and other artifacts that don't match what the scene would look like when running outside of the editor. Click to make changes."))
-				]
-				+ SHorizontalBox::Slot()
-				.Padding(ToolbarSlotPadding)
-				.HAlign(HAlign_Center)
-				.VAlign(VAlign_Fill)
-				[
-					SNew(SExtensionPanel)
-					.ExtensionPanelID("LevelViewportToolBar.MiddleExtension")
-					.ExtensionContext(ExtensionContextObject)
-				]
-				+ SHorizontalBox::Slot()
-				.Padding(ToolbarSlotPadding)
-				.AutoWidth()
-				.HAlign(HAlign_Right)
-				.VAlign(VAlign_Fill)
-				[
-					SNew(SExtensionPanel)
-					.ExtensionPanelID("LevelViewportToolBar.RightExtension")
-					.ExtensionContext(ExtensionContextObject)
-				]
-				+ SHorizontalBox::Slot()
-				.AutoWidth()
-				.MaxWidth(TAttribute<float>::CreateSP(this, &SLevelViewportToolBar::GetTransformToolbarWidth))
-				.Padding(ToolbarSlotPadding)
-				.HAlign(HAlign_Right)
-				[
-					SAssignNew(TransformToolbar, STransformViewportToolBar)
-					.Viewport(ViewportRef)
-					.CommandList(ViewportRef->GetCommandList())
-					.Extenders(LevelEditorModule.GetToolBarExtensibilityManager()->GetAllExtenders())
-					.Visibility(ViewportRef, &SLevelViewport::GetTransformToolbarVisibility)
-				]
-				+ SHorizontalBox::Slot()
-				.HAlign(HAlign_Right)
-				.AutoWidth()
-				.Padding(ToolbarSlotPadding)
-				[
-					//The Maximize/Minimize button is only displayed when not in Immersive mode.
-					SNew(SEditorViewportToolBarButton)
-					.ButtonType(EUserInterfaceActionType::ToggleButton)
-					.CheckBoxStyle(&FAppStyle::Get().GetWidgetStyle<FCheckBoxStyle>("EditorViewportToolBar.MaximizeRestoreButton"))
-					.IsChecked(ViewportRef, &SLevelViewport::IsMaximized)
-					.OnClicked(ViewportRef, &SLevelViewport::OnToggleMaximize)
-					.Visibility(ViewportRef, &SLevelViewport::GetMaximizeToggleVisibility)
-					.Image("EditorViewportToolBar.Maximize")
-					.ToolTipText(LOCTEXT("Maximize_ToolTip", "Maximizes or restores this viewport"))
-				]
-				+ SHorizontalBox::Slot()
-				.HAlign(HAlign_Right)
-				.AutoWidth()
-				.Padding(ToolbarSlotPadding)
+			]
+			+ SVerticalBox::Slot()
+			.AutoHeight()
+			.Padding(0.f, 8.f, 10.f, 0.f)
+			.VAlign(VAlign_Center)
+			.HAlign(HAlign_Right)
+			[
+				SNew(SBorder)
+				.BorderImage(FAppStyle::Get().GetBrush("LandscapeEditor.ActionableMessage.Border"))
 				[
-					//The Restore from Immersive' button is only displayed when the editor is in Immersive mode.
-					SNew(SEditorViewportToolBarButton)
-					.ButtonType(EUserInterfaceActionType::Button)
-					.OnClicked(ViewportRef, &SLevelViewport::OnToggleMaximize)
-					.Visibility(ViewportRef, &SLevelViewport::GetCloseImmersiveButtonVisibility)
-					.Image("EditorViewportToolBar.RestoreFromImmersive.Normal")
-					.ToolTipText(LOCTEXT("RestoreFromImmersive_ToolTip", "Restore from Immersive"))
+					SAssignNew(ActionableMessageViewportWidget, SActionableMessageViewportWidget)
+					.Visibility_Lambda([this]()
+					{
+						return ActionableMessageViewportWidget->GetVisibility();
+					})
 				]
 			]
 		]
 	];
-
+	
 	SViewportToolBar::Construct(SViewportToolBar::FArguments());
 }
 
diff --git a/Engine/Source/Editor/LevelEditor/Private/SLevelViewportToolBar.h b/Engine/Source/Editor/LevelEditor/Private/SLevelViewportToolBar.h
--- a/Engine/Source/Editor/LevelEditor/Private/SLevelViewportToolBar.h
+++ b/Engine/Source/Editor/LevelEditor/Private/SLevelViewportToolBar.h
@@ -12,6 +12,7 @@
 class ACameraActor;
 class FExtender;
 class SExtensionPanel;
+class SActionableMessageViewportWidget;
 class STransformViewportToolBar;
 class FMenuBuilder;
 class UToolMenu;
@@ -287,6 +288,9 @@ private:
 	/** STransformViewportToolBar menu */
 	TSharedPtr<STransformViewportToolBar> TransformToolbar;
 
+	/** Viewport widget for warning messages */
+	TSharedPtr<SActionableMessageViewportWidget> ActionableMessageViewportWidget;
+
 	/** The previous max STransformViewportToolBar width to allow deterministic size calculations */
 	mutable float TransformToolbar_CachedMaxWidth = 0.0f;
 };
diff --git a/Engine/Source/Runtime/Engine/Private/UnrealEngine.cpp b/Engine/Source/Runtime/Engine/Private/UnrealEngine.cpp
--- a/Engine/Source/Runtime/Engine/Private/UnrealEngine.cpp
+++ b/Engine/Source/Runtime/Engine/Private/UnrealEngine.cpp
@@ -12388,13 +12388,6 @@ float DrawMapWarnings(UWorld* World, FViewport* Viewport, FCanvas* Canvas, UCanv
 	}
 #endif
 
-#endif
-
-#if WITH_EDITOR && !(UE_BUILD_SHIPPING || UE_BUILD_TEST)
-	if (ULandscapeSubsystem* LandscapeSubsystem = World->GetSubsystem<ULandscapeSubsystem>())
-	{
-		LandscapeSubsystem->DisplayMessages(Canvas, MessageX, MessageY);
-	}
 #endif
 
 	if (World->NumTextureStreamingUnbuiltComponents > 0 || World->NumTextureStreamingDirtyResources > 0)
diff --git a/Engine/Source/Runtime/Landscape/Private/ActionableMessageSubsystem.cpp b/Engine/Source/Runtime/Landscape/Private/ActionableMessageSubsystem.cpp
new file mode 100644
--- /dev/null
+++ b/Engine/Source/Runtime/Landscape/Private/ActionableMessageSubsystem.cpp
@@ -0,0 +1,38 @@
+// Copyright Epic Games, Inc. All Rights Reserved.
+
+#include "ActionableMessageSubsystem.h"
+
+#if WITH_EDITOR
+
+#include UE_INLINE_GENERATED_CPP_BY_NAME(ActionableMessageSubsystem)
+
+bool operator!=(const FActionableMessage& LHS, const FActionableMessage& RHS)
+{
+	return !LHS.Message.EqualTo(RHS.Message) || !LHS.ActionMessage.EqualTo(RHS.ActionMessage) || !LHS.Tooltip.EqualTo(RHS.Tooltip)
+		|| (!FMemory::Memcmp(&LHS.ActionCallback, &RHS.ActionCallback, sizeof (LHS.ActionCallback)));
+}
+
+void UActionableMessageSubsystem::SetActionableMessage(FName InProvider, const FActionableMessage& InActionableMessage)
+{
+	const TSharedPtr<FActionableMessage>* ValuePtr = ProviderActionableMessageMap.Find(InProvider);
+
+	if (ValuePtr == nullptr)
+	{
+		ProviderActionableMessageMap.Emplace(InProvider, MakeShared<FActionableMessage>(InActionableMessage));
+		++StateID;
+	}
+	else if	(**ValuePtr != InActionableMessage)
+	{
+		**ValuePtr = InActionableMessage;
+		++StateID;
+	}
+	
+}
+
+void UActionableMessageSubsystem::ClearActionableMessage(FName InProvider)
+{
+	ProviderActionableMessageMap.Remove(InProvider);
+	++StateID;
+}
+
+#endif // WITH_EDITOR
\ No newline at end of file
diff --git a/Engine/Source/Runtime/Landscape/Private/LandscapeSubsystem.cpp b/Engine/Source/Runtime/Landscape/Private/LandscapeSubsystem.cpp
--- a/Engine/Source/Runtime/Landscape/Private/LandscapeSubsystem.cpp
+++ b/Engine/Source/Runtime/Landscape/Private/LandscapeSubsystem.cpp
@@ -39,6 +39,7 @@
 #include UE_INLINE_GENERATED_CPP_BY_NAME(LandscapeSubsystem)
 
 #if WITH_EDITOR
+#include "ActionableMessageSubsystem.h"
 #include "FileHelpers.h"
 #include "Editor.h"
 #endif
@@ -67,6 +68,66 @@ DECLARE_CYCLE_STAT(TEXT("LandscapeSubsystem Tick"), STAT_LandscapeSubsystemTick,
 
 #define LOCTEXT_NAMESPACE "LandscapeSubsystem"
 
+namespace UE::Landscape
+{
+#if WITH_EDITOR
+	void MarkModifiedLandscapesAsDirty()
+	{
+		if (UWorld* World = GEditor->GetEditorWorldContext().World())
+		{
+			if (ULandscapeSubsystem* LandscapeSubsystem = World->GetSubsystem<ULandscapeSubsystem>())
+			{
+				LandscapeSubsystem->MarkModifiedLandscapesAsDirty();
+			}
+		}
+	}
+	
+	void BuildGrassMaps()
+	{
+		if (UWorld* World = GEditor->GetEditorWorldContext().World())
+		{
+			if (ULandscapeSubsystem* LandscapeSubsystem = World->GetSubsystem<ULandscapeSubsystem>())
+			{
+				LandscapeSubsystem->BuildGrassMaps();
+			}
+		}
+	}
+
+	void BuildPhysicalMaterial()
+	{
+		if (UWorld* World = GEditor->GetEditorWorldContext().World())
+		{
+			if (ULandscapeSubsystem* LandscapeSubsystem = World->GetSubsystem<ULandscapeSubsystem>())
+			{
+				LandscapeSubsystem->BuildPhysicalMaterial();
+			}
+		}
+	}
+
+	void BuildNanite()
+	{
+		if (UWorld* World = GEditor->GetEditorWorldContext().World())
+		{
+			if (ULandscapeSubsystem* LandscapeSubsystem = World->GetSubsystem<ULandscapeSubsystem>())
+			{
+				LandscapeSubsystem->BuildNanite();
+			}
+		}
+	}
+
+	void BuildAll()
+	{
+		if (UWorld* World = GEditor->GetEditorWorldContext().World())
+		{
+			if (ULandscapeSubsystem* LandscapeSubsystem = World->GetSubsystem<ULandscapeSubsystem>())
+			{
+				LandscapeSubsystem->BuildAll();
+			}
+		}
+	}
+#endif //WITH_EDITOR
+}
+
 ULandscapeSubsystem::ULandscapeSubsystem()
 {
 }
@@ -502,6 +563,23 @@ void ULandscapeSubsystem::Tick(float DeltaTime)
 	NaniteMeshBuildEvents.RemoveAllSwap([](const FGraphEventRef& Ref) -> bool { return Ref->IsComplete(); });
 
 	LastTickFrameNumber = FrameNumber;
+	
+	if (UActionableMessageSubsystem* ActionableMessageSubsystem = World->GetSubsystem<UActionableMessageSubsystem>())
+	{
+		FActionableMessage ActionableMessage;
+		const FName LandscapeMessageProvider = TEXT("Landscape");
+
+		if (GetActionableMessage(ActionableMessage))
+		{
+			ActionableMessageSubsystem->SetActionableMessage(LandscapeMessageProvider, ActionableMessage);
+		}
+		else
+		{
+			ActionableMessageSubsystem->ClearActionableMessage(LandscapeMessageProvider);
+		}
+	}
+
+
 #endif // WITH_EDITOR
 }
 
@@ -783,52 +861,78 @@ ALandscapeProxy* ULandscapeSubsystem::FindOrAddLandscapeProxy(ULandscapeInfo* La
 
 void ULandscapeSubsystem::DisplayMessages(FCanvas* Canvas, float& XPos, float& YPos)
 {
-	TRACE_CPUPROFILER_EVENT_SCOPE(ULandscapeSubsystem::DisplayMessages);
+	
+}
 
-	const int32 FontSizeY = 20;
-	FCanvasTextItem SmallTextItem(FVector2D(0, 0), FText::GetEmpty(), GEngine->GetSmallFont(), FLinearColor::White);
-	SmallTextItem.EnableShadow(FLinearColor::Black);
+bool ULandscapeSubsystem::GetActionableMessage(FActionableMessage& OutActionableMessage)
+{
+	TRACE_CPUPROFILER_EVENT_SCOPE(ULandscapeSubsystem::GetActionableMessage);
 
-	auto DisplayMessageForOutdatedDataFlag = [&SmallTextItem, Canvas, XPos, &YPos, FontSizeY, this] (UE::Landscape::EOutdatedDataFlags InOutdatedDataFlag, const FTextFormat& InTextFormat)
+	const FText DefaultMessage = LOCTEXT("LandscapeModifiedText", "Landscape is not up to date");
+	const FText DefaultTooltip = LOCTEXT("LandscapeModifiedToolTip", "External asset modification lead to the landscape not being up to date. Please rebuild the landscape.");
+	const FText DefaultActionMessage = LOCTEXT("LandscapeRebuildText", "Rebuild");
+
+	const TArray<ALandscapeProxy*> OutdatedGrassMapProxies = GetOutdatedProxies(UE::Landscape::EOutdatedDataFlags::GrassMaps, /*bInMustMatchAllFlags = */false);
+	const TArray<ALandscapeProxy*> OutdatedPhysicalMaterialProxies = GetOutdatedProxies(UE::Landscape::EOutdatedDataFlags::PhysicalMaterials, /*bInMustMatchAllFlags = */false);
+	const TArray<ALandscapeProxy*> OutdatedNaniteProxies = GetOutdatedProxies(UE::Landscape::EOutdatedDataFlags::NaniteMeshes, /*bInMustMatchAllFlags = */false);
+	
+	const int32 OutdatedGrassmapProxiesCount = OutdatedGrassMapProxies.Num();
+	const int32 OutdatedPhysicalMaterialProxiesCount = OutdatedPhysicalMaterialProxies.Num();
+	const int32 OutdatedNaniteProxiesCount = OutdatedNaniteProxies.Num();
+
+	const int32 OutdatedFlags = ((OutdatedGrassmapProxiesCount > 0) ? 1 : 0) + ((OutdatedPhysicalMaterialProxiesCount > 0) ? 1 : 0) + ((OutdatedNaniteProxiesCount > 0) ? 1 : 0);
+
+	if (HasModifiedLandscapes())
 	{
-		TArray<ALandscapeProxy*> OutdatedProxies = GetOutdatedProxies(InOutdatedDataFlag, /*bInMustMatchAllFlags = */false);
-		if (int32 OutdatedProxiesCount = OutdatedProxies.Num())
-		{
-			SmallTextItem.SetColor(FLinearColor::Red);
-			SmallTextItem.Text = FText::Format(InTextFormat, OutdatedProxiesCount);
-			Canvas->DrawItem(SmallTextItem, FVector2D(XPos, YPos));
-			YPos += FontSizeY;
-		}
-	};
+		OutActionableMessage.Message = LOCTEXT("LandscapeModifiedText.Dirty", "Landscape contains unsaved changes");
+		OutActionableMessage.Tooltip = DefaultTooltip;
+		OutActionableMessage.ActionMessage = LOCTEXT("LandscapeModifiedText.MarkDirty", "Mark Dirty");
+		OutActionableMessage.ActionCallback = UE::Landscape::MarkModifiedLandscapesAsDirty;
+
+		return true;
+	}
+	
+	if (OutdatedFlags > 1)
+	{
+		OutActionableMessage.Message = DefaultMessage;
+		OutActionableMessage.Tooltip = DefaultTooltip;
+		OutActionableMessage.ActionMessage = DefaultActionMessage;
+		OutActionableMessage.ActionCallback = UE::Landscape::BuildAll;
 
-	// Outdated grass maps message :
-	DisplayMessageForOutdatedDataFlag(UE::Landscape::EOutdatedDataFlags::GrassMaps, LOCTEXT("GRASS_MAPS_NEED_TO_BE_REBUILT_FMT", "LANDSCAPE: {0} {0}|plural(one=ACTOR,other=ACTORS) WITH GRASS MAPS {0}|plural(one=NEEDS,other=NEED) TO BE REBUILT"));
+		return true;
+	}
+	
+	if (OutdatedGrassmapProxiesCount > 0)
+	{
+		OutActionableMessage.Message = FText::Format(LOCTEXT("GRASS_MAPS_NEED_TO_BE_REBUILT_FMT", "{0} Landscape {0}|plural(one=actor,other=actors) with grass maps {0}|plural(one=needs,other=need) to be rebuilt"), OutdatedGrassmapProxiesCount);
+		OutActionableMessage.Tooltip = DefaultTooltip;
+		OutActionableMessage.ActionMessage = DefaultActionMessage;
+		OutActionableMessage.ActionCallback = UE::Landscape::BuildGrassMaps;
 
-	// Outdated physical materials message :
-	DisplayMessageForOutdatedDataFlag(UE::Landscape::EOutdatedDataFlags::PhysicalMaterials, LOCTEXT("LANDSCAPE_PHYSICALMATERIAL_NEED_TO_BE_REBUILT_FMT", "LANDSCAPE: {0} {0}|plural(one=ACTOR,other=ACTORS) WITH PHYSICAL MATERIALS {0}|plural(one=NEEDS,other=NEED) TO BE REBUILT"));
+		return true;
+	}
+	
+	if (OutdatedPhysicalMaterialProxiesCount > 0)
+	{
+		OutActionableMessage.Message = FText::Format(LOCTEXT("LANDSCAPE_PHYSICALMATERIAL_NEED_TO_BE_REBUILT_FMT", "{0} Landscape {0}|plural(one=actor,other=actors) with physical materials {0}|plural(one=needs,other=need) to be rebuilt"), OutdatedPhysicalMaterialProxiesCount);
+		OutActionableMessage.Tooltip = DefaultTooltip;
+		OutActionableMessage.ActionMessage = DefaultActionMessage;
+		OutActionableMessage.ActionCallback = UE::Landscape::BuildPhysicalMaterial;
 
-	// Outdated Nanite meshes message :
-	DisplayMessageForOutdatedDataFlag(UE::Landscape::EOutdatedDataFlags::NaniteMeshes, LOCTEXT("LANDSCAPE_NANITE_MESHES_NEED_TO_BE_REBUILT_FMT", "LANDSCAPE: {0} {0}|plural(one=ACTOR,other=ACTORS) WITH NANITE MESHES {0}|plural(one=NEEDS,other=NEED) TO BE REBUILT"));
+		return true;
+	}
 
-	// TODO [jonathan.bard] : this should be handled in the same way as the other cases (UE::Landscape::EOutdatedDataFlags::DirtyActors), but we need to slightly refactor the system so that it's 
-	//  based on ALandscapeProxy, rather than ULandscapeInfo/UPackage... : 
-	if (ULandscapeInfoMap* LandscapeInfoMap = ULandscapeInfoMap::FindLandscapeInfoMap(GetWorld()))
+	if (OutdatedNaniteProxiesCount > 0)
 	{
-		int32 ModifiedNotDirtyCount = 0;
-		ForEachLandscapeInfo([&ModifiedNotDirtyCount](ULandscapeInfo* LandscapeInfo)
-		{
-			ModifiedNotDirtyCount += LandscapeInfo->GetModifiedPackageCount();
-			return true;
-		});
-				
-		if (ModifiedNotDirtyCount > 0)
-		{
-			SmallTextItem.SetColor(FLinearColor::Red);
-			SmallTextItem.Text = FText::Format(LOCTEXT("LANDSCAPE_NEED_TO_BE_SAVED", "LANDSCAPE: NEED TO SAVE TO SHOW CHANGES IN A COOKED GAME ({0} {0}|plural(one=object,other=objects))"), ModifiedNotDirtyCount);
-			Canvas->DrawItem(SmallTextItem, FVector2D(XPos, YPos));
-			YPos += FontSizeY;
-		}
+		OutActionableMessage.Message = FText::Format(LOCTEXT("LANDSCAPE_NANITE_MESHES_NEED_TO_BE_REBUILT_FMT", "{0} Landscape {0}|plural(one=actor,other=actors) with Nanite meshes {0}|plural(one=needs,other=need) to be rebuilt"), OutdatedNaniteProxiesCount);
+		OutActionableMessage.Tooltip = DefaultTooltip;
+		OutActionableMessage.ActionMessage = DefaultActionMessage;
+		OutActionableMessage.ActionCallback = UE::Landscape::BuildNanite;
+
+		return true;
 	}
+
+	return false;
 }
 
 FDateTime ULandscapeSubsystem::GetAppCurrentDateTime()
diff --git a/Engine/Source/Runtime/Landscape/Public/ActionableMessageSubsystem.h b/Engine/Source/Runtime/Landscape/Public/ActionableMessageSubsystem.h
new file mode 100644
--- /dev/null
+++ b/Engine/Source/Runtime/Landscape/Public/ActionableMessageSubsystem.h
@@ -0,0 +1,76 @@
+// Copyright Epic Games, Inc. All Rights Reserved.
+
+#pragma once
+
+#if WITH_EDITOR
+
+#include "CoreMinimal.h"
+#include "Subsystems/WorldSubsystem.h"
+
+#include "ActionableMessageSubsystem.generated.h"
+
+struct FActionableMessage
+{
+	/**
+	 * The message that will be displayed in the widget. 
+	 */
+	FText Message;
+
+	/**
+	 * The text displayed inside the action button, if any. 
+	 */
+	FText ActionMessage;
+
+	/**
+	 * The tooltip of the widget. 
+	 */
+	FText Tooltip;
+
+	/**
+	 * The callback associated with the action button, if any. 
+	 */
+	TFunction<void()> ActionCallback;
+};
+
+FORCEINLINE bool operator!=(const FActionableMessage& LHS, const FActionableMessage& RHS);
+
+/**
+ * Used to push messages that will be used by a viewport widget to present to the user.
+ */
+UCLASS()
+class LANDSCAPE_API UActionableMessageSubsystem : public UWorldSubsystem
+{
+	GENERATED_BODY()
+
+public:
+	/**
+	 * Gets all current actionable messages.
+	 * @return A map containing current actionable messages.
+	 */
+	const TMap<FName, TSharedPtr<FActionableMessage>>& GetActionableMessages() const { return ProviderActionableMessageMap; }
+
+	/**
+	 * Pushes a unique actionable message.  
+	 * @param InProvider The name of the entry.
+	 * @param InActionableMessage The content of the actionable message.
+	 */
+	void SetActionableMessage(FName InProvider, const FActionableMessage& InActionableMessage);
+
+	/**
+	 * Removes an entry from the current actionable messages. 
+	 * @param InProvider The name of the entry to remove.
+	 */
+	void ClearActionableMessage(FName InProvider);
+
+	/**
+	 * Returns the current stateID. Any modification in the current entries will update it. 
+	 * @return the current stateID
+	 */
+	uint32 GetStateID() const { return StateID; }
+	
+private:
+	TMap<FName, TSharedPtr<FActionableMessage>> ProviderActionableMessageMap;
+	uint32 StateID = 0;
+};
+
+#endif // WITH_EDITOR
diff --git a/Engine/Source/Runtime/Landscape/Public/LandscapeSubsystem.h b/Engine/Source/Runtime/Landscape/Public/LandscapeSubsystem.h
--- a/Engine/Source/Runtime/Landscape/Public/LandscapeSubsystem.h
+++ b/Engine/Source/Runtime/Landscape/Public/LandscapeSubsystem.h
@@ -16,12 +16,21 @@ class FLandscapeNotificationManager;
 class ULandscapeComponent;
 class FLandscapeGrassMapsBuilder;
 class FLandscapeTextureStreamingManager;
+struct FActionableMessage;
 struct FDateTime;
 struct FScopedSlowTask;
 
 namespace UE::Landscape
 {
 	enum class EOutdatedDataFlags : uint8;
+
+#if WITH_EDITOR
+	void LANDSCAPE_API MarkModifiedLandscapesAsDirty();
+	void LANDSCAPE_API BuildGrassMaps();
+	void LANDSCAPE_API BuildPhysicalMaterial();
+	void LANDSCAPE_API BuildNanite();
+	void LANDSCAPE_API BuildAll();
+#endif // WITH_EDITOR
 } // end of namespace UE::Landscape
 
 #if WITH_EDITOR
@@ -103,7 +112,11 @@ public:
 	LANDSCAPE_API bool IsGridBased() const;
 	LANDSCAPE_API void ChangeGridSize(ULandscapeInfo* LandscapeInfo, uint32 NewGridSizeInComponents);
 	LANDSCAPE_API ALandscapeProxy* FindOrAddLandscapeProxy(ULandscapeInfo* LandscapeInfo, const FIntPoint& SectionBase);
+
+	UE_DEPRECATED(5.5, "DisplayMessages is now deprecated.")
 	LANDSCAPE_API void DisplayMessages(class FCanvas* Canvas, float& XPos, float& YPos);
+
+	LANDSCAPE_API bool GetActionableMessage(FActionableMessage& OutActionableMessage);
 	LANDSCAPE_API void MarkModifiedLandscapesAsDirty();
 	LANDSCAPE_API void SaveModifiedLandscapes();
 	LANDSCAPE_API bool HasModifiedLandscapes() const;
