TEDS-Outliner: Add source control filters

diff --git a/Engine/Source/Editor/LevelEditor/Private/LevelEditorOutlinerSettings.cpp b/Engine/Source/Editor/LevelEditor/Private/LevelEditorOutlinerSettings.cpp
--- a/Engine/Source/Editor/LevelEditor/Private/LevelEditorOutlinerSettings.cpp
+++ b/Engine/Source/Editor/LevelEditor/Private/LevelEditorOutlinerSettings.cpp
@@ -102,6 +102,9 @@ void FLevelEditorOutlinerSettings::SetupBuiltInCategories()
 	
 	TSharedPtr<FFilterCategory> VPFiltersCategory = MakeShared<FFilterCategory>(LOCTEXT("VPFiltersCategory", "Virtual Production"), FText::GetEmpty());
 	FilterBarCategories.Add(FLevelEditorOutlinerBuiltInCategories::VirtualProduction(), VPFiltersCategory);
+
+	TSharedPtr<FFilterCategory> SCCFiltersCategory = MakeShared<FFilterCategory>(LOCTEXT("SCCFiltersCategory", "Revision Control"), FText::GetEmpty());
+	FilterBarCategories.Add(FLevelEditorOutlinerBuiltInCategories::SourceControl(), SCCFiltersCategory);
 	
 	// Now convert some of the built in placement mode categories we want to filter categories and add them
 
@@ -331,7 +334,7 @@ void FLevelEditorOutlinerSettings::OnUnsavedAssetRemoved(const FString& InAsset)
 void FLevelEditorOutlinerSettings::CreateSCCFilters()
 {
 	// Source Control Category
-	TSharedPtr<FFilterCategory> SCCFiltersCategory = MakeShared<FFilterCategory>(LOCTEXT("SCCFiltersCategory", "Revision Control"), FText::GetEmpty());
+	TSharedPtr<FFilterCategory> SCCFiltersCategory = GetFilterCategory(FLevelEditorOutlinerBuiltInCategories::SourceControl());
 	
 	// Uncontrolled Actors Filter
 	FUncontrolledChangelistsModule& UncontrolledChangelistModule = FUncontrolledChangelistsModule::Get();
diff --git a/Engine/Source/Editor/LevelEditor/Public/LevelEditorOutlinerSettings.h b/Engine/Source/Editor/LevelEditor/Public/LevelEditorOutlinerSettings.h
--- a/Engine/Source/Editor/LevelEditor/Public/LevelEditorOutlinerSettings.h
+++ b/Engine/Source/Editor/LevelEditor/Public/LevelEditorOutlinerSettings.h
@@ -29,6 +29,7 @@ struct FLevelEditorOutlinerBuiltInCategories
 	static FName Visual()			 { static FName Name("Visual");				return Name; }
 	static FName Volumes()			 { static FName Name("Volumes");			return Name; }
 	static FName VirtualProduction() { static FName Name("VirtualProduction");	return Name; }
+	static FName SourceControl()	 { static FName Name("RevisionControl");	return Name; }
 };
 
 /** Helper class to manage initalization options specific to the Level Editor Outliners
