Fix level viewport tool bar using the wrong width for its dimension.
#tests Repro from the jira

diff --git a/Engine/Source/Editor/LevelEditor/Private/SLevelViewportToolBar.cpp b/Engine/Source/Editor/LevelEditor/Private/SLevelViewportToolBar.cpp
--- a/Engine/Source/Editor/LevelEditor/Private/SLevelViewportToolBar.cpp
+++ b/Engine/Source/Editor/LevelEditor/Private/SLevelViewportToolBar.cpp
@@ -2037,37 +2037,31 @@ float SLevelViewportToolBar::GetTransformToolbarWidth() const
 			TransformToolbar_CachedMaxWidth = TransformToolbarWidth;
 		}
 
+		const float ToolbarWidthMinusPreviousTransformToolbar = GetDesiredSize().X - TransformToolbar_CachedMaxWidth;
+		const float ToolbarWidthEstimate = ToolbarWidthMinusPreviousTransformToolbar + TransformToolbarWidth;
+
+		float DpiScale = 1.0f;
 		{
-			FLevelEditorViewportClient* LevelEditorViewportClient = GetLevelViewportClient();
-			if (LevelEditorViewportClient && LevelEditorViewportClient->Viewport)
+			IMainFrameModule& MainFrameModule = FModuleManager::LoadModuleChecked<IMainFrameModule>(TEXT("MainFrame"));
+			const TSharedPtr<SWindow>& MainFrameParentWindow = MainFrameModule.GetParentWindow();
+			if (MainFrameParentWindow.IsValid())
 			{
-				const float ViewportWidth = static_cast<float>(LevelEditorViewportClient->Viewport->GetSizeXY().X);
-				const float ToolbarWidthMinusPreviousTransformToolbar = GetDesiredSize().X - TransformToolbar_CachedMaxWidth;
-				const float ToolbarWidthEstimate = ToolbarWidthMinusPreviousTransformToolbar + TransformToolbarWidth;
-
-				float DpiScale = 1.0f;
-				{
-					IMainFrameModule& MainFrameModule = FModuleManager::LoadModuleChecked<IMainFrameModule>(TEXT("MainFrame"));
-					const TSharedPtr<SWindow>& MainFrameParentWindow = MainFrameModule.GetParentWindow();
-					if (MainFrameParentWindow.IsValid())
-					{
-						DpiScale = MainFrameParentWindow->GetDPIScaleFactor();
-					}
-				}
-
-				const float OverflowWidth = ToolbarWidthEstimate * DpiScale - ViewportWidth;
-				if (OverflowWidth > 0.0f)
-				{
-					// There isn't enough space in the viewport to show the toolbar!
-					// Try and shrink the transform toolbar (which has an overflow area) to make things fit
-					TransformToolbar_CachedMaxWidth = FMath::Max(FMath::Min(4.0f, TransformToolbarWidth), TransformToolbarWidth - OverflowWidth / DpiScale);
-				}
-				else
-				{
-					TransformToolbar_CachedMaxWidth = TransformToolbarWidth;
-				}
+				DpiScale = MainFrameParentWindow->GetDPIScaleFactor();
 			}
 		}
+
+		const float ViewportToolBarWidth = static_cast<float>(GetCachedGeometry().GetLocalSize().X);
+		const float OverflowWidth = ToolbarWidthEstimate * DpiScale - ViewportToolBarWidth;
+		if (OverflowWidth > 0.0f)
+		{
+			// There isn't enough space in the viewport to show the toolbar!
+			// Try and shrink the transform toolbar (which has an overflow area) to make things fit
+			TransformToolbar_CachedMaxWidth = FMath::Max(FMath::Min(4.0f, TransformToolbarWidth), TransformToolbarWidth - OverflowWidth / DpiScale);
+		}
+		else
+		{
+			TransformToolbar_CachedMaxWidth = TransformToolbarWidth;
+		}
 		
 		return TransformToolbar_CachedMaxWidth;
 	}
