Allow Game View to be toggled in ortho viewports

diff --git a/Engine/Plugins/VirtualProduction/VirtualCamera/Source/VirtualCamera/Private/FunctionLibraries/GameViewFunctionLibrary.cpp b/Engine/Plugins/VirtualProduction/VirtualCamera/Source/VirtualCamera/Private/FunctionLibraries/GameViewFunctionLibrary.cpp
--- a/Engine/Plugins/VirtualProduction/VirtualCamera/Source/VirtualCamera/Private/FunctionLibraries/GameViewFunctionLibrary.cpp
+++ b/Engine/Plugins/VirtualProduction/VirtualCamera/Source/VirtualCamera/Private/FunctionLibraries/GameViewFunctionLibrary.cpp
@@ -26,8 +26,10 @@ void UGameViewFunctionLibrary::ToggleGameView(EVCamTargetViewportID ViewportID)
 bool UGameViewFunctionLibrary::CanToggleGameView(EVCamTargetViewportID ViewportID)
 {
 #if WITH_EDITOR
+PRAGMA_DISABLE_DEPRECATION_WARNINGS
 	const TSharedPtr<SLevelViewport> Viewport = UE::VCamCore::GetLevelViewport(ViewportID);
 	return Viewport && Viewport->CanToggleGameView();
+PRAGMA_ENABLE_DEPRECATION_WARNINGS
 #else
 	return false;
 #endif
diff --git a/Engine/Source/Developer/FunctionalTesting/Private/AutomationBlueprintFunctionLibrary.cpp b/Engine/Source/Developer/FunctionalTesting/Private/AutomationBlueprintFunctionLibrary.cpp
--- a/Engine/Source/Developer/FunctionalTesting/Private/AutomationBlueprintFunctionLibrary.cpp
+++ b/Engine/Source/Developer/FunctionalTesting/Private/AutomationBlueprintFunctionLibrary.cpp
@@ -668,7 +668,7 @@ public:
 			{
 				if (NeedGameViewToggle)
 				{
-					if (LevelViewport->IsInGameView() && LevelViewport->CanToggleGameView())
+					if (LevelViewport->IsInGameView())
 					{
 						LevelViewport->ToggleGameView();
 					}
@@ -1237,7 +1237,7 @@ UAutomationEditorTask* UAutomationBlueprintFunctionLibrary::TakeHighResScreensho
 			FLevelEditorModule& LevelEditor = FModuleManager::GetModuleChecked<FLevelEditorModule>("LevelEditor");
 			SLevelViewport* LevelViewport = LevelEditor.GetFirstActiveLevelViewport().Get();
 			bool bNeedGameViewToggle = bForceGameView && !LevelViewport->IsInGameView();
-			if (bNeedGameViewToggle && LevelViewport->CanToggleGameView())
+			if (bNeedGameViewToggle)
 			{
 				LevelViewport->ToggleGameView();
 			}
diff --git a/Engine/Source/Editor/LevelEditor/Private/SLevelViewport.cpp b/Engine/Source/Editor/LevelEditor/Private/SLevelViewport.cpp
--- a/Engine/Source/Editor/LevelEditor/Private/SLevelViewport.cpp
+++ b/Engine/Source/Editor/LevelEditor/Private/SLevelViewport.cpp
@@ -1456,7 +1456,7 @@ void SLevelViewport::BindOptionCommands( FUICommandList& OutCommandList )
 	OutCommandList.MapAction(
 		ViewportActions.ToggleGameView,
 		FExecuteAction::CreateSP( this, &SLevelViewport::ToggleGameView ),
-		FCanExecuteAction::CreateSP( this, &SLevelViewport::CanToggleGameView ),
+		FCanExecuteAction(),
 		FIsActionChecked::CreateSP( this, &SLevelViewport::IsInGameView ) );
 
 	OutCommandList.MapAction(
@@ -2044,21 +2044,18 @@ void SLevelViewport::OnTakeHighResScreenshot()
 
 void SLevelViewport::ToggleGameView()
 {
-	if( LevelViewportClient->IsPerspective() )
-	{
-		bool bGameViewEnable = !LevelViewportClient->IsInGameView();
+	bool bGameViewEnable = !LevelViewportClient->IsInGameView();
 
-		// "Mode Widget" should not automatically be reactivated by selecting an actor after "Game View" is enabled
-		LevelViewportClient->bAlwaysShowModeWidgetAfterSelectionChanges = bGameViewEnable ? false : true;
+	// "Mode Widget" should not automatically be reactivated by selecting an actor after "Game View" is enabled
+	LevelViewportClient->bAlwaysShowModeWidgetAfterSelectionChanges = bGameViewEnable ? false : true;
 
-		LevelViewportClient->SetGameView(bGameViewEnable);
+	LevelViewportClient->SetGameView(bGameViewEnable);
 
-		if (!bGameViewEnable)
-		{
-			// LevelViewportClient->bShowWidget is set to "false" when entering game mode
-			// Need to turn it back to "true" when exiting game mode
-			LevelViewportClient->ShowWidget(true);
-		}
+	if (!bGameViewEnable)
+	{
+		// LevelViewportClient->bShowWidget is set to "false" when entering game mode
+		// Need to turn it back to "true" when exiting game mode
+		LevelViewportClient->ShowWidget(true);
 	}
 }
 
diff --git a/Engine/Source/Editor/LevelEditor/Public/SLevelViewport.h b/Engine/Source/Editor/LevelEditor/Public/SLevelViewport.h
--- a/Engine/Source/Editor/LevelEditor/Public/SLevelViewport.h
+++ b/Engine/Source/Editor/LevelEditor/Public/SLevelViewport.h
@@ -195,6 +195,7 @@ public:
 	/**
 	 * @return true if we can toggle game view
 	 */
+	UE_DEPRECATED(5.5, "CanToggleGameView is deprecated. Please update usage to assume that this function would return true in the future. All level viewports should support game view moving forward.")
 	bool CanToggleGameView() const;
 
 	/**
