[Viewport Toolbar] Fix errant merge issue from CL 33047928

diff --git a/Engine/Source/Editor/LevelEditor/Private/SLevelViewport.cpp b/Engine/Source/Editor/LevelEditor/Private/SLevelViewport.cpp
--- a/Engine/Source/Editor/LevelEditor/Private/SLevelViewport.cpp
+++ b/Engine/Source/Editor/LevelEditor/Private/SLevelViewport.cpp
@@ -2042,21 +2042,18 @@ void SLevelViewport::OnTakeHighResScreenshot()
 
 void SLevelViewport::ToggleGameView()
 {
-	if( LevelViewportClient->IsPerspective() )
-	{
-		bool bGameViewEnable = !LevelViewportClient->IsInGameView();
+	bool bGameViewEnable = !LevelViewportClient->IsInGameView();
 
-		// "Mode Widget" should not automatically be reactivated by selecting an actor after "Game View" is enabled
-		LevelViewportClient->bAlwaysShowModeWidgetAfterSelectionChanges = bGameViewEnable ? false : true;
+	// "Mode Widget" should not automatically be reactivated by selecting an actor after "Game View" is enabled
+	LevelViewportClient->bAlwaysShowModeWidgetAfterSelectionChanges = bGameViewEnable ? false : true;
 
-		LevelViewportClient->SetGameView(bGameViewEnable);
+	LevelViewportClient->SetGameView(bGameViewEnable);
 
-		if (!bGameViewEnable)
-		{
-			// LevelViewportClient->bShowWidget is set to "false" when entering game mode
-			// Need to turn it back to "true" when exiting game mode
-			LevelViewportClient->ShowWidget(true);
-		}
+	if (!bGameViewEnable)
+	{
+		// LevelViewportClient->bShowWidget is set to "false" when entering game mode
+		// Need to turn it back to "true" when exiting game mode
+		LevelViewportClient->ShowWidget(true);
 	}
 }
 
