Level Viewport: Added Scene Captures to the list of cameras that users can look through in the viewport.

diff --git a/Engine/Source/Editor/LevelEditor/Private/SLevelViewportToolBar.cpp b/Engine/Source/Editor/LevelEditor/Private/SLevelViewportToolBar.cpp
--- a/Engine/Source/Editor/LevelEditor/Private/SLevelViewportToolBar.cpp
+++ b/Engine/Source/Editor/LevelEditor/Private/SLevelViewportToolBar.cpp
@@ -56,6 +56,7 @@
 #include "Interfaces/IMainFrameModule.h"
 #include "SCommonEditorViewportToolbarBase.h"
 #include "SActionableMessageViewportWidget.h"
+#include "Engine/SceneCapture.h"
 
 #define LOCTEXT_NAMESPACE "LevelViewportToolBar"
 
@@ -951,12 +952,12 @@ void SLevelViewportToolBar::SetLevelProfile( FString DeviceProfileName )
 	UIManager->SetProfile( DeviceProfileName );
 }
 
-void SLevelViewportToolBar::GeneratePlacedCameraMenuEntries(FToolMenuSection& Section, TArray<ACameraActor*> Cameras) const
+void SLevelViewportToolBar::GeneratePlacedCameraMenuEntries(FToolMenuSection& Section, TArray<AActor*> LookThroughActors) const
 {
 	FSlateIcon CameraIcon( FAppStyle::GetAppStyleSetName(), "ClassIcon.CameraComponent" );
 
 	// Sort the cameras to make the ordering predictable for users.
-	Cameras.StableSort([](const ACameraActor& Left, const ACameraActor& Right)
+	LookThroughActors.StableSort([](const AActor& Left, const AActor& Right)
 	{
 		// Do "natural sorting" via SceneOutliner::FNumericStringWrapper to make more sense to humans (also matches the Scene Outliner). This sorts "Camera2" before "Camera10" which a normal lexicographical sort wouldn't.
 		SceneOutliner::FNumericStringWrapper LeftWrapper(FString(Left.GetActorLabel()));
@@ -965,12 +966,12 @@ void SLevelViewportToolBar::GeneratePlacedCameraMenuEntries(FToolMenuSection& Se
 		return LeftWrapper < RightWrapper;
 	});
 
-	for( ACameraActor* CameraActor : Cameras )
+	for( AActor* LookThroughActor : LookThroughActors)
 	{
 		// Needed for the delegate hookup to work below
-		AActor* GenericActor = CameraActor;
+		AActor* GenericActor = LookThroughActor;
 
-		FText ActorDisplayName = FText::FromString(CameraActor->GetActorLabel());
+		FText ActorDisplayName = FText::FromString(LookThroughActor->GetActorLabel());
 		FUIAction LookThroughCameraAction(
 			FExecuteAction::CreateSP(Viewport.Pin().ToSharedRef(), &SLevelViewport::OnActorLockToggleFromMenu, GenericActor),
 			FCanExecuteAction(),
@@ -981,10 +982,10 @@ void SLevelViewportToolBar::GeneratePlacedCameraMenuEntries(FToolMenuSection& Se
 	}
 }
 
-void SLevelViewportToolBar::GeneratePlacedCameraMenuEntries(UToolMenu* Menu, TArray<ACameraActor*> Cameras) const
+void SLevelViewportToolBar::GeneratePlacedCameraMenuEntries(UToolMenu* Menu, TArray<AActor*> LookThroughActors) const
 {
 	FToolMenuSection& Section = Menu->AddSection("Section");
-	GeneratePlacedCameraMenuEntries(Section, Cameras);
+	GeneratePlacedCameraMenuEntries(Section, LookThroughActors);
 }
 
 void SLevelViewportToolBar::GenerateViewportTypeMenu(FToolMenuSection& Section) const
@@ -1056,26 +1057,31 @@ void SLevelViewportToolBar::FillCameraMenu(UToolMenu* Menu) const
 		Section.AddMenuEntry(FEditorViewportCommands::Get().Back);
 	}
 
-	TArray<ACameraActor*> Cameras;
+	TArray<AActor*> LookThroughActors;
 
 	for( TActorIterator<ACameraActor> It(GetWorld().Get()); It; ++It )
 	{
-		Cameras.Add( *It );
+		LookThroughActors.Add( Cast<AActor>(*It) );
 	}
 
-	FText CameraActorsHeading = LOCTEXT("CameraActorsHeading", "Placed Cameras");
+	for (TActorIterator<ASceneCapture> It(GetWorld().Get()); It; ++It)
+	{
+		LookThroughActors.Add(Cast<AActor>(*It));
+	}
+
+	FText CameraActorsHeading = LOCTEXT("CameraActorsHeading", "Placed Cameras and Scene Capture Actors");
 
 	// Don't add too many cameras to the top level menu or else it becomes too large
 	const uint32 MaxCamerasInTopLevelMenu = 10;
-	if( Cameras.Num() > MaxCamerasInTopLevelMenu )
+	if(LookThroughActors.Num() > MaxCamerasInTopLevelMenu )
 	{
 		FToolMenuSection& Section = Menu->AddSection("CameraActors");
-		Section.AddSubMenu("CameraActors", CameraActorsHeading, LOCTEXT("LookThroughPlacedCameras_ToolTip", "Look through and pilot placed cameras"), FNewToolMenuDelegate::CreateSP(this, &SLevelViewportToolBar::GeneratePlacedCameraMenuEntries, Cameras ) );
+		Section.AddSubMenu("CameraActors", CameraActorsHeading, LOCTEXT("LookThroughPlacedCameras_ToolTip", "Look through and pilot placed cameras"), FNewToolMenuDelegate::CreateSP(this, &SLevelViewportToolBar::GeneratePlacedCameraMenuEntries, LookThroughActors ) );
 	}
 	else
 	{
 		FToolMenuSection& Section = Menu->AddSection("CameraActors", CameraActorsHeading);
-		GeneratePlacedCameraMenuEntries(Section, Cameras);
+		GeneratePlacedCameraMenuEntries(Section, LookThroughActors);
 	}
 
 	{
diff --git a/Engine/Source/Editor/LevelEditor/Private/SLevelViewportToolBar.h b/Engine/Source/Editor/LevelEditor/Private/SLevelViewportToolBar.h
--- a/Engine/Source/Editor/LevelEditor/Private/SLevelViewportToolBar.h
+++ b/Engine/Source/Editor/LevelEditor/Private/SLevelViewportToolBar.h
@@ -134,20 +134,20 @@ private:
 	void FillCameraMenu(UToolMenu* Menu) const;
 
 	/**
-	 * Generates menu entries for placed cameras (e.g CameraActors
+	 * Generates menu entries for placed cameras and Scene Capture Actors (e.g CameraActors, SceneCapture)
 	 *
 	 * @param Menu	The menu to add menu entries to
-	 * @param Cameras	The list of cameras to add
+	 * @param Cameras	The list of actors to add
 	 */
-	void GeneratePlacedCameraMenuEntries(UToolMenu* Menu, TArray<ACameraActor*> Cameras) const;
+	void GeneratePlacedCameraMenuEntries(UToolMenu* Menu, TArray<AActor*> LookThroughActors) const;
 
 	/**
-	 * Generates menu entries for placed cameras (e.g CameraActors
+	 * Generates menu entries for placed cameras and Scene Capture Actors (e.g CameraActors, SceneCapture)
 	 *
 	 * @param Section	The menu section to add menu entries to
-	 * @param Cameras	The list of cameras to add
+	 * @param Cameras	The list of actors to add
 	 */
-	void GeneratePlacedCameraMenuEntries(FToolMenuSection& Section, TArray<ACameraActor*> Cameras) const;
+	void GeneratePlacedCameraMenuEntries(FToolMenuSection& Section, TArray<AActor*> LookThroughActors) const;
 
 	/**
 	 * Generates menu entries for changing the type of the viewport
