Move LEVELEDITOR_API from class FLevelEditorCommands to a new FLevelEditorCommands::Get()

diff --git a/Engine/Source/Editor/LevelEditor/Private/LevelEditorActions.cpp b/Engine/Source/Editor/LevelEditor/Private/LevelEditorActions.cpp
--- a/Engine/Source/Editor/LevelEditor/Private/LevelEditorActions.cpp
+++ b/Engine/Source/Editor/LevelEditor/Private/LevelEditorActions.cpp
@@ -4007,6 +4007,11 @@ void FLevelEditorCommands::RegisterCommands()
 	UI_COMMAND(OpenMergeActor, "Merge Actors", "Opens the Merge Actor panel", EUserInterfaceActionType::Button, FInputChord());
 }
 
+FORCENOINLINE const FLevelEditorCommands& FLevelEditorCommands::Get()
+{
+	return TCommands<FLevelEditorCommands>::Get();
+}
+
 UE_ENABLE_OPTIMIZATION_SHIP
 
 void FLevelEditorActionCallbacks::FixupGroupActor_Clicked()
diff --git a/Engine/Source/Editor/LevelEditor/Public/LevelEditorActions.h b/Engine/Source/Editor/LevelEditor/Public/LevelEditorActions.h
--- a/Engine/Source/Editor/LevelEditor/Public/LevelEditorActions.h
+++ b/Engine/Source/Editor/LevelEditor/Public/LevelEditorActions.h
@@ -25,24 +25,19 @@ class UTypedElementSelectionSet;
 /**
  * Unreal level editor actions
  */
-class LEVELEDITOR_API FLevelEditorCommands : public TCommands<FLevelEditorCommands>
+class FLevelEditorCommands : public TCommands<FLevelEditorCommands>
 {
 
 public:
 	FLevelEditorCommands();
 
-	PRAGMA_DISABLE_DEPRECATION_WARNINGS
-	FLevelEditorCommands(FLevelEditorCommands&&) = default;
-	FLevelEditorCommands(const FLevelEditorCommands&) = default;
-	FLevelEditorCommands& operator=(FLevelEditorCommands&&) = default;
-	FLevelEditorCommands& operator=(const FLevelEditorCommands&) = default;
-	PRAGMA_ENABLE_DEPRECATION_WARNINGS
-
 	/**
 	 * Initialize commands
 	 */
 	virtual void RegisterCommands() override;
 
+	LEVELEDITOR_API FORCENOINLINE static const FLevelEditorCommands& Get();
+
 public:
 	
 	TSharedPtr< FUICommandInfo > BrowseDocumentation;
