Add ToolMenus-based viewport toolbar "LevelEditor.ViewportToolbar" behind CVAR LevelEditor.ToolMenusViewportToolbar

diff --git a/Engine/Source/Editor/LevelEditor/Private/SLevelViewport.cpp b/Engine/Source/Editor/LevelEditor/Private/SLevelViewport.cpp
--- a/Engine/Source/Editor/LevelEditor/Private/SLevelViewport.cpp
+++ b/Engine/Source/Editor/LevelEditor/Private/SLevelViewport.cpp
@@ -101,11 +101,17 @@
 #include "SWorldPartitionViewportWidget.h"
 #include "LevelViewportLayout.h"
 #include "EditorViewportTabContent.h"
+#include "EditorViewportCommands.h"
 
 static const FName LevelEditorName("LevelEditor");
 static FAutoConsoleCommand EnableInViewportMenu(TEXT("Editor.EnableInViewportMenu"), TEXT("Enables the new in-viewport property menu"), FConsoleCommandDelegate::CreateStatic(&SLevelViewport::EnableInViewportMenu));
 bool SLevelViewport::bInViewportMenuEnabled = false;
 
+static TAutoConsoleVariable<int> CVarLevelEditorToolMenusViewportToolbar(
+	TEXT("LevelEditor.ToolMenusViewportToolbar"),
+	0,
+	TEXT("If set to 1, the new UToolMenus-based level editor viewport toolbar will be displayed below the current level editor viewport toolbar. If set to 2, the new will replace the current. If set to 0 (default), the current is shown but not the new."));
+
 #define LOCTEXT_NAMESPACE "LevelViewport"
 
 // @todo Slate Hack: Disallow game UI to be used in play in viewport until GWorld problem is fixed
@@ -1886,28 +1892,65 @@ TSharedRef<FEditorViewportClient> SLevelViewport::MakeEditorViewportClient()
 	return LevelViewportClient.ToSharedRef();
 }
 
+void BuildToolMenusViewportToolbar()
+{
+	// Add a ToolMenus-based viewport toolbar.
+	UToolMenu* const ViewportToolbarMenu = UToolMenus::Get()->RegisterMenu(
+		"LevelEditor.ViewportToolbar", NAME_None /* parent */, EMultiBoxType::ToolBar);
+
+	// Transforms section
+	{
+		FToolMenuSection& Section = ViewportToolbarMenu->FindOrAddSection(
+			"Transforms", LOCTEXT("TransformsSectionLabel", "Transforms"));
+
+		Section.AddEntry(FToolMenuEntry::InitToolBarButton(FEditorViewportCommands::Get().SelectMode));
+		Section.AddEntry(FToolMenuEntry::InitToolBarButton(FEditorViewportCommands::Get().TranslateMode));
+		Section.AddEntry(FToolMenuEntry::InitToolBarButton(FEditorViewportCommands::Get().RotateMode));
+		Section.AddEntry(FToolMenuEntry::InitToolBarButton(FEditorViewportCommands::Get().ScaleMode));
+	}
+}
+
 TSharedPtr<SWidget> SLevelViewport::MakeViewportToolbar()
 {
+	TSharedRef<SVerticalBox> VerticalBox = SNew(SVerticalBox)
+		.Visibility(EVisibility::SelfHitTestInvisible);
 
-	// Build our toolbar level toolbar
-	TSharedRef< SLevelViewportToolBar > ToolBar =
-		SNew( SLevelViewportToolBar )
-		.Viewport( SharedThis( this ) )
-		.Visibility( this, &SLevelViewport::GetToolBarVisibility )
-		.IsEnabled( FSlateApplication::Get().GetNormalExecutionAttribute() );
+	const int CVarLevelEditorToolMenusViewportToolbarValue = CVarLevelEditorToolMenusViewportToolbar.GetValueOnGameThread();
+	const bool bShowOldViewportToolbar = CVarLevelEditorToolMenusViewportToolbarValue != 2;
+	const bool bShowNewViewportToolbar = CVarLevelEditorToolMenusViewportToolbarValue > 0;
 
+	if (bShowOldViewportToolbar)
+	{
+		TSharedRef<SLevelViewportToolBar> Toolbar = SNew(SLevelViewportToolBar)
+			.Viewport(SharedThis(this))
+			.Visibility(this, &SLevelViewport::GetToolBarVisibility)
+			.IsEnabled(FSlateApplication::Get().GetNormalExecutionAttribute());
 
-	return 
-		SNew(SVerticalBox)
-		.Visibility( EVisibility::SelfHitTestInvisible )
-		+SVerticalBox::Slot()
-		.AutoHeight()
-		.Padding(0, 1.0f, 0, 0)
-		.VAlign(VAlign_Top)
-		[
-			ToolBar
-		]
-		+SVerticalBox::Slot()
+		VerticalBox->AddSlot()
+			.AutoHeight()
+			.Padding(0, 1.0f, 0, 0)
+			.VAlign(VAlign_Top)
+			[
+				Toolbar
+			];
+	}
+
+	if (bShowNewViewportToolbar)
+	{
+		BuildToolMenusViewportToolbar();
+
+		VerticalBox->AddSlot()
+			.AutoHeight()
+			.Padding(0, 1.0f, 0, 0)
+			.VAlign(VAlign_Top)
+			[
+				UToolMenus::Get()->GenerateWidget(
+					"LevelEditor.ViewportToolbar",
+					FToolMenuContext(GetLevelViewportClient().GetEditorViewportWidget()->GetCommandList()))
+			];
+	}
+
+	VerticalBox->AddSlot()
 		.VAlign(VAlign_Top)
 		.HAlign(HAlign_Left)
 		[
@@ -1915,6 +1958,8 @@ TSharedPtr<SWidget> SLevelViewport::MakeViewportToolbar()
 			.Viewport( SharedThis( this ) )
 			.Visibility(this, &SLevelViewport::GetLockedIconVisibility)
 		];
+
+	return VerticalBox;
 }
 
 void SLevelViewport::OnUndo()
