Improve the experimental UToolMenu "LevelEditor.ViewportToolbar"

* Changes to LevelEditor.ToolMenusViewportToolbar now reflect immediately.
* CVAR value is stored in a file-private variable.

diff --git a/Engine/Source/Editor/LevelEditor/Private/SLevelViewport.cpp b/Engine/Source/Editor/LevelEditor/Private/SLevelViewport.cpp
--- a/Engine/Source/Editor/LevelEditor/Private/SLevelViewport.cpp
+++ b/Engine/Source/Editor/LevelEditor/Private/SLevelViewport.cpp
@@ -107,10 +107,17 @@ static const FName LevelEditorName("LevelEditor");
 static FAutoConsoleCommand EnableInViewportMenu(TEXT("Editor.EnableInViewportMenu"), TEXT("Enables the new in-viewport property menu"), FConsoleCommandDelegate::CreateStatic(&SLevelViewport::EnableInViewportMenu));
 bool SLevelViewport::bInViewportMenuEnabled = false;
 
-static TAutoConsoleVariable<int> CVarLevelEditorToolMenusViewportToolbar(
+namespace UE::SLevelViewport::Private
+{
+	int32 CVarLevelEditorToolMenusViewportToolbarValue = 0;
+}
+
+static FAutoConsoleVariableRef CVarLevelEditorToolMenusViewportToolbar(
 	TEXT("LevelEditor.ToolMenusViewportToolbar"),
-	0,
-	TEXT("If set to 1, the new UToolMenus-based level editor viewport toolbar will be displayed below the current level editor viewport toolbar. If set to 2, the new will replace the current. If set to 0 (default), the current is shown but not the new."));
+	UE::SLevelViewport::Private::CVarLevelEditorToolMenusViewportToolbarValue,
+	TEXT("If set to 1, the new UToolMenus-based level editor viewport toolbar will be displayed below the current level editor viewport toolbar. If set to 2, the new will replace the current. If set to 0 (default), the current is shown but not the new."),
+	ECVF_Default
+);
 
 #define LOCTEXT_NAMESPACE "LevelViewport"
 
@@ -1912,45 +1919,57 @@ void BuildToolMenusViewportToolbar()
 
 TSharedPtr<SWidget> SLevelViewport::MakeViewportToolbar()
 {
-	TSharedRef<SVerticalBox> VerticalBox = SNew(SVerticalBox)
-		.Visibility(EVisibility::SelfHitTestInvisible);
-
-	const int CVarLevelEditorToolMenusViewportToolbarValue = CVarLevelEditorToolMenusViewportToolbar.GetValueOnGameThread();
-	const bool bShowOldViewportToolbar = CVarLevelEditorToolMenusViewportToolbarValue != 2;
-	const bool bShowNewViewportToolbar = CVarLevelEditorToolMenusViewportToolbarValue > 0;
+	const TSharedRef<SLevelViewportToolBar> OldViewportToolbar = SNew(SLevelViewportToolBar)
+		.Viewport(SharedThis(this))
+		.Visibility_Lambda([this]() -> EVisibility
+		{
+			const bool bShowOldViewportToolbar = UE::SLevelViewport::Private::CVarLevelEditorToolMenusViewportToolbarValue != 2;
+			if (!bShowOldViewportToolbar)
+			{
+				return EVisibility::Collapsed;
+			}
 
-	if (bShowOldViewportToolbar)
-	{
-		TSharedRef<SLevelViewportToolBar> Toolbar = SNew(SLevelViewportToolBar)
-			.Viewport(SharedThis(this))
-			.Visibility(this, &SLevelViewport::GetToolBarVisibility)
-			.IsEnabled(FSlateApplication::Get().GetNormalExecutionAttribute());
+			return GetToolBarVisibility();
+		})
+		.IsEnabled(FSlateApplication::Get().GetNormalExecutionAttribute());
 
-		VerticalBox->AddSlot()
-			.AutoHeight()
-			.Padding(0, 1.0f, 0, 0)
-			.VAlign(VAlign_Top)
-			[
-				Toolbar
-			];
-	}
+	BuildToolMenusViewportToolbar();
 
-	if (bShowNewViewportToolbar)
-	{
-		BuildToolMenusViewportToolbar();
+	const TSharedRef<SWidget> NewViewportToolbar = SNew(SBox)
+		.Visibility_Lambda([this]() -> EVisibility
+		{
+			const bool bShowNewViewportToolbar = UE::SLevelViewport::Private::CVarLevelEditorToolMenusViewportToolbarValue > 0;
+			if (!bShowNewViewportToolbar)
+			{
+				return EVisibility::Collapsed;
+			}
 
-		VerticalBox->AddSlot()
-			.AutoHeight()
-			.Padding(0, 1.0f, 0, 0)
-			.VAlign(VAlign_Top)
-			[
-				UToolMenus::Get()->GenerateWidget(
-					"LevelEditor.ViewportToolbar",
-					FToolMenuContext(GetLevelViewportClient().GetEditorViewportWidget()->GetCommandList()))
-			];
-	}
+			return GetToolBarVisibility();
+		})
+		[
+			UToolMenus::Get()->GenerateWidget(
+				"LevelEditor.ViewportToolbar",
+				FToolMenuContext(GetLevelViewportClient().GetEditorViewportWidget()->GetCommandList()))
+		];
 
-	VerticalBox->AddSlot()
+	return 
+		SNew(SVerticalBox)
+		.Visibility( EVisibility::SelfHitTestInvisible )
+		+SVerticalBox::Slot()
+		.AutoHeight()
+		.Padding(0, 1.0f, 0, 0)
+		.VAlign(VAlign_Top)
+		[
+			OldViewportToolbar
+		]
+		+SVerticalBox::Slot()
+		.AutoHeight()
+		.Padding(0, 1.0f, 0, 0)
+		.VAlign(VAlign_Top)
+		[
+			NewViewportToolbar
+		]
+		+SVerticalBox::Slot()
 		.VAlign(VAlign_Top)
 		.HAlign(HAlign_Left)
 		[
@@ -1958,8 +1977,6 @@ TSharedPtr<SWidget> SLevelViewport::MakeViewportToolbar()
 			.Viewport( SharedThis( this ) )
 			.Visibility(this, &SLevelViewport::GetLockedIconVisibility)
 		];
-
-	return VerticalBox;
 }
 
 void SLevelViewport::OnUndo()
