Use the SLevelViewport command list when generating a widget for the UToolMenu "LevelEditor.ViewportToolbar".

This can prevent crashes if the viewport widget hasn't yet been set.

diff --git a/Engine/Source/Editor/LevelEditor/Private/SLevelViewport.cpp b/Engine/Source/Editor/LevelEditor/Private/SLevelViewport.cpp
--- a/Engine/Source/Editor/LevelEditor/Private/SLevelViewport.cpp
+++ b/Engine/Source/Editor/LevelEditor/Private/SLevelViewport.cpp
@@ -1947,9 +1947,7 @@ TSharedPtr<SWidget> SLevelViewport::MakeViewportToolbar()
 			return GetToolBarVisibility();
 		})
 		[
-			UToolMenus::Get()->GenerateWidget(
-				"LevelEditor.ViewportToolbar",
-				FToolMenuContext(GetLevelViewportClient().GetEditorViewportWidget()->GetCommandList()))
+			UToolMenus::Get()->GenerateWidget("LevelEditor.ViewportToolbar", FToolMenuContext(GetCommandList()))
 		];
 
 	return 
