Add "Settings" section to UToolMenu "LevelEditor.ViewportToolbar"

* Add buttons to the new section to toggle realtime rendering, maximize/restore the viewport, and toggle immersive mode.
* Add a context to the "LevelEditor.ViewportToolbar" when a widget is generated for it. This context carries the SLevelViewport so the maximize/restore button can act on the correct viewport (the same button is drawn in all level editor viewports).

diff --git a/Engine/Source/Editor/LevelEditor/Private/SLevelViewport.cpp b/Engine/Source/Editor/LevelEditor/Private/SLevelViewport.cpp
--- a/Engine/Source/Editor/LevelEditor/Private/SLevelViewport.cpp
+++ b/Engine/Source/Editor/LevelEditor/Private/SLevelViewport.cpp
@@ -102,6 +102,8 @@
 #include "LevelViewportLayout.h"
 #include "EditorViewportTabContent.h"
 #include "EditorViewportCommands.h"
+#include "ViewportToolbar/LevelEditorViewportToolbarSections.h"
+#include "ViewportToolbar/LevelViewportContext.h"
 
 static const FName LevelEditorName("LevelEditor");
 static FAutoConsoleCommand EnableInViewportMenu(TEXT("Editor.EnableInViewportMenu"), TEXT("Enables the new in-viewport property menu"), FConsoleCommandDelegate::CreateStatic(&SLevelViewport::EnableInViewportMenu));
@@ -1899,24 +1901,6 @@ TSharedRef<FEditorViewportClient> SLevelViewport::MakeEditorViewportClient()
 	return LevelViewportClient.ToSharedRef();
 }
 
-void BuildToolMenusViewportToolbar()
-{
-	// Add a ToolMenus-based viewport toolbar.
-	UToolMenu* const ViewportToolbarMenu = UToolMenus::Get()->RegisterMenu(
-		"LevelEditor.ViewportToolbar", NAME_None /* parent */, EMultiBoxType::ToolBar);
-
-	// Transforms section
-	{
-		FToolMenuSection& Section = ViewportToolbarMenu->FindOrAddSection(
-			"Transforms", LOCTEXT("TransformsSectionLabel", "Transforms"));
-
-		Section.AddEntry(FToolMenuEntry::InitToolBarButton(FEditorViewportCommands::Get().SelectMode));
-		Section.AddEntry(FToolMenuEntry::InitToolBarButton(FEditorViewportCommands::Get().TranslateMode));
-		Section.AddEntry(FToolMenuEntry::InitToolBarButton(FEditorViewportCommands::Get().RotateMode));
-		Section.AddEntry(FToolMenuEntry::InitToolBarButton(FEditorViewportCommands::Get().ScaleMode));
-	}
-}
-
 TSharedPtr<SWidget> SLevelViewport::MakeViewportToolbar()
 {
 	const TSharedRef<SLevelViewportToolBar> OldViewportToolbar = SNew(SLevelViewportToolBar)
@@ -1933,8 +1917,31 @@ TSharedPtr<SWidget> SLevelViewport::MakeViewportToolbar()
 		})
 		.IsEnabled(FSlateApplication::Get().GetNormalExecutionAttribute());
 
-	BuildToolMenusViewportToolbar();
+	// Register the viewport toolbar if another viewport hasn't already (it's shared).
+	{
+		const FName LevelEditorViewportToolbarName = "LevelEditor.ViewportToolbar";
+
+		if (!UToolMenus::Get()->IsMenuRegistered(LevelEditorViewportToolbarName))
+		{
+			UToolMenu* const ViewportToolbarMenu = UToolMenus::Get()->RegisterMenu(
+				LevelEditorViewportToolbarName, NAME_None /* parent */, EMultiBoxType::ToolBar);
+
+			UE::LevelEditor::AddViewportToolbarTransformsSection(ViewportToolbarMenu);
+			UE::LevelEditor::AddLevelEditorViewportToolbarSettingsSection(ViewportToolbarMenu);
+		}
+	}
+
+	FToolMenuContext ViewportToolbarContext;
+	{
+		ViewportToolbarContext.AppendCommandList(GetCommandList());
+
+		ULevelViewportContext* const ContextObject = NewObject<ULevelViewportContext>();
+		ContextObject->LevelViewport = SharedThis(this);
+
+		ViewportToolbarContext.AddObject(ContextObject);
+	}
 
+	// clang-format off
 	const TSharedRef<SWidget> NewViewportToolbar = SNew(SBox)
 		.Visibility_Lambda([this]() -> EVisibility
 		{
@@ -1947,8 +1954,9 @@ TSharedPtr<SWidget> SLevelViewport::MakeViewportToolbar()
 			return GetToolBarVisibility();
 		})
 		[
-			UToolMenus::Get()->GenerateWidget("LevelEditor.ViewportToolbar", FToolMenuContext(GetCommandList()))
+			UToolMenus::Get()->GenerateWidget("LevelEditor.ViewportToolbar", ViewportToolbarContext)
 		];
+	// clang-format on
 
 	return 
 		SNew(SVerticalBox)
diff --git a/Engine/Source/Editor/LevelEditor/Private/ViewportToolbar/LevelEditorViewportToolbarSections.cpp b/Engine/Source/Editor/LevelEditor/Private/ViewportToolbar/LevelEditorViewportToolbarSections.cpp
new file mode 100644
--- /dev/null
+++ b/Engine/Source/Editor/LevelEditor/Private/ViewportToolbar/LevelEditorViewportToolbarSections.cpp
@@ -0,0 +1,78 @@
+// Copyright Epic Games, Inc. All Rights Reserved.
+
+#include "ViewportToolbar/LevelEditorViewportToolbarSections.h"
+
+#include "EditorViewportCommands.h"
+#include "LevelViewportActions.h"
+#include "LevelViewportContext.h"
+#include "SLevelViewport.h"
+#include "ToolMenu.h"
+
+#define LOCTEXT_NAMESPACE "LevelEditorViewportToolbar"
+
+namespace UE::LevelEditor
+{
+
+// TODO: Move this outside the level editor and make it publicly available to anyone building a viewport toolbar.
+void AddViewportToolbarTransformsSection(UToolMenu* InMenu)
+{
+	FToolMenuSection& Section = InMenu->FindOrAddSection("Transforms", LOCTEXT("TransformsSectionLabel", "Transforms"));
+
+	Section.AddEntry(FToolMenuEntry::InitToolBarButton(FEditorViewportCommands::Get().SelectMode));
+	Section.AddEntry(FToolMenuEntry::InitToolBarButton(FEditorViewportCommands::Get().TranslateMode));
+	Section.AddEntry(FToolMenuEntry::InitToolBarButton(FEditorViewportCommands::Get().RotateMode));
+	Section.AddEntry(FToolMenuEntry::InitToolBarButton(FEditorViewportCommands::Get().ScaleMode));
+}
+
+void AddLevelEditorViewportToolbarSettingsSection(UToolMenu* InMenu)
+{
+	FToolMenuSection& Section = InMenu->FindOrAddSection("Settings", LOCTEXT("SettingsSectionLabel", "Settings"));
+
+	// Add realtime rendering toggle.
+	Section.AddEntry(FToolMenuEntry::InitToolBarButton(FEditorViewportCommands::Get().ToggleRealTime));
+
+	// Add maximize/restore viewport button.
+	{
+		FToolUIAction MaximizeRestoreAction;
+		MaximizeRestoreAction.ExecuteAction = FToolMenuExecuteAction::CreateLambda([](const FToolMenuContext& Context) {
+			ULevelViewportContext* const LevelViewportContext = Context.FindContext<ULevelViewportContext>();
+			if (!LevelViewportContext)
+			{
+				return;
+			}
+
+			if (const TSharedPtr<SLevelViewport> LevelViewport = LevelViewportContext->LevelViewport.Pin())
+			{
+				LevelViewport->OnToggleMaximize();
+			}
+		});
+		MaximizeRestoreAction.GetActionCheckState = FToolMenuGetActionCheckState::CreateLambda(
+			[](const FToolMenuContext& Context) -> ECheckBoxState {
+				ULevelViewportContext* const LevelViewportContext = Context.FindContext<ULevelViewportContext>();
+				if (!LevelViewportContext)
+				{
+					return ECheckBoxState::Undetermined;
+				}
+
+				if (const TSharedPtr<SLevelViewport> LevelViewport = LevelViewportContext->LevelViewport.Pin())
+				{
+					return LevelViewport->IsMaximized() ? ECheckBoxState::Checked : ECheckBoxState::Unchecked;
+				}
+
+				return ECheckBoxState::Undetermined;
+			});
+
+		Section.AddEntry(FToolMenuEntry::InitToolBarButton("MaximizeRestore", MaximizeRestoreAction,
+			LOCTEXT("MaximizeRestoreLabel", "Maximize/restore"),
+			LOCTEXT("MaximizeRestoreTooltip", "Maximizes or restores this viewport"),
+			FSlateIcon(FAppStyle::GetAppStyleSetName(), "EditorViewportToolBar.Maximize.Normal"),
+			EUserInterfaceActionType::ToggleButton));
+	}
+
+	// Add immersive mode toggle.
+	Section.AddEntry(FToolMenuEntry::InitToolBarButton(FLevelViewportCommands::Get().ToggleImmersive));
+}
+
+} // namespace UE::LevelEditor
+
+#undef LOCTEXT_NAMESPACE
diff --git a/Engine/Source/Editor/LevelEditor/Private/ViewportToolbar/LevelEditorViewportToolbarSections.h b/Engine/Source/Editor/LevelEditor/Private/ViewportToolbar/LevelEditorViewportToolbarSections.h
new file mode 100644
--- /dev/null
+++ b/Engine/Source/Editor/LevelEditor/Private/ViewportToolbar/LevelEditorViewportToolbarSections.h
@@ -0,0 +1,13 @@
+// Copyright Epic Games, Inc. All Rights Reserved.
+
+#pragma once
+
+class UToolMenu;
+
+namespace UE::LevelEditor
+{
+
+void AddViewportToolbarTransformsSection(UToolMenu* InMenu);
+void AddLevelEditorViewportToolbarSettingsSection(UToolMenu* InMenu);
+
+} // namespace UE::LevelEditor
diff --git a/Engine/Source/Editor/LevelEditor/Private/ViewportToolbar/LevelViewportContext.h b/Engine/Source/Editor/LevelEditor/Private/ViewportToolbar/LevelViewportContext.h
new file mode 100644
--- /dev/null
+++ b/Engine/Source/Editor/LevelEditor/Private/ViewportToolbar/LevelViewportContext.h
@@ -0,0 +1,17 @@
+// Copyright Epic Games, Inc. All Rights Reserved.
+
+#pragma once
+
+#include "Templates/SharedPointer.h"
+#include "SLevelViewport.h"
+#include "UObject/Object.h"
+
+#include "LevelViewportContext.generated.h"
+
+UCLASS()
+class LEVELEDITOR_API ULevelViewportContext : public UObject
+{
+	GENERATED_BODY()
+public:
+	TWeakPtr<SLevelViewport> LevelViewport;
+};
