[Backout] - CL33201403
Original CL Desc
-----------------------------------------------------------------
Add "Settings" section to UToolMenu "LevelEditor.ViewportToolbar"

* Add buttons to the new section to toggle realtime rendering, maximize/restore the viewport, and toggle immersive mode.
* Add a context to the "LevelEditor.ViewportToolbar" when a widget is generated for it. This context carries the SLevelViewport so the maximize/restore button can act on the correct viewport (the same button is drawn in all level editor viewports).

diff --git a/Engine/Source/Editor/LevelEditor/Private/SLevelViewport.cpp b/Engine/Source/Editor/LevelEditor/Private/SLevelViewport.cpp
--- a/Engine/Source/Editor/LevelEditor/Private/SLevelViewport.cpp
+++ b/Engine/Source/Editor/LevelEditor/Private/SLevelViewport.cpp
@@ -102,8 +102,6 @@
 #include "LevelViewportLayout.h"
 #include "EditorViewportTabContent.h"
 #include "EditorViewportCommands.h"
-#include "ViewportToolbar/LevelEditorViewportToolbarSections.h"
-#include "ViewportToolbar/LevelViewportContext.h"
 
 static const FName LevelEditorName("LevelEditor");
 static FAutoConsoleCommand EnableInViewportMenu(TEXT("Editor.EnableInViewportMenu"), TEXT("Enables the new in-viewport property menu"), FConsoleCommandDelegate::CreateStatic(&SLevelViewport::EnableInViewportMenu));
@@ -1901,6 +1899,24 @@ TSharedRef<FEditorViewportClient> SLevelViewport::MakeEditorViewportClient()
 	return LevelViewportClient.ToSharedRef();
 }
 
+void BuildToolMenusViewportToolbar()
+{
+	// Add a ToolMenus-based viewport toolbar.
+	UToolMenu* const ViewportToolbarMenu = UToolMenus::Get()->RegisterMenu(
+		"LevelEditor.ViewportToolbar", NAME_None /* parent */, EMultiBoxType::ToolBar);
+
+	// Transforms section
+	{
+		FToolMenuSection& Section = ViewportToolbarMenu->FindOrAddSection(
+			"Transforms", LOCTEXT("TransformsSectionLabel", "Transforms"));
+
+		Section.AddEntry(FToolMenuEntry::InitToolBarButton(FEditorViewportCommands::Get().SelectMode));
+		Section.AddEntry(FToolMenuEntry::InitToolBarButton(FEditorViewportCommands::Get().TranslateMode));
+		Section.AddEntry(FToolMenuEntry::InitToolBarButton(FEditorViewportCommands::Get().RotateMode));
+		Section.AddEntry(FToolMenuEntry::InitToolBarButton(FEditorViewportCommands::Get().ScaleMode));
+	}
+}
+
 TSharedPtr<SWidget> SLevelViewport::MakeViewportToolbar()
 {
 	const TSharedRef<SLevelViewportToolBar> OldViewportToolbar = SNew(SLevelViewportToolBar)
@@ -1917,31 +1933,8 @@ TSharedPtr<SWidget> SLevelViewport::MakeViewportToolbar()
 		})
 		.IsEnabled(FSlateApplication::Get().GetNormalExecutionAttribute());
 
-	// Register the viewport toolbar if another viewport hasn't already (it's shared).
-	{
-		const FName LevelEditorViewportToolbarName = "LevelEditor.ViewportToolbar";
-
-		if (!UToolMenus::Get()->IsMenuRegistered(LevelEditorViewportToolbarName))
-		{
-			UToolMenu* const ViewportToolbarMenu = UToolMenus::Get()->RegisterMenu(
-				LevelEditorViewportToolbarName, NAME_None /* parent */, EMultiBoxType::ToolBar);
-
-			UE::LevelEditor::AddViewportToolbarTransformsSection(ViewportToolbarMenu);
-			UE::LevelEditor::AddLevelEditorViewportToolbarSettingsSection(ViewportToolbarMenu);
-		}
-	}
-
-	FToolMenuContext ViewportToolbarContext;
-	{
-		ViewportToolbarContext.AppendCommandList(GetCommandList());
-
-		ULevelViewportContext* const ContextObject = NewObject<ULevelViewportContext>();
-		ContextObject->LevelViewport = SharedThis(this);
-
-		ViewportToolbarContext.AddObject(ContextObject);
-	}
+	BuildToolMenusViewportToolbar();
 
-	// clang-format off
 	const TSharedRef<SWidget> NewViewportToolbar = SNew(SBox)
 		.Visibility_Lambda([this]() -> EVisibility
 		{
@@ -1954,9 +1947,8 @@ TSharedPtr<SWidget> SLevelViewport::MakeViewportToolbar()
 			return GetToolBarVisibility();
 		})
 		[
-			UToolMenus::Get()->GenerateWidget("LevelEditor.ViewportToolbar", ViewportToolbarContext)
+			UToolMenus::Get()->GenerateWidget("LevelEditor.ViewportToolbar", FToolMenuContext(GetCommandList()))
 		];
-	// clang-format on
 
 	return 
 		SNew(SVerticalBox)
diff --git a/Engine/Source/Editor/LevelEditor/Private/ViewportToolbar/LevelEditorViewportToolbarSections.cpp b/Engine/Source/Editor/LevelEditor/Private/ViewportToolbar/LevelEditorViewportToolbarSections.cpp
deleted file mode 100644
diff --git a/Engine/Source/Editor/LevelEditor/Private/ViewportToolbar/LevelEditorViewportToolbarSections.h b/Engine/Source/Editor/LevelEditor/Private/ViewportToolbar/LevelEditorViewportToolbarSections.h
deleted file mode 100644
diff --git a/Engine/Source/Editor/LevelEditor/Private/ViewportToolbar/LevelViewportContext.h b/Engine/Source/Editor/LevelEditor/Private/ViewportToolbar/LevelViewportContext.h
deleted file mode 100644
