Remove unnecessary private include path to SceneOutliner

#rnx

diff --git a/Engine/Source/Editor/LevelEditor/LevelEditor.Build.cs b/Engine/Source/Editor/LevelEditor/LevelEditor.Build.cs
--- a/Engine/Source/Editor/LevelEditor/LevelEditor.Build.cs
+++ b/Engine/Source/Editor/LevelEditor/LevelEditor.Build.cs
@@ -6,12 +6,6 @@ public class LevelEditor : ModuleRules
 {
 	public LevelEditor(ReadOnlyTargetRules Target) : base(Target)
 	{
-		PrivateIncludePaths.AddRange(
-			new string[] {
-				System.IO.Path.Combine(GetModuleDirectory("SceneOutliner"), "Private"),
-			}
-		);
-
 		PrivateIncludePathModuleNames.AddRange(
 			new string[] {
 				"AssetTools",
