Add some additional logging about AActor flags when TypedElementLevelEditorSelectionCustomization fails to deselect an actor due to unreachability
* Will log flags of both level and actor if the Level's state considers an actor unreachable.  This can cause deselection to fail.

diff --git a/Engine/Source/Editor/LevelEditor/Private/Elements/Actor/ActorElementLevelEditorSelectionCustomization.cpp b/Engine/Source/Editor/LevelEditor/Private/Elements/Actor/ActorElementLevelEditorSelectionCustomization.cpp
--- a/Engine/Source/Editor/LevelEditor/Private/Elements/Actor/ActorElementLevelEditorSelectionCustomization.cpp
+++ b/Engine/Source/Editor/LevelEditor/Private/Elements/Actor/ActorElementLevelEditorSelectionCustomization.cpp
@@ -89,18 +89,25 @@ bool IsActorReachable(const AActor* Actor)
 
 	// Ensure that neither the level nor the actor is being destroyed or is unreachable
 	const EObjectFlags InvalidSelectableFlags = RF_BeginDestroyed;
+	bool ReturnValue = true;
 	if (Actor->GetLevel()->HasAnyFlags(InvalidSelectableFlags) || (!GIsTransacting && (!IsValidChecked(Actor->GetLevel()) || Actor->GetLevel()->IsUnreachable())))
 	{
-		UE_LOG(LogActorLevelEditorSelection, Warning, TEXT("SelectActor: %s (%s)"), TEXT("The requested operation could not be completed because the level has invalid flags."), *Actor->GetActorLabel());
-		return false;
+		UE_LOG(LogActorLevelEditorSelection, Warning, TEXT("SelectActor: %s (%s) Flags: (0x%08X)"),
+			TEXT("The requested operation could not be completed because the level has invalid flags."),
+			*Actor->GetActorLabel(),
+			static_cast<std::underlying_type_t<EInternalObjectFlags>>(Actor->GetInternalFlags()));
+		ReturnValue = false;
 	}
 	if (Actor->HasAnyFlags(InvalidSelectableFlags) || (!GIsTransacting && (!IsValidChecked(Actor) || Actor->IsUnreachable())))
 	{
-		UE_LOG(LogActorLevelEditorSelection, Warning, TEXT("SelectActor: %s (%s)"), TEXT("The requested operation could not be completed because the actor has invalid flags."), *Actor->GetActorLabel());
-		return false;
+		UE_LOG(LogActorLevelEditorSelection, Warning, TEXT("SelectActor: %s (%s) Flags: (0x%08X)"),
+			TEXT("The requested operation could not be completed because the actor has invalid flags."),
+			*Actor->GetActorLabel(),
+			static_cast<std::underlying_type_t<EInternalObjectFlags>>(Actor->GetInternalFlags()));
+		ReturnValue = false;
 	}
 
-	return true;
+	return ReturnValue;
 }
 
 }
