Fix crash occurring when Editor.EnableInViewportMenu is set and selecting and tabbing quickly for an object containing properties shown in viewport details

diff --git a/Engine/Source/Editor/LevelEditor/Private/SInViewportDetails.cpp b/Engine/Source/Editor/LevelEditor/Private/SInViewportDetails.cpp
--- a/Engine/Source/Editor/LevelEditor/Private/SInViewportDetails.cpp
+++ b/Engine/Source/Editor/LevelEditor/Private/SInViewportDetails.cpp
@@ -337,10 +337,18 @@ void SInViewportDetails::SetObjects(const TArray<UObject*>& InObjects, bool bFor
 		GenerateWidget();
 		if (!Nodes.Num())
 		{
-			// Do this on a delay so that we are are not caught in a loop of creating and hiding the menu
-			GEditor->GetTimerManager()->SetTimerForNextTick([this]()
+			TWeakPtr<SInViewportDetails> ThisWeak = SharedThis(this);
+
+			// Do this on a delay so that we are not caught in a loop of creating and hiding the menu
+			GEditor->GetTimerManager()->SetTimerForNextTick([ThisWeak]()
 				{
-					OwningViewport.Pin()->HideInViewportContextMenu();
+					if (TSharedPtr<SInViewportDetails> This = ThisWeak.Pin())
+					{
+						if (TSharedPtr<SEditorViewport> Viewport = This->OwningViewport.Pin())
+						{
+							Viewport->HideInViewportContextMenu();
+						}
+					}
 				});
 		}
 	}
