Use a section menu for the "Settings" section of the "LevelEditor.ViewportToolbar" menu add "Material Quality Level" and "Preview Platform" submenus

Note that this updates the level editor (non-viewport) toolbar "Settings" menu to use the same code for its "Material Quality Level" and "Preview Platform" submenus as is now used in the "LevelEditor.ViewportToolbar" menu.

diff --git a/Engine/Source/Editor/LevelEditor/Private/LevelEditorToolBar.cpp b/Engine/Source/Editor/LevelEditor/Private/LevelEditorToolBar.cpp
--- a/Engine/Source/Editor/LevelEditor/Private/LevelEditorToolBar.cpp
+++ b/Engine/Source/Editor/LevelEditor/Private/LevelEditorToolBar.cpp
@@ -63,6 +63,7 @@
 #include "Styling/ToolBarStyle.h"
 #include "PlatformInfo.h"
 #include "DataDrivenShaderPlatformInfo.h"
+#include "ViewportToolbar/LevelEditorViewportToolbarSections.h"
 
 FName FLevelEditorToolBar::SecondaryModeToolbarName("LevelEditor.SecondaryToolbar");
 
@@ -1650,32 +1651,6 @@ TSharedRef< SWidget > FLevelEditorToolBar::MakeLevelEditorToolBar( const TShared
 		];
 }
 
-static void MakeMaterialQualityLevelMenu( UToolMenu* InMenu )
-{
-	{
-		FToolMenuSection& Section = InMenu->AddSection("LevelEditorMaterialQualityLevel", NSLOCTEXT( "LevelToolBarViewMenu", "MaterialQualityLevelHeading", "Material Quality Level" ) );
-		Section.AddMenuEntry(FLevelEditorCommands::Get().MaterialQualityLevel_Low);
-		Section.AddMenuEntry(FLevelEditorCommands::Get().MaterialQualityLevel_Medium);
-		Section.AddMenuEntry(FLevelEditorCommands::Get().MaterialQualityLevel_High);
-		Section.AddMenuEntry(FLevelEditorCommands::Get().MaterialQualityLevel_Epic);
-	}
-}
-
-static void MakeShaderModelPreviewMenu( UToolMenu* InMenu )
-{
-#define LOCTEXT_NAMESPACE "LevelToolBarViewMenu"
-
-	FToolMenuSection& Section = InMenu->AddSection("EditorPreviewMode", LOCTEXT("EditorPreviewModeDevices", "Preview Devices"));
-
-	// Preview platforms discovered from ITargetPlatforms.
-	for (auto& Item : FLevelEditorCommands::Get().PreviewPlatformOverrides)
-	{
-		Section.AddMenuEntry(Item);
-	}
-
-#undef LOCTEXT_NAMESPACE
-}
-
 static void MakeScalabilityMenu( UToolMenu* InMenu )
 {
 	{
@@ -1774,17 +1749,8 @@ void FLevelEditorToolBar::RegisterQuickSettingsMenu()
 			LOCTEXT( "ScalabilitySubMenu_ToolTip", "Open the engine scalability settings" ),
 			FNewToolMenuDelegate::CreateStatic( &MakeScalabilityMenu ) );
 
-		Section.AddSubMenu(
-			"MaterialQualityLevel",
-			LOCTEXT( "MaterialQualityLevelSubMenu", "Material Quality Level" ),
-			LOCTEXT( "MaterialQualityLevelSubMenu_ToolTip", "Sets the value of the CVar \"r.MaterialQualityLevel\" (low=0, high=1, medium=2, Epic=3). This affects materials via the QualitySwitch material expression." ),
-			FNewToolMenuDelegate::CreateStatic( &MakeMaterialQualityLevelMenu ) );
-
-		Section.AddSubMenu(
-			"FeatureLevelPreview",
-			LOCTEXT("PreviewPlatformSubMenu", "Preview Platform"),
-			LOCTEXT("PreviewPlatformSubMenu_ToolTip", "Sets the preview platform used by the main editor"),
-			FNewToolMenuDelegate::CreateStatic(&MakeShaderModelPreviewMenu));
+		UE::LevelEditor::AddMaterialQualityLevelSubmenu(Section);
+		UE::LevelEditor::AddFeatureLevelPreviewSubmenu(Section);
 	}
 
 	{
diff --git a/Engine/Source/Editor/LevelEditor/Private/ViewportToolbar/LevelEditorViewportToolbarSections.cpp b/Engine/Source/Editor/LevelEditor/Private/ViewportToolbar/LevelEditorViewportToolbarSections.cpp
--- a/Engine/Source/Editor/LevelEditor/Private/ViewportToolbar/LevelEditorViewportToolbarSections.cpp
+++ b/Engine/Source/Editor/LevelEditor/Private/ViewportToolbar/LevelEditorViewportToolbarSections.cpp
@@ -3,6 +3,7 @@
 #include "ViewportToolbar/LevelEditorViewportToolbarSections.h"
 
 #include "EditorViewportCommands.h"
+#include "LevelEditorActions.h"
 #include "LevelViewportActions.h"
 #include "LevelViewportContext.h"
 #include "SLevelViewport.h"
@@ -24,13 +25,52 @@ void AddViewportToolbarTransformsSection(UToolMenu* InMenu)
 	Section.AddEntry(FToolMenuEntry::InitToolBarButton(FEditorViewportCommands::Get().ScaleMode));
 }
 
+void AddMaterialQualityLevelSubmenu(FToolMenuSection& Section)
+{
+	Section.AddSubMenu("MaterialQualityLevel",
+		NSLOCTEXT("LevelToolBarViewMenu", "MaterialQualityLevelSubMenu", "Material Quality Level"),
+		NSLOCTEXT("LevelToolBarViewMenu",
+			"MaterialQualityLevelSubMenu_ToolTip",
+			"Sets the value of the CVar \"r.MaterialQualityLevel\" (low=0, high=1, medium=2, Epic=3). This affects "
+			"materials via the QualitySwitch material expression."),
+		FNewToolMenuDelegate::CreateLambda([](UToolMenu* InMenu) -> void {
+			FToolMenuSection& Section = InMenu->AddSection("LevelEditorMaterialQualityLevel",
+				NSLOCTEXT("LevelToolBarViewMenu", "MaterialQualityLevelHeading", "Material Quality Level"));
+			Section.AddMenuEntry(FLevelEditorCommands::Get().MaterialQualityLevel_Low);
+			Section.AddMenuEntry(FLevelEditorCommands::Get().MaterialQualityLevel_Medium);
+			Section.AddMenuEntry(FLevelEditorCommands::Get().MaterialQualityLevel_High);
+			Section.AddMenuEntry(FLevelEditorCommands::Get().MaterialQualityLevel_Epic);
+		}));
+}
+
+void AddFeatureLevelPreviewSubmenu(FToolMenuSection& Section)
+{
+	Section.AddSubMenu("FeatureLevelPreview",
+		NSLOCTEXT("LevelToolBarViewMenu", "PreviewPlatformSubMenu", "Preview Platform"),
+		NSLOCTEXT(
+			"LevelToolBarViewMenu", "PreviewPlatformSubMenu_ToolTip", "Sets the preview platform used by the main editor"),
+		FNewToolMenuDelegate::CreateLambda([](UToolMenu* InMenu) -> void {
+			FToolMenuSection& Section = InMenu->AddSection(
+				"EditorPreviewMode", LOCTEXT("EditorPreviewModeDevices", "Preview Devices"));
+			// Preview platforms discovered from ITargetPlatforms.
+			for (auto& Item : FLevelEditorCommands::Get().PreviewPlatformOverrides)
+			{
+				Section.AddMenuEntry(Item);
+			}
+		}));
+}
+
 void AddLevelEditorViewportToolbarSettingsSection(UToolMenu* InMenu)
 {
 	FToolMenuSection& Section = InMenu->FindOrAddSection("Settings", LOCTEXT("SettingsSectionLabel", "Settings"));
+	Section.SetShowSectionMenu(true);
 
 	// Add realtime rendering toggle.
 	Section.AddEntry(FToolMenuEntry::InitToolBarButton(FEditorViewportCommands::Get().ToggleRealTime));
 
+	AddMaterialQualityLevelSubmenu(Section);
+	AddFeatureLevelPreviewSubmenu(Section);
+
 	// Add maximize/restore viewport button.
 	{
 		FToolUIAction MaximizeRestoreAction;
diff --git a/Engine/Source/Editor/LevelEditor/Private/ViewportToolbar/LevelEditorViewportToolbarSections.h b/Engine/Source/Editor/LevelEditor/Private/ViewportToolbar/LevelEditorViewportToolbarSections.h
--- a/Engine/Source/Editor/LevelEditor/Private/ViewportToolbar/LevelEditorViewportToolbarSections.h
+++ b/Engine/Source/Editor/LevelEditor/Private/ViewportToolbar/LevelEditorViewportToolbarSections.h
@@ -2,12 +2,15 @@
 
 #pragma once
 
+struct FToolMenuSection;
 class UToolMenu;
 
 namespace UE::LevelEditor
 {
 
 void AddViewportToolbarTransformsSection(UToolMenu* InMenu);
+void AddMaterialQualityLevelSubmenu(FToolMenuSection& Section);
+void AddFeatureLevelPreviewSubmenu(FToolMenuSection& Section);
 void AddLevelEditorViewportToolbarSettingsSection(UToolMenu* InMenu);
 
 } // namespace UE::LevelEditor
