Flag the "realtime toggle" as top-level in the new viewport toolbar

diff --git a/Engine/Source/Editor/LevelEditor/Private/ViewportToolbar/LevelEditorViewportToolbarSections.cpp b/Engine/Source/Editor/LevelEditor/Private/ViewportToolbar/LevelEditorViewportToolbarSections.cpp
--- a/Engine/Source/Editor/LevelEditor/Private/ViewportToolbar/LevelEditorViewportToolbarSections.cpp
+++ b/Engine/Source/Editor/LevelEditor/Private/ViewportToolbar/LevelEditorViewportToolbarSections.cpp
@@ -66,7 +66,7 @@ void AddLevelEditorViewportToolbarSettingsSection(UToolMenu* InMenu)
 	Section.SetShowSectionMenu(true);
 
 	// Add realtime rendering toggle.
-	Section.AddEntry(FToolMenuEntry::InitToolBarButton(FEditorViewportCommands::Get().ToggleRealTime));
+	Section.AddMenuEntry(FEditorViewportCommands::Get().ToggleRealTime).SetShowInToolbarTopLevel(true);
 
 	AddMaterialQualityLevelSubmenu(Section);
 	AddFeatureLevelPreviewSubmenu(Section);
