Use ToolMenu submenus instead of sections for grouping the new viewport toolbar functionality

The upside of this is that the submenus of each group ("section submenus") can have proper ToolMenu sections if they are ToolMenu submenus. Previously, when we grouped using sections, we couldn't section each section, because ToolMenus doesn't support nested sections.

diff --git a/Engine/Source/Developer/ToolMenus/Private/ToolMenuSection.cpp b/Engine/Source/Developer/ToolMenus/Private/ToolMenuSection.cpp
--- a/Engine/Source/Developer/ToolMenus/Private/ToolMenuSection.cpp
+++ b/Engine/Source/Developer/ToolMenus/Private/ToolMenuSection.cpp
@@ -16,7 +16,6 @@ FToolMenuSection::FToolMenuSection()
 	: ToolMenuSectionDynamic(nullptr)
 	, bIsRegistering(false)
 	, bAddedDuringRegister(false)
-	, bShowSectionMenu(false)
 {
 }
 
@@ -34,7 +33,6 @@ void FToolMenuSection::InitGeneratedSectionCopy(const FToolMenuSection& Source,
 	InsertPosition = Source.InsertPosition;
 	Construct = Source.Construct;
 	Context = InContext;
-	bShowSectionMenu = Source.bShowSectionMenu;
 }
 
 bool FToolMenuSection::IsRegistering() const
@@ -165,11 +163,6 @@ const FToolMenuEntry* FToolMenuSection::FindEntry(const FName InName) const
 	return nullptr;
 }
 
-void FToolMenuSection::SetShowSectionMenu(bool InShow)
-{
-	bShowSectionMenu = InShow;
-}
-
 int32 FToolMenuSection::IndexOfBlock(const FName InName) const
 {
 	for (int32 i=0; i < Blocks.Num(); ++i)
diff --git a/Engine/Source/Developer/ToolMenus/Private/ToolMenus.cpp b/Engine/Source/Developer/ToolMenus/Private/ToolMenus.cpp
--- a/Engine/Source/Developer/ToolMenus/Private/ToolMenus.cpp
+++ b/Engine/Source/Developer/ToolMenus/Private/ToolMenus.cpp
@@ -1445,7 +1445,7 @@ void UToolMenus::ExtractChildBlocksFromSubMenu(
 	}
 }
 
-void UToolMenus::PopulateToolBarBuilderWithGrandchildren(
+void UToolMenus::PopulateToolBarBuilderWithTopLevelChildren(
 	FToolBarBuilder& ToolBarBuilder, UToolMenu* ParentMenu, FToolMenuEntry& InBlock)
 {
 	TArray<TPair<UToolMenu*, FToolMenuEntry*>> SubMenuBlockPairs;
@@ -1468,18 +1468,16 @@ void UToolMenus::PopulateToolBarBuilderWithGrandchildren(
 			break;
 		}
 
-		// We pass in true for bSectionHasSubMenu here because we're referring to the top-level section that
-		// called us and not the section we're currently iterating through. It is that top-level section
-		// that we're populating here, not the section the current Block came from.
-		constexpr bool bSectionHasSubMenu = true;
-		PopulateToolBarBuilderWithEntry(ToolBarBuilder, SubMenu, bSectionHasSubMenu, *Block);
+		if (Block->bShowInToolbarTopLevel)
+		{
+			PopulateToolBarBuilderWithEntry(ToolBarBuilder, SubMenu, *Block);
+		}
 
 		ExtractChildBlocksFromSubMenu(SubMenu, *Block, SubMenuBlockPairs);
 	}
 }
 
-void UToolMenus::PopulateToolBarBuilderWithEntry(
-	FToolBarBuilder& ToolBarBuilder, UToolMenu* MenuData, bool bSectionHasSubMenu, FToolMenuEntry& Block)
+void UToolMenus::PopulateToolBarBuilderWithEntry(FToolBarBuilder& ToolBarBuilder, UToolMenu* MenuData, FToolMenuEntry& Block)
 {
 	if (Block.ToolBarData.ConstructLegacy.IsBound())
 	{
@@ -1487,12 +1485,6 @@ void UToolMenus::PopulateToolBarBuilderWithEntry(
 		return;
 	}
 
-	const bool bAddEntry = !bSectionHasSubMenu || (bSectionHasSubMenu && Block.bShowInToolbarTopLevel);
-	if (!bAddEntry)
-	{
-		return;
-	}
-
 	ToolBarBuilder.BeginStyleOverride(Block.StyleNameOverride);
 
 	if (Block.Type == EMultiBlockType::ToolBarButton || (Block.Type == EMultiBlockType::MenuEntry && !Block.IsSubMenu()))
@@ -1564,6 +1556,9 @@ void UToolMenus::PopulateToolBarBuilderWithEntry(
 			FUIAction UIAction = ConvertUIAction(Block, MenuData->Context);
 			ToolBarBuilder.AddComboButton(UIAction, Delegate, Block.Label, Block.ToolTip, Block.Icon,
 				Block.ToolBarData.bSimpleComboBox, Block.TutorialHighlightName);
+
+			// Also add any top-level flagged children to the toolbar.
+			PopulateToolBarBuilderWithTopLevelChildren(ToolBarBuilder, MenuData, Block);
 		}
 	}
 	else if (Block.Type == EMultiBlockType::Separator)
@@ -1634,59 +1629,9 @@ void UToolMenus::PopulateToolBarBuilder(FToolBarBuilder& ToolBarBuilder, UToolMe
 
 		ToolBarBuilder.BeginSection(Section.Name);
 
-		if (Section.bShowSectionMenu)
-		{
-			TWeakObjectPtr<UToolMenu> WeakMenuPtr = MenuData;
-			const FOnGetContent MenuGenerator = FOnGetContent::CreateLambda(
-				[this, WeakMenuPtr, SectionName = Section.Name]() -> TSharedRef<SWidget> {
-					TStrongObjectPtr<UToolMenu> Menu = WeakMenuPtr.Pin();
-					// If we got no menu, there's nothing to do.
-					if (!Menu)
-					{
-						return SNullWidget::NullWidget;
-					}
-
-					FToolMenuSection* const Section = Menu->FindSection(SectionName);
-					// If our section doesn't exist anymore, we have nothing to do.
-					if (!Section)
-					{
-						return SNullWidget::NullWidget;
-					}
-
-					// Most of the following code came from UToolMenus::PopulateMenuBuilder.
-
-					FMenuBuilder MenuBuilder(true, nullptr);
-					MenuBuilder.SetSearchable(Menu->bSearchable);
-
-					// Don't pass in a label here because the label is already displayed as the combo button text
-					// and the menu will only have this single section.
-					MenuBuilder.BeginSection(Section->Name);
-
-					for (FToolMenuEntry& Block : Section->Blocks)
-					{
-						FPopulateMenuBuilderWithToolMenuEntry PopulateMenuBuilderWithToolMenuEntry(
-							MenuBuilder, Menu.Get(), *Section, Block, /* bAllowSubMenuCollapse= */ true);
-						PopulateMenuBuilderWithToolMenuEntry.Populate();
-					}
-
-					MenuBuilder.EndSection();
-
-					AddReferencedContextObjects(MenuBuilder.GetMultiBox(), Menu.Get());
-
-					return MenuBuilder.MakeWidget();
-				});
-
-			ToolBarBuilder.AddComboButton(FUIAction(), MenuGenerator, Section.Label,
-				LOCTEXT("ToolbarSectionMenuTooltip", "Menu that contains all entries of this section"));
-		}
-
 		for (FToolMenuEntry& Block : Section.Blocks)
 		{
-			PopulateToolBarBuilderWithEntry(ToolBarBuilder, MenuData, Section.bShowSectionMenu, Block);
-			if (Section.bShowSectionMenu)
-			{
-				PopulateToolBarBuilderWithGrandchildren(ToolBarBuilder, MenuData, Block);
-			}
+			PopulateToolBarBuilderWithEntry(ToolBarBuilder, MenuData, Block);
 		}
 
 		ToolBarBuilder.EndSection();
diff --git a/Engine/Source/Developer/ToolMenus/Public/ToolMenuSection.h b/Engine/Source/Developer/ToolMenus/Public/ToolMenuSection.h
--- a/Engine/Source/Developer/ToolMenus/Public/ToolMenuSection.h
+++ b/Engine/Source/Developer/ToolMenus/Public/ToolMenuSection.h
@@ -85,19 +85,6 @@ public:
 	FToolMenuEntry* FindEntry(const FName InName);
 	const FToolMenuEntry* FindEntry(const FName InName) const;
 
-	/**
-	 * Adds a menu to the section that contain all entries of the section.
-	 *
-	 * The use case of this section menu is to promote the most important entries to show in the toolbar section, while
-	 * the rest of the section's entries are available in the section menu (including the ones promoted to the toolbar).
-	 *
-	 * This feature only applies to sections added to toolbars, not sections added to menus.
-	 * If SetShowSectionMenu(true) is called on a section in a menu, nothing will happen.
-	 *
-	 * @param InShow Shows a section menu if true and doesn't if false (default).
-	 */
-	void SetShowSectionMenu(bool InShow);
-
 private:
 
 	void InitGeneratedSectionCopy(const FToolMenuSection& Source, FToolMenuContext& InContext);
@@ -144,5 +131,4 @@ private:
 
 	bool bIsRegistering;
 	bool bAddedDuringRegister;
-	bool bShowSectionMenu;
 };
diff --git a/Engine/Source/Developer/ToolMenus/Public/ToolMenus.h b/Engine/Source/Developer/ToolMenus/Public/ToolMenus.h
--- a/Engine/Source/Developer/ToolMenus/Public/ToolMenus.h
+++ b/Engine/Source/Developer/ToolMenus/Public/ToolMenus.h
@@ -388,12 +388,9 @@ private:
 
 	void ExtractChildBlocksFromSubMenu(
 		UToolMenu* ParentMenu, FToolMenuEntry& InBlock, TArray<TPair<UToolMenu*, FToolMenuEntry*>>& SubMenuBlockPairs);
-	void PopulateToolBarBuilderWithGrandchildren(
+	void PopulateToolBarBuilderWithTopLevelChildren(
 		FToolBarBuilder& ToolBarBuilder, UToolMenu* ParentMenu, FToolMenuEntry& InBlock);
-	void PopulateToolBarBuilderWithEntry(FToolBarBuilder& PopulateToolBarBuilderWithEntry,
-		UToolMenu* MenuData,
-		bool bSectionHasSubMenu,
-		FToolMenuEntry& Entry);
+	void PopulateToolBarBuilderWithEntry(FToolBarBuilder& ToolBarBuilder, UToolMenu* MenuData, FToolMenuEntry& Entry);
 	void PopulateToolBarBuilder(FToolBarBuilder& ToolBarBuilder, UToolMenu* MenuData);
 
 	TSharedRef<SWidget> GenerateToolbarComboButtonMenu(TWeakObjectPtr<UToolMenu> InParent, const FName InBlockName);
diff --git a/Engine/Source/Editor/LevelEditor/Private/ViewportToolbar/LevelEditorViewportToolbarSections.cpp b/Engine/Source/Editor/LevelEditor/Private/ViewportToolbar/LevelEditorViewportToolbarSections.cpp
--- a/Engine/Source/Editor/LevelEditor/Private/ViewportToolbar/LevelEditorViewportToolbarSections.cpp
+++ b/Engine/Source/Editor/LevelEditor/Private/ViewportToolbar/LevelEditorViewportToolbarSections.cpp
@@ -17,12 +17,33 @@ namespace UE::LevelEditor
 // TODO: Move this outside the level editor and make it publicly available to anyone building a viewport toolbar.
 void AddViewportToolbarTransformsSection(UToolMenu* InMenu)
 {
-	FToolMenuSection& Section = InMenu->FindOrAddSection("Transforms", LOCTEXT("TransformsSectionLabel", "Transforms"));
-
-	Section.AddEntry(FToolMenuEntry::InitToolBarButton(FEditorViewportCommands::Get().SelectMode));
-	Section.AddEntry(FToolMenuEntry::InitToolBarButton(FEditorViewportCommands::Get().TranslateMode));
-	Section.AddEntry(FToolMenuEntry::InitToolBarButton(FEditorViewportCommands::Get().RotateMode));
-	Section.AddEntry(FToolMenuEntry::InitToolBarButton(FEditorViewportCommands::Get().ScaleMode));
+	FToolMenuSection& Section = InMenu->FindOrAddSection("Left");
+
+	Section.AddSubMenu("Transforms", LOCTEXT("TransformsSubmenuLabel", "Transforms"),
+		LOCTEXT("TransformsSubmenuTooltip", "Viewport-related transforms tools"),
+		FNewToolMenuDelegate::CreateLambda([](UToolMenu* Submenu) -> void {
+			FToolMenuSection& Section = Submenu->FindOrAddSection(NAME_None);
+
+			FToolMenuEntry SelectMode = FToolMenuEntry::InitMenuEntry(FEditorViewportCommands::Get().SelectMode);
+			SelectMode.UserInterfaceActionType = EUserInterfaceActionType::RadioButton;
+			SelectMode.SetShowInToolbarTopLevel(true);
+			Section.AddEntry(SelectMode);
+
+			FToolMenuEntry TranslateMode = FToolMenuEntry::InitMenuEntry(FEditorViewportCommands::Get().TranslateMode);
+			TranslateMode.UserInterfaceActionType = EUserInterfaceActionType::RadioButton;
+			TranslateMode.SetShowInToolbarTopLevel(true);
+			Section.AddEntry(TranslateMode);
+
+			FToolMenuEntry RotateMode = FToolMenuEntry::InitMenuEntry(FEditorViewportCommands::Get().RotateMode);
+			RotateMode.UserInterfaceActionType = EUserInterfaceActionType::RadioButton;
+			RotateMode.SetShowInToolbarTopLevel(true);
+			Section.AddEntry(RotateMode);
+
+			FToolMenuEntry ScaleMode = FToolMenuEntry::InitMenuEntry(FEditorViewportCommands::Get().ScaleMode);
+			ScaleMode.UserInterfaceActionType = EUserInterfaceActionType::RadioButton;
+			ScaleMode.SetShowInToolbarTopLevel(true);
+			Section.AddEntry(ScaleMode);
+		}));
 }
 
 void AddMaterialQualityLevelSubmenu(FToolMenuSection& Section)
@@ -62,56 +83,62 @@ void AddFeatureLevelPreviewSubmenu(FToolMenuSection& Section)
 
 void AddLevelEditorViewportToolbarSettingsSection(UToolMenu* InMenu)
 {
-	FToolMenuSection& Section = InMenu->FindOrAddSection("Settings", LOCTEXT("SettingsSectionLabel", "Settings"));
-	Section.SetShowSectionMenu(true);
+	FToolMenuSection& RightSection = InMenu->FindOrAddSection("Right");
 
-	// Add realtime rendering toggle.
-	Section.AddMenuEntry(FEditorViewportCommands::Get().ToggleRealTime).SetShowInToolbarTopLevel(true);
+	RightSection.AddSubMenu("Settings", LOCTEXT("SettingsSubmenuLabel", "Settings"),
+		LOCTEXT("SettingsSubmenuTooltip", "Viewport-related settings"),
+		FNewToolMenuDelegate::CreateLambda([](UToolMenu* Submenu) -> void {
+			FToolMenuSection& UnnamedSection = Submenu->FindOrAddSection(NAME_None);
 
-	AddMaterialQualityLevelSubmenu(Section);
-	AddFeatureLevelPreviewSubmenu(Section);
+			// Add realtime rendering toggle.
+			UnnamedSection.AddMenuEntry(FEditorViewportCommands::Get().ToggleRealTime).SetShowInToolbarTopLevel(true);
 
-	// Add maximize/restore viewport button.
-	{
-		FToolUIAction MaximizeRestoreAction;
-		MaximizeRestoreAction.ExecuteAction = FToolMenuExecuteAction::CreateLambda([](const FToolMenuContext& Context) {
-			ULevelViewportContext* const LevelViewportContext = Context.FindContext<ULevelViewportContext>();
-			if (!LevelViewportContext)
-			{
-				return;
-			}
+			AddMaterialQualityLevelSubmenu(UnnamedSection);
+			AddFeatureLevelPreviewSubmenu(UnnamedSection);
 
-			if (const TSharedPtr<::SLevelViewport> LevelViewport = LevelViewportContext->LevelViewport.Pin())
+			// Add maximize/restore viewport button.
 			{
-				LevelViewport->OnToggleMaximize();
+				FToolUIAction MaximizeRestoreAction;
+				MaximizeRestoreAction.ExecuteAction = FToolMenuExecuteAction::CreateLambda(
+					[](const FToolMenuContext& Context) {
+						ULevelViewportContext* const LevelViewportContext = Context.FindContext<ULevelViewportContext>();
+						if (!LevelViewportContext)
+						{
+							return;
+						}
+
+						if (const TSharedPtr<::SLevelViewport> LevelViewport = LevelViewportContext->LevelViewport.Pin())
+						{
+							LevelViewport->OnToggleMaximize();
+						}
+					});
+				MaximizeRestoreAction.GetActionCheckState = FToolMenuGetActionCheckState::CreateLambda(
+					[](const FToolMenuContext& Context) -> ECheckBoxState {
+						ULevelViewportContext* const LevelViewportContext = Context.FindContext<ULevelViewportContext>();
+						if (!LevelViewportContext)
+						{
+							return ECheckBoxState::Undetermined;
+						}
+
+						if (const TSharedPtr<::SLevelViewport> LevelViewport = LevelViewportContext->LevelViewport.Pin())
+						{
+							return LevelViewport->IsMaximized() ? ECheckBoxState::Checked : ECheckBoxState::Unchecked;
+						}
+
+						return ECheckBoxState::Undetermined;
+					});
+
+				UnnamedSection.AddEntry(FToolMenuEntry::InitToolBarButton("MaximizeRestore", MaximizeRestoreAction,
+					LOCTEXT("MaximizeRestoreLabel", "Maximize/restore"),
+					LOCTEXT("MaximizeRestoreTooltip", "Maximizes or restores this viewport"),
+					FSlateIcon(FAppStyle::GetAppStyleSetName(), "EditorViewportToolBar.Maximize.Normal"),
+					EUserInterfaceActionType::ToggleButton));
 			}
-		});
-		MaximizeRestoreAction.GetActionCheckState = FToolMenuGetActionCheckState::CreateLambda(
-			[](const FToolMenuContext& Context) -> ECheckBoxState {
-				ULevelViewportContext* const LevelViewportContext = Context.FindContext<ULevelViewportContext>();
-				if (!LevelViewportContext)
-				{
-					return ECheckBoxState::Undetermined;
-				}
-
-				if (const TSharedPtr<::SLevelViewport> LevelViewport = LevelViewportContext->LevelViewport.Pin())
-				{
-					return LevelViewport->IsMaximized() ? ECheckBoxState::Checked : ECheckBoxState::Unchecked;
-				}
-
-				return ECheckBoxState::Undetermined;
-			});
-
-		Section.AddEntry(FToolMenuEntry::InitToolBarButton("MaximizeRestore", MaximizeRestoreAction,
-			LOCTEXT("MaximizeRestoreLabel", "Maximize/restore"),
-			LOCTEXT("MaximizeRestoreTooltip", "Maximizes or restores this viewport"),
-			FSlateIcon(FAppStyle::GetAppStyleSetName(), "EditorViewportToolBar.Maximize.Normal"),
-			EUserInterfaceActionType::ToggleButton));
-	}
 
-	// Add immersive mode toggle.
-	Section.AddEntry(FToolMenuEntry::InitToolBarButton(FLevelViewportCommands::Get().ToggleImmersive));
-}
+			// Add immersive mode toggle.
+			UnnamedSection.AddEntry(FToolMenuEntry::InitToolBarButton(FLevelViewportCommands::Get().ToggleImmersive));
+		}));
+	}
 
 } // namespace UE::LevelEditor
 
