Add the viewport layouts menu to the new viewport toolbar

diff --git a/Engine/Source/Editor/LevelEditor/Private/SLevelViewport.cpp b/Engine/Source/Editor/LevelEditor/Private/SLevelViewport.cpp
--- a/Engine/Source/Editor/LevelEditor/Private/SLevelViewport.cpp
+++ b/Engine/Source/Editor/LevelEditor/Private/SLevelViewport.cpp
@@ -1926,8 +1926,15 @@ TSharedPtr<SWidget> SLevelViewport::MakeViewportToolbar()
 			UToolMenu* const ViewportToolbarMenu = UToolMenus::Get()->RegisterMenu(
 				LevelEditorViewportToolbarName, NAME_None /* parent */, EMultiBoxType::ToolBar);
 
-			UE::LevelEditor::AddViewportToolbarTransformsSection(ViewportToolbarMenu);
-			UE::LevelEditor::AddLevelEditorViewportToolbarSettingsSection(ViewportToolbarMenu);
+			{
+				FToolMenuSection& LeftSection = ViewportToolbarMenu->FindOrAddSection("Left");
+				UE::LevelEditor::AddViewportToolbarTransformsSection(LeftSection);
+			}
+
+			{
+				FToolMenuSection& RightSection = ViewportToolbarMenu->FindOrAddSection("Right");
+				UE::LevelEditor::AddLevelEditorViewportToolbarSettingsSection(RightSection);
+			}
 		}
 	}
 
diff --git a/Engine/Source/Editor/LevelEditor/Private/SLevelViewportToolBar.cpp b/Engine/Source/Editor/LevelEditor/Private/SLevelViewportToolBar.cpp
--- a/Engine/Source/Editor/LevelEditor/Private/SLevelViewportToolBar.cpp
+++ b/Engine/Source/Editor/LevelEditor/Private/SLevelViewportToolBar.cpp
@@ -57,6 +57,7 @@
 #include "SCommonEditorViewportToolbarBase.h"
 #include "SActionableMessageViewportWidget.h"
 #include "Engine/SceneCapture.h"
+#include "ViewportToolbar/LevelEditorViewportToolbarSections.h"
 
 #define LOCTEXT_NAMESPACE "LevelViewportToolBar"
 
@@ -823,14 +824,11 @@ void SLevelViewportToolBar::FillOptionsMenu(UToolMenu* Menu)
 
 		{
 			FToolMenuSection& Section = Menu->AddSection("LevelViewportLayouts");
-			Section.AddSubMenu(
-				"Configs",
-				LOCTEXT("ConfigsSubMenu", "Layouts"),
-				FText::GetEmpty(),
-				FNewToolMenuDelegate::CreateSP(this, &SLevelViewportToolBar::GenerateViewportConfigsMenu),
-				false,
-				FSlateIcon(FAppStyle::Get().GetStyleSetName(), "Icons.Layout")
-			);
+			Section.AddSubMenu("Configs", LOCTEXT("ConfigsSubMenu", "Layouts"), FText::GetEmpty(),
+				FNewToolMenuDelegate::CreateLambda([WeakViewport = Viewport](UToolMenu* InMenu) {
+					UE::LevelEditor::GenerateViewportLayoutsMenu(InMenu, WeakViewport.Pin());
+				}),
+				false, FSlateIcon(FAppStyle::Get().GetStyleSetName(), "Icons.Layout"));
 		}
 
 		{
@@ -1105,125 +1103,6 @@ void SLevelViewportToolBar::FillCameraMenu(UToolMenu* Menu) const
 	}
 }
 
-void SLevelViewportToolBar::GenerateViewportConfigsMenu(UToolMenu* Menu) const
-{
-	check (Viewport.IsValid());
-	TSharedPtr<FUICommandList> CommandList = Viewport.Pin()->GetCommandList();
-
-	// Disable searching in this menu because it only contains visual representations of
-	// viewport layouts without any searchable text.
-	Menu->bSearchable = false;
-
-	{
-		FToolMenuSection& Section = Menu->AddSection("LevelViewportOnePaneConfigs", LOCTEXT("OnePaneConfigHeader", "One Pane"));
-
-		FSlimHorizontalToolBarBuilder OnePaneButton(CommandList, FMultiBoxCustomization::None);
-		OnePaneButton.SetLabelVisibility(EVisibility::Collapsed);
-		OnePaneButton.SetStyle(&FAppStyle::Get(), "ViewportLayoutToolbar");
-
-		OnePaneButton.AddToolBarButton(FLevelViewportCommands::Get().ViewportConfig_OnePane);
-
-		Section.AddEntry(FToolMenuEntry::InitWidget(
-			"LevelViewportOnePaneConfigs",
-			SNew(SHorizontalBox)
-			+ SHorizontalBox::Slot()
-			.AutoWidth()
-			[
-				OnePaneButton.MakeWidget()
-			]
-			+ SHorizontalBox::Slot()
-			.FillWidth(1)
-			[
-				SNullWidget::NullWidget
-			],
-			FText::GetEmpty(), true
-			));
-	}
-
-	{
-		FToolMenuSection& Section = Menu->AddSection("LevelViewportTwoPaneConfigs", LOCTEXT("TwoPaneConfigHeader", "Two Panes"));
-		FSlimHorizontalToolBarBuilder TwoPaneButtons(CommandList, FMultiBoxCustomization::None);
-		TwoPaneButtons.SetLabelVisibility(EVisibility::Collapsed);
-		TwoPaneButtons.SetStyle(&FAppStyle::Get(), "ViewportLayoutToolbar");
-
-		TwoPaneButtons.AddToolBarButton(FLevelViewportCommands::Get().ViewportConfig_TwoPanesH, NAME_None, FText());
-		TwoPaneButtons.AddToolBarButton(FLevelViewportCommands::Get().ViewportConfig_TwoPanesV, NAME_None, FText());
-
-		Section.AddEntry(FToolMenuEntry::InitWidget(
-			"LevelViewportTwoPaneConfigs",
-			SNew(SHorizontalBox)
-			+SHorizontalBox::Slot()
-			.AutoWidth()
-			[
-				TwoPaneButtons.MakeWidget()
-			]
-			+SHorizontalBox::Slot()
-			.FillWidth(1)
-			[
-				SNullWidget::NullWidget
-			],
-			FText::GetEmpty(), true
-			));
-	}
-
-	{
-		FToolMenuSection& Section = Menu->AddSection("LevelViewportThreePaneConfigs", LOCTEXT("ThreePaneConfigHeader", "Three Panes"));
-		FSlimHorizontalToolBarBuilder ThreePaneButtons(CommandList, FMultiBoxCustomization::None);
-		ThreePaneButtons.SetLabelVisibility(EVisibility::Collapsed);
-		ThreePaneButtons.SetStyle(&FAppStyle::Get(), "ViewportLayoutToolbar");
-
-		ThreePaneButtons.AddToolBarButton(FLevelViewportCommands::Get().ViewportConfig_ThreePanesLeft, NAME_None, FText());
-		ThreePaneButtons.AddToolBarButton(FLevelViewportCommands::Get().ViewportConfig_ThreePanesRight, NAME_None, FText());
-		ThreePaneButtons.AddToolBarButton(FLevelViewportCommands::Get().ViewportConfig_ThreePanesTop, NAME_None, FText());
-		ThreePaneButtons.AddToolBarButton(FLevelViewportCommands::Get().ViewportConfig_ThreePanesBottom, NAME_None, FText());
-
-		Section.AddEntry(FToolMenuEntry::InitWidget(
-			"LevelViewportThreePaneConfigs",
-			SNew(SHorizontalBox)
-			+SHorizontalBox::Slot()
-			.AutoWidth()
-			[
-				ThreePaneButtons.MakeWidget()
-			]
-			+SHorizontalBox::Slot()
-			.FillWidth(1)
-			[
-				SNullWidget::NullWidget
-			],
-			FText::GetEmpty(), true
-			));
-	}
-
-	{
-		FToolMenuSection& Section = Menu->AddSection("LevelViewportFourPaneConfigs", LOCTEXT("FourPaneConfigHeader", "Four Panes"));
-		FSlimHorizontalToolBarBuilder FourPaneButtons(CommandList, FMultiBoxCustomization::None);
-		FourPaneButtons.SetLabelVisibility(EVisibility::Collapsed);
-		FourPaneButtons.SetStyle(&FAppStyle::Get(), "ViewportLayoutToolbar");
-
-		FourPaneButtons.AddToolBarButton(FLevelViewportCommands::Get().ViewportConfig_FourPanes2x2, NAME_None, FText());
-		FourPaneButtons.AddToolBarButton(FLevelViewportCommands::Get().ViewportConfig_FourPanesLeft, NAME_None, FText());
-		FourPaneButtons.AddToolBarButton(FLevelViewportCommands::Get().ViewportConfig_FourPanesRight, NAME_None, FText());
-		FourPaneButtons.AddToolBarButton(FLevelViewportCommands::Get().ViewportConfig_FourPanesTop, NAME_None, FText());
-		FourPaneButtons.AddToolBarButton(FLevelViewportCommands::Get().ViewportConfig_FourPanesBottom, NAME_None, FText());	
-
-		Section.AddEntry(FToolMenuEntry::InitWidget(
-			"LevelViewportFourPaneConfigs",
-			SNew(SHorizontalBox)
-			+SHorizontalBox::Slot()
-			.AutoWidth()
-			[
-				FourPaneButtons.MakeWidget()
-			]
-			+SHorizontalBox::Slot()
-			.FillWidth(1)
-			[
-				SNullWidget::NullWidget
-			],
-			FText::GetEmpty(), true
-			));
-	}
-}
-
 TSharedRef<SWidget> SLevelViewportToolBar::GenerateShowMenu() const
 {
 	static const FName MenuName("LevelEditor.LevelViewportToolbar.Show");
diff --git a/Engine/Source/Editor/LevelEditor/Private/SLevelViewportToolBar.h b/Engine/Source/Editor/LevelEditor/Private/SLevelViewportToolBar.h
--- a/Engine/Source/Editor/LevelEditor/Private/SLevelViewportToolBar.h
+++ b/Engine/Source/Editor/LevelEditor/Private/SLevelViewportToolBar.h
@@ -265,9 +265,6 @@ private:
 	 */
 	void FillShowHLODsMenu(UToolMenu* Menu) const;
 
-	/** Generates the layout sub-menu content */
-	void GenerateViewportConfigsMenu(UToolMenu* Menu) const;
-
 	/** Gets the world we are editing */
 	TWeakObjectPtr<UWorld> GetWorld() const;
 
diff --git a/Engine/Source/Editor/LevelEditor/Private/ViewportToolbar/LevelEditorViewportToolbarSections.cpp b/Engine/Source/Editor/LevelEditor/Private/ViewportToolbar/LevelEditorViewportToolbarSections.cpp
--- a/Engine/Source/Editor/LevelEditor/Private/ViewportToolbar/LevelEditorViewportToolbarSections.cpp
+++ b/Engine/Source/Editor/LevelEditor/Private/ViewportToolbar/LevelEditorViewportToolbarSections.cpp
@@ -7,6 +7,7 @@
 #include "LevelViewportActions.h"
 #include "LevelViewportContext.h"
 #include "SLevelViewport.h"
+#include "Templates/SharedPointer.h"
 #include "ToolMenu.h"
 
 #define LOCTEXT_NAMESPACE "LevelEditorViewportToolbar"
@@ -15,11 +16,9 @@ namespace UE::LevelEditor
 {
 
 // TODO: Move this outside the level editor and make it publicly available to anyone building a viewport toolbar.
-void AddViewportToolbarTransformsSection(UToolMenu* InMenu)
+void AddViewportToolbarTransformsSection(FToolMenuSection& InSection)
 {
-	FToolMenuSection& Section = InMenu->FindOrAddSection("Left");
-
-	Section.AddSubMenu("Transforms", LOCTEXT("TransformsSubmenuLabel", "Transforms"),
+	InSection.AddSubMenu("Transforms", LOCTEXT("TransformsSubmenuLabel", "Transforms"),
 		LOCTEXT("TransformsSubmenuTooltip", "Viewport-related transforms tools"),
 		FNewToolMenuDelegate::CreateLambda([](UToolMenu* Submenu) -> void {
 			FToolMenuSection& Section = Submenu->FindOrAddSection(NAME_None);
@@ -81,11 +80,132 @@ void AddFeatureLevelPreviewSubmenu(FToolMenuSection& Section)
 		}));
 }
 
-void AddLevelEditorViewportToolbarSettingsSection(UToolMenu* InMenu)
+void GenerateViewportLayoutsMenu(UToolMenu* InMenu, TSharedPtr<::SLevelViewport> InViewport)
 {
-	FToolMenuSection& RightSection = InMenu->FindOrAddSection("Right");
+	TSharedPtr<FUICommandList> CommandList = InViewport->GetCommandList();
+
+	// Disable searching in this menu because it only contains visual representations of
+	// viewport layouts without any searchable text.
+	InMenu->bSearchable = false;
+
+	{
+		FToolMenuSection& Section = InMenu->AddSection(
+			"LevelViewportOnePaneConfigs", LOCTEXT("OnePaneConfigHeader", "One Pane"));
+
+		FSlimHorizontalToolBarBuilder OnePaneButton(CommandList, FMultiBoxCustomization::None);
+		OnePaneButton.SetLabelVisibility(EVisibility::Collapsed);
+		OnePaneButton.SetStyle(&FAppStyle::Get(), "ViewportLayoutToolbar");
+
+		OnePaneButton.AddToolBarButton(FLevelViewportCommands::Get().ViewportConfig_OnePane);
+
+		Section.AddEntry(FToolMenuEntry::InitWidget("LevelViewportOnePaneConfigs",
+			// clang-format off
+			SNew(SHorizontalBox)
+			+ SHorizontalBox::Slot()
+			.AutoWidth()
+			[
+				OnePaneButton.MakeWidget()
+			]
+			+ SHorizontalBox::Slot()
+			.FillWidth(1)
+			[
+				SNullWidget::NullWidget
+			],
+			// clang-format on
+			FText::GetEmpty(), true));
+	}
+
+	{
+		FToolMenuSection& Section = InMenu->AddSection(
+			"LevelViewportTwoPaneConfigs", LOCTEXT("TwoPaneConfigHeader", "Two Panes"));
+		FSlimHorizontalToolBarBuilder TwoPaneButtons(CommandList, FMultiBoxCustomization::None);
+		TwoPaneButtons.SetLabelVisibility(EVisibility::Collapsed);
+		TwoPaneButtons.SetStyle(&FAppStyle::Get(), "ViewportLayoutToolbar");
+
+		TwoPaneButtons.AddToolBarButton(FLevelViewportCommands::Get().ViewportConfig_TwoPanesH, NAME_None, FText());
+		TwoPaneButtons.AddToolBarButton(FLevelViewportCommands::Get().ViewportConfig_TwoPanesV, NAME_None, FText());
+
+		Section.AddEntry(FToolMenuEntry::InitWidget("LevelViewportTwoPaneConfigs",
+			// clang-format off
+			SNew(SHorizontalBox)
+			+SHorizontalBox::Slot()
+			.AutoWidth()
+			[
+				TwoPaneButtons.MakeWidget()
+			]
+			+SHorizontalBox::Slot()
+			.FillWidth(1)
+			[
+				SNullWidget::NullWidget
+			],
+			// clang-format on
+			FText::GetEmpty(), true));
+	}
+
+	{
+		FToolMenuSection& Section = InMenu->AddSection(
+			"LevelViewportThreePaneConfigs", LOCTEXT("ThreePaneConfigHeader", "Three Panes"));
+		FSlimHorizontalToolBarBuilder ThreePaneButtons(CommandList, FMultiBoxCustomization::None);
+		ThreePaneButtons.SetLabelVisibility(EVisibility::Collapsed);
+		ThreePaneButtons.SetStyle(&FAppStyle::Get(), "ViewportLayoutToolbar");
+
+		ThreePaneButtons.AddToolBarButton(FLevelViewportCommands::Get().ViewportConfig_ThreePanesLeft, NAME_None, FText());
+		ThreePaneButtons.AddToolBarButton(FLevelViewportCommands::Get().ViewportConfig_ThreePanesRight, NAME_None, FText());
+		ThreePaneButtons.AddToolBarButton(FLevelViewportCommands::Get().ViewportConfig_ThreePanesTop, NAME_None, FText());
+		ThreePaneButtons.AddToolBarButton(
+			FLevelViewportCommands::Get().ViewportConfig_ThreePanesBottom, NAME_None, FText());
 
-	RightSection.AddSubMenu("Settings", LOCTEXT("SettingsSubmenuLabel", "Settings"),
+		Section.AddEntry(FToolMenuEntry::InitWidget("LevelViewportThreePaneConfigs",
+			// clang-format off
+			SNew(SHorizontalBox)
+			+SHorizontalBox::Slot()
+			.AutoWidth()
+			[
+				ThreePaneButtons.MakeWidget()
+			]
+			+SHorizontalBox::Slot()
+			.FillWidth(1)
+			[
+				SNullWidget::NullWidget
+			],
+			// clang-format on
+			FText::GetEmpty(), true));
+	}
+
+	{
+		FToolMenuSection& Section = InMenu->AddSection(
+			"LevelViewportFourPaneConfigs", LOCTEXT("FourPaneConfigHeader", "Four Panes"));
+		FSlimHorizontalToolBarBuilder FourPaneButtons(CommandList, FMultiBoxCustomization::None);
+		FourPaneButtons.SetLabelVisibility(EVisibility::Collapsed);
+		FourPaneButtons.SetStyle(&FAppStyle::Get(), "ViewportLayoutToolbar");
+
+		FourPaneButtons.AddToolBarButton(FLevelViewportCommands::Get().ViewportConfig_FourPanes2x2, NAME_None, FText());
+		FourPaneButtons.AddToolBarButton(FLevelViewportCommands::Get().ViewportConfig_FourPanesLeft, NAME_None, FText());
+		FourPaneButtons.AddToolBarButton(FLevelViewportCommands::Get().ViewportConfig_FourPanesRight, NAME_None, FText());
+		FourPaneButtons.AddToolBarButton(FLevelViewportCommands::Get().ViewportConfig_FourPanesTop, NAME_None, FText());
+		FourPaneButtons.AddToolBarButton(FLevelViewportCommands::Get().ViewportConfig_FourPanesBottom, NAME_None, FText());
+
+		Section.AddEntry(FToolMenuEntry::InitWidget("LevelViewportFourPaneConfigs",
+			// clang-format off
+			SNew(SHorizontalBox)
+			+SHorizontalBox::Slot()
+			.AutoWidth()
+			[
+				FourPaneButtons.MakeWidget()
+			]
+			+SHorizontalBox::Slot()
+			.FillWidth(1)
+			[
+				SNullWidget::NullWidget
+			],
+			// clang-format on
+			FText::GetEmpty(), true));
+	}
+}
+
+void AddLevelEditorViewportToolbarSettingsSection(FToolMenuSection& InSection)
+{
+	InSection.AddSubMenu("Settings", LOCTEXT("SettingsSubmenuLabel", "Settings"),
 		LOCTEXT("SettingsSubmenuTooltip", "Viewport-related settings"),
 		FNewToolMenuDelegate::CreateLambda([](UToolMenu* Submenu) -> void {
 			FToolMenuSection& UnnamedSection = Submenu->FindOrAddSection(NAME_None);
@@ -137,8 +257,24 @@ void AddLevelEditorViewportToolbarSettingsSection(UToolMenu* InMenu)
 
 			// Add immersive mode toggle.
 			UnnamedSection.AddEntry(FToolMenuEntry::InitToolBarButton(FLevelViewportCommands::Get().ToggleImmersive));
+
+			UnnamedSection.AddSubMenu("ViewportLayouts", LOCTEXT("ViewportLayoutsLabel", "Layouts"),
+				LOCTEXT("ViewportLayoutsTooltip", "Configure the layouts of the viewport windows"),
+				FNewToolMenuDelegate::CreateLambda([](UToolMenu* InMenu) {
+					ULevelViewportContext* const LevelViewportContext = InMenu->FindContext<ULevelViewportContext>();
+					if (!LevelViewportContext)
+					{
+						return;
+					}
+
+					if (const TSharedPtr<::SLevelViewport> LevelViewport = LevelViewportContext->LevelViewport.Pin())
+					{
+						GenerateViewportLayoutsMenu(InMenu, LevelViewport);
+					}
+				}),
+				false, FSlateIcon(FAppStyle::Get().GetStyleSetName(), "Icons.Layout"));
 		}));
-	}
+}
 
 } // namespace UE::LevelEditor
 
diff --git a/Engine/Source/Editor/LevelEditor/Private/ViewportToolbar/LevelEditorViewportToolbarSections.h b/Engine/Source/Editor/LevelEditor/Private/ViewportToolbar/LevelEditorViewportToolbarSections.h
--- a/Engine/Source/Editor/LevelEditor/Private/ViewportToolbar/LevelEditorViewportToolbarSections.h
+++ b/Engine/Source/Editor/LevelEditor/Private/ViewportToolbar/LevelEditorViewportToolbarSections.h
@@ -2,15 +2,19 @@
 
 #pragma once
 
-struct FToolMenuSection;
+#include "Templates/SharedPointerFwd.h"
+
+class SLevelViewport;
 class UToolMenu;
+struct FToolMenuSection;
 
 namespace UE::LevelEditor
 {
 
-void AddViewportToolbarTransformsSection(UToolMenu* InMenu);
+void AddViewportToolbarTransformsSection(FToolMenuSection& InSection);
 void AddMaterialQualityLevelSubmenu(FToolMenuSection& Section);
 void AddFeatureLevelPreviewSubmenu(FToolMenuSection& Section);
-void AddLevelEditorViewportToolbarSettingsSection(UToolMenu* InMenu);
+void GenerateViewportLayoutsMenu(UToolMenu* InMenu, TSharedPtr<::SLevelViewport> InViewport);
+void AddLevelEditorViewportToolbarSettingsSection(FToolMenuSection& InSection);
 
 } // namespace UE::LevelEditor
