Convert the new viewport toolbar entries MaximizeRestore and ToggleImmersive from toolbar buttons to menu entries

diff --git a/Engine/Source/Editor/LevelEditor/Private/ViewportToolbar/LevelEditorViewportToolbarSections.cpp b/Engine/Source/Editor/LevelEditor/Private/ViewportToolbar/LevelEditorViewportToolbarSections.cpp
--- a/Engine/Source/Editor/LevelEditor/Private/ViewportToolbar/LevelEditorViewportToolbarSections.cpp
+++ b/Engine/Source/Editor/LevelEditor/Private/ViewportToolbar/LevelEditorViewportToolbarSections.cpp
@@ -248,15 +248,16 @@ void AddLevelEditorViewportToolbarSettingsSection(FToolMenuSection& InSection)
 						return ECheckBoxState::Undetermined;
 					});
 
-				UnnamedSection.AddEntry(FToolMenuEntry::InitToolBarButton("MaximizeRestore", MaximizeRestoreAction,
-					LOCTEXT("MaximizeRestoreLabel", "Maximize/restore"),
-					LOCTEXT("MaximizeRestoreTooltip", "Maximizes or restores this viewport"),
-					FSlateIcon(FAppStyle::GetAppStyleSetName(), "EditorViewportToolBar.Maximize.Normal"),
-					EUserInterfaceActionType::ToggleButton));
+				UnnamedSection
+					.AddMenuEntry("MaximizeRestore", LOCTEXT("MaximizeRestoreLabel", "Maximize/restore"),
+						LOCTEXT("MaximizeRestoreTooltip", "Maximizes or restores this viewport"),
+						FSlateIcon(FAppStyle::GetAppStyleSetName(), "EditorViewportToolBar.Maximize.Normal"),
+						MaximizeRestoreAction, EUserInterfaceActionType::ToggleButton)
+					.SetShowInToolbarTopLevel(true);
 			}
 
 			// Add immersive mode toggle.
-			UnnamedSection.AddEntry(FToolMenuEntry::InitToolBarButton(FLevelViewportCommands::Get().ToggleImmersive));
+			UnnamedSection.AddMenuEntry(FLevelViewportCommands::Get().ToggleImmersive);
 
 			UnnamedSection.AddSubMenu("ViewportLayouts", LOCTEXT("ViewportLayoutsLabel", "Layouts"),
 				LOCTEXT("ViewportLayoutsTooltip", "Configure the layouts of the viewport windows"),
