Add "EditorRequestBeginPlay" to expose starting PIE to scripting

diff --git a/Engine/Source/Editor/LevelEditor/Private/LevelEditorSubsystem.cpp b/Engine/Source/Editor/LevelEditor/Private/LevelEditorSubsystem.cpp
--- a/Engine/Source/Editor/LevelEditor/Private/LevelEditorSubsystem.cpp
+++ b/Engine/Source/Editor/LevelEditor/Private/LevelEditorSubsystem.cpp
@@ -281,6 +281,20 @@ bool ULevelEditorSubsystem::EditorGetGameView(FName ViewportConfigKey)
 	return false;
 }
 
+void ULevelEditorSubsystem::EditorRequestBeginPlay()
+{
+	FLevelEditorModule& LevelEditorModule = FModuleManager::GetModuleChecked<FLevelEditorModule>("LevelEditor");
+
+	TSharedPtr<IAssetViewport> ActiveLevelViewport = LevelEditorModule.GetFirstActiveViewport();
+	if (ActiveLevelViewport.IsValid())
+	{
+		FRequestPlaySessionParams SessionParams;
+		SessionParams.WorldType = EPlaySessionWorldType::PlayInEditor;
+		SessionParams.DestinationSlateViewport = ActiveLevelViewport;
+
+		GUnrealEd->RequestPlaySession(SessionParams);
+	}
+}
 
 void ULevelEditorSubsystem::EditorRequestEndPlay()
 {
diff --git a/Engine/Source/Editor/LevelEditor/Public/LevelEditorSubsystem.h b/Engine/Source/Editor/LevelEditor/Public/LevelEditorSubsystem.h
--- a/Engine/Source/Editor/LevelEditor/Public/LevelEditorSubsystem.h
+++ b/Engine/Source/Editor/LevelEditor/Public/LevelEditorSubsystem.h
@@ -67,6 +67,9 @@ public:
 	UFUNCTION(BlueprintCallable, Category = "Editor Scripting | Level Utility", meta = (DevelopmentOnly))
 	bool EditorGetGameView(FName ViewportConfigKey = NAME_None);
 
+	UFUNCTION(BlueprintCallable, Category = "Editor Scripting | Level Utility", meta = (DevelopmentOnly))
+	void EditorRequestBeginPlay();
+
 	UFUNCTION(BlueprintCallable, Category = "Editor Scripting | Level Utility", meta = (DevelopmentOnly))
 	void EditorRequestEndPlay();
 
