Move the realtime toggle of the new viewport toolbar to the top-level if realtime is off for the current viewport

Note that this top-level state doesn't react properly when unrelated code changes the top-level state. This will improve in the future when no other unrelated UI toggles the realtime state (right now it's easy to change in via the old viewport toolbar).

diff --git a/Engine/Source/Editor/LevelEditor/Private/SLevelViewportToolBar.cpp b/Engine/Source/Editor/LevelEditor/Private/SLevelViewportToolBar.cpp
--- a/Engine/Source/Editor/LevelEditor/Private/SLevelViewportToolBar.cpp
+++ b/Engine/Source/Editor/LevelEditor/Private/SLevelViewportToolBar.cpp
@@ -1983,8 +1983,8 @@ FReply SLevelViewportToolBar::OnRealtimeWarningClicked()
 EVisibility SLevelViewportToolBar::GetRealtimeWarningVisibility() const
 {
 	FLevelEditorViewportClient& ViewportClient = Viewport.Pin()->GetLevelViewportClient();
-	// If the viewport is not realtime and there is no override then realtime is off
-	return !ViewportClient.IsRealtime() && !ViewportClient.IsRealtimeOverrideSet() && ViewportClient.IsPerspective() ? EVisibility::Visible : EVisibility::Collapsed;
+	const bool bWarn = UE::LevelEditor::ShowViewportRealtimeWarning(ViewportClient);
+	return bWarn ? EVisibility::Visible : EVisibility::Collapsed;
 }
 
 FText SLevelViewportToolBar::GetScalabilityWarningLabel() const
diff --git a/Engine/Source/Editor/LevelEditor/Private/ViewportToolbar/LevelEditorViewportToolbarSections.cpp b/Engine/Source/Editor/LevelEditor/Private/ViewportToolbar/LevelEditorViewportToolbarSections.cpp
--- a/Engine/Source/Editor/LevelEditor/Private/ViewportToolbar/LevelEditorViewportToolbarSections.cpp
+++ b/Engine/Source/Editor/LevelEditor/Private/ViewportToolbar/LevelEditorViewportToolbarSections.cpp
@@ -9,9 +9,46 @@
 #include "SLevelViewport.h"
 #include "Templates/SharedPointer.h"
 #include "ToolMenu.h"
+#include "ToolMenus.h"
 
 #define LOCTEXT_NAMESPACE "LevelEditorViewportToolbar"
 
+namespace UE::LevelEditor::Private
+{
+
+TOptional<bool> UpdateAndGetRealtimeWarningFromContext(const FToolMenuContext& Context)
+{
+	if (ULevelViewportContext* const LevelViewportContext = Context.FindContext<ULevelViewportContext>())
+	{
+		if (const TSharedPtr<::SLevelViewport> LevelViewport = LevelViewportContext->LevelViewport.Pin())
+		{
+			const bool bShowWarning = ShowViewportRealtimeWarning(LevelViewport->GetLevelViewportClient());
+
+			LevelViewportContext->CachedShouldShowRealtimeOffWarning = bShowWarning;
+
+			return LevelViewportContext->CachedShouldShowRealtimeOffWarning;
+		}
+	}
+
+	return TOptional<bool>();
+}
+
+TOptional<bool> IsDirtyRealtimeWarningFromContext(const FToolMenuContext& Context)
+{
+	if (ULevelViewportContext* const LevelViewportContext = Context.FindContext<ULevelViewportContext>())
+	{
+		if (const TSharedPtr<::SLevelViewport> LevelViewport = LevelViewportContext->LevelViewport.Pin())
+		{
+			const bool bShowWarning = ShowViewportRealtimeWarning(LevelViewport->GetLevelViewportClient());
+			return bShowWarning != LevelViewportContext->CachedShouldShowRealtimeOffWarning;
+		}
+	}
+
+	return TOptional<bool>();
+}
+
+} // namespace UE::LevelEditor::Private
+
 namespace UE::LevelEditor
 {
 
@@ -45,6 +82,11 @@ void AddViewportToolbarTransformsSection(FToolMenuSection& InSection)
 		}));
 }
 
+bool ShowViewportRealtimeWarning(FLevelEditorViewportClient& ViewportClient)
+{
+	return !ViewportClient.IsRealtime() && !ViewportClient.IsRealtimeOverrideSet() && ViewportClient.IsPerspective();
+}
+
 void AddMaterialQualityLevelSubmenu(FToolMenuSection& Section)
 {
 	Section.AddSubMenu("MaterialQualityLevel",
@@ -211,7 +253,76 @@ void AddLevelEditorViewportToolbarSettingsSection(FToolMenuSection& InSection)
 			FToolMenuSection& UnnamedSection = Submenu->FindOrAddSection(NAME_None);
 
 			// Add realtime rendering toggle.
-			UnnamedSection.AddMenuEntry(FEditorViewportCommands::Get().ToggleRealTime).SetShowInToolbarTopLevel(true);
+			UnnamedSection.AddDynamicEntry("ToggleRealtimeDynamicSection",
+				FNewToolMenuSectionDelegate::CreateLambda([](FToolMenuSection& InnerSection) -> void {
+					FToolUIAction RealtimeToggleAction;
+					RealtimeToggleAction.ExecuteAction = FToolMenuExecuteAction::CreateLambda(
+						[](const FToolMenuContext& Context) -> void {
+							ULevelViewportContext* const LevelViewportContext =
+								Context.FindContext<ULevelViewportContext>();
+							if (!LevelViewportContext)
+							{
+								return;
+							}
+
+							if (const TSharedPtr<::SLevelViewport> LevelViewport = LevelViewportContext->LevelViewport.Pin())
+							{
+								LevelViewport->OnToggleRealtime();
+								UToolMenus::Get()->RefreshAllWidgets();
+							}
+						});
+
+					RealtimeToggleAction.GetActionCheckState = FToolMenuGetActionCheckState::CreateLambda(
+						[](const FToolMenuContext& Context) -> ECheckBoxState {
+							ULevelViewportContext* const LevelViewportContext =
+								Context.FindContext<ULevelViewportContext>();
+							if (!LevelViewportContext)
+							{
+								return ECheckBoxState::Undetermined;
+							}
+
+							// Check if the realtime warn state is outdated and if so refresh widgets to update our
+							// top-level status.
+							if (const TOptional<bool> IsDirty = Private::IsDirtyRealtimeWarningFromContext(Context);
+								IsDirty.IsSet() && IsDirty.GetValue())
+							{
+								UToolMenus::Get()->RefreshAllWidgets();
+							}
+
+							if (const TSharedPtr<::SLevelViewport> LevelViewport = LevelViewportContext->LevelViewport.Pin())
+							{
+								return LevelViewport->IsRealtime() ? ECheckBoxState::Checked : ECheckBoxState::Unchecked;
+							}
+
+							return ECheckBoxState::Undetermined;
+						});
+
+					bool bDisplayTopLevel = false;
+					if (const TOptional<bool> ShouldWarn = Private::UpdateAndGetRealtimeWarningFromContext(
+							InnerSection.Context);
+						ShouldWarn.IsSet())
+					{
+						bDisplayTopLevel = ShouldWarn.GetValue();
+					}
+					else
+					{
+						// If we couldn't get the warn state, pretend we don't have to warn.
+						bDisplayTopLevel = false;
+					}
+
+					const FText Tooltip =
+						bDisplayTopLevel
+							? LOCTEXT("ToggleRealtimeTooltip_WarnRealtimeOff",
+								"This viewport is not updating in realtime.  Click to turn on realtime mode.")
+							: LOCTEXT("ToggleRealtimeTooltip", "Toggle realtime rendering of the viewport");
+
+					FToolMenuEntry ToggleRealtime = FToolMenuEntry::InitMenuEntry("ToggleRealtime",
+						LOCTEXT("ToggleRealtimeLabel", "Realtime"), Tooltip,
+						FSlateIcon(FAppStyle::GetAppStyleSetName(), "EditorViewport.ToggleRealTime"),
+						RealtimeToggleAction, EUserInterfaceActionType::ToggleButton);
+					ToggleRealtime.SetShowInToolbarTopLevel(bDisplayTopLevel);
+					InnerSection.AddEntry(ToggleRealtime);
+				}));
 
 			AddMaterialQualityLevelSubmenu(UnnamedSection);
 			AddFeatureLevelPreviewSubmenu(UnnamedSection);
diff --git a/Engine/Source/Editor/LevelEditor/Private/ViewportToolbar/LevelEditorViewportToolbarSections.h b/Engine/Source/Editor/LevelEditor/Private/ViewportToolbar/LevelEditorViewportToolbarSections.h
--- a/Engine/Source/Editor/LevelEditor/Private/ViewportToolbar/LevelEditorViewportToolbarSections.h
+++ b/Engine/Source/Editor/LevelEditor/Private/ViewportToolbar/LevelEditorViewportToolbarSections.h
@@ -4,6 +4,7 @@
 
 #include "Templates/SharedPointerFwd.h"
 
+class FLevelEditorViewportClient;
 class SLevelViewport;
 class UToolMenu;
 struct FToolMenuSection;
@@ -12,6 +13,7 @@ namespace UE::LevelEditor
 {
 
 void AddViewportToolbarTransformsSection(FToolMenuSection& InSection);
+bool ShowViewportRealtimeWarning(FLevelEditorViewportClient& ViewportClient);
 void AddMaterialQualityLevelSubmenu(FToolMenuSection& Section);
 void AddFeatureLevelPreviewSubmenu(FToolMenuSection& Section);
 void GenerateViewportLayoutsMenu(UToolMenu* InMenu, TSharedPtr<::SLevelViewport> InViewport);
diff --git a/Engine/Source/Editor/LevelEditor/Private/ViewportToolbar/LevelViewportContext.h b/Engine/Source/Editor/LevelEditor/Private/ViewportToolbar/LevelViewportContext.h
--- a/Engine/Source/Editor/LevelEditor/Private/ViewportToolbar/LevelViewportContext.h
+++ b/Engine/Source/Editor/LevelEditor/Private/ViewportToolbar/LevelViewportContext.h
@@ -12,6 +12,10 @@ UCLASS()
 class LEVELEDITOR_API ULevelViewportContext : public UObject
 {
 	GENERATED_BODY()
+
 public:
 	TWeakPtr<SLevelViewport> LevelViewport;
+	// This allows the viewport toolbar to know if it needs to refresh itself to update the top-level state of the
+	// realtime button.
+	TOptional<bool> CachedShouldShowRealtimeOffWarning;
 };
