Fix editor "Copy Actor File Path" not working with non-OFPA actors.

diff --git a/Engine/Source/Editor/LevelEditor/Private/LevelEditorActions.cpp b/Engine/Source/Editor/LevelEditor/Private/LevelEditorActions.cpp
--- a/Engine/Source/Editor/LevelEditor/Private/LevelEditorActions.cpp
+++ b/Engine/Source/Editor/LevelEditor/Private/LevelEditorActions.cpp
@@ -2955,12 +2955,37 @@ void FLevelEditorActionCallbacks::SnapObjectToView_Clicked()
 
 void FLevelEditorActionCallbacks::CopyActorFilePathtoClipboard_Clicked()
 {
-	TArray<const UObject*> Objects;
+	TStringBuilder<1024> Result;
+
+	TArray<AActor*> SelectedActors;
 	for (FSelectionIterator It(GEditor->GetSelectedActorIterator()); It; ++It)
 	{
-		Objects.Add(*It);
+		SelectedActors.Add(Cast<AActor>(*It));
+	}
+
+	for (AActor* Actor : SelectedActors)
+	{
+		const UPackage* Package = Actor->GetPackage();
+		const FString LocalFullPath(Package->GetLoadedPath().GetLocalFullPath());
+		if (SelectedActors.Num() > 1)
+		{
+			const FString& ActorLabel = Actor->GetActorLabel(false);
+			if (ActorLabel.Len())
+			{
+				Result.Append(ActorLabel);
+			}
+			Result.Append(TEXT("("));
+			Result.Append(Actor->GetName());
+			Result.Append(TEXT("): "));
+		}
+		Result.Append(FPaths::ConvertRelativePathToFull(LocalFullPath));
+		Result.Append(LINE_TERMINATOR);
+	}
+
+	if (Result.Len())
+	{
+		FPlatformApplicationMisc::ClipboardCopy(*Result);
 	}
-	FExternalPackageHelper::CopyObjectsExternalPackageFilePathToClipboard(Objects);
 }
 
 void FLevelEditorActionCallbacks::SaveActor_Clicked()
