Add "Performance & Scalability" submenu to the new viewport toolbar

This moves the a few things into the new menu and also adds the "Engine Scalability" and "Screen Percentage" submenus which previously didn't exist in the new viewport toolbar.

diff --git a/Engine/Source/Editor/LevelEditor/Private/SLevelViewport.cpp b/Engine/Source/Editor/LevelEditor/Private/SLevelViewport.cpp
--- a/Engine/Source/Editor/LevelEditor/Private/SLevelViewport.cpp
+++ b/Engine/Source/Editor/LevelEditor/Private/SLevelViewport.cpp
@@ -1933,6 +1933,7 @@ TSharedPtr<SWidget> SLevelViewport::MakeViewportToolbar()
 
 			{
 				FToolMenuSection& RightSection = ViewportToolbarMenu->FindOrAddSection("Right");
+				UE::LevelEditor::AddViewportToolbarPerformanceAndScalabilitySubmenu(RightSection);
 				UE::LevelEditor::AddLevelEditorViewportToolbarSettingsSection(RightSection);
 			}
 		}
diff --git a/Engine/Source/Editor/LevelEditor/Private/ViewportToolbar/LevelEditorViewportToolbarSections.cpp b/Engine/Source/Editor/LevelEditor/Private/ViewportToolbar/LevelEditorViewportToolbarSections.cpp
--- a/Engine/Source/Editor/LevelEditor/Private/ViewportToolbar/LevelEditorViewportToolbarSections.cpp
+++ b/Engine/Source/Editor/LevelEditor/Private/ViewportToolbar/LevelEditorViewportToolbarSections.cpp
@@ -6,7 +6,9 @@
 #include "LevelEditorActions.h"
 #include "LevelViewportActions.h"
 #include "LevelViewportContext.h"
+#include "SCommonEditorViewportToolbarBase.h"
 #include "SLevelViewport.h"
+#include "SScalabilitySettings.h"
 #include "Templates/SharedPointer.h"
 #include "ToolMenu.h"
 #include "ToolMenus.h"
@@ -52,6 +54,11 @@ TOptional<bool> IsDirtyRealtimeWarningFromContext(const FToolMenuContext& Contex
 namespace UE::LevelEditor
 {
 
+bool ShowViewportRealtimeWarning(FLevelEditorViewportClient& ViewportClient)
+{
+	return !ViewportClient.IsRealtime() && !ViewportClient.IsRealtimeOverrideSet() && ViewportClient.IsPerspective();
+}
+
 // TODO: Move this outside the level editor and make it publicly available to anyone building a viewport toolbar.
 void AddViewportToolbarTransformsSection(FToolMenuSection& InSection)
 {
@@ -82,9 +89,21 @@ void AddViewportToolbarTransformsSection(FToolMenuSection& InSection)
 		}));
 }
 
-bool ShowViewportRealtimeWarning(FLevelEditorViewportClient& ViewportClient)
+void AddFeatureLevelPreviewSubmenu(FToolMenuSection& Section)
 {
-	return !ViewportClient.IsRealtime() && !ViewportClient.IsRealtimeOverrideSet() && ViewportClient.IsPerspective();
+	Section.AddSubMenu("FeatureLevelPreview",
+		NSLOCTEXT("LevelToolBarViewMenu", "PreviewPlatformSubMenu", "Preview Platform"),
+		NSLOCTEXT("LevelToolBarViewMenu", "PreviewPlatformSubMenu_ToolTip",
+			"Sets the preview platform used by the main editor"),
+		FNewToolMenuDelegate::CreateLambda([](UToolMenu* InMenu) -> void {
+			FToolMenuSection& Section = InMenu->AddSection(
+				"EditorPreviewMode", LOCTEXT("EditorPreviewModeDevices", "Preview Devices"));
+			// Preview platforms discovered from ITargetPlatforms.
+			for (auto& Item : FLevelEditorCommands::Get().PreviewPlatformOverrides)
+			{
+				Section.AddMenuEntry(Item);
+			}
+		}));
 }
 
 void AddMaterialQualityLevelSubmenu(FToolMenuSection& Section)
@@ -105,19 +124,136 @@ void AddMaterialQualityLevelSubmenu(FToolMenuSection& Section)
 		}));
 }
 
-void AddFeatureLevelPreviewSubmenu(FToolMenuSection& Section)
+void AddViewportToolbarPerformanceAndScalabilitySubmenu(FToolMenuSection& InSection)
 {
-	Section.AddSubMenu("FeatureLevelPreview",
-		NSLOCTEXT("LevelToolBarViewMenu", "PreviewPlatformSubMenu", "Preview Platform"),
-		NSLOCTEXT(
-			"LevelToolBarViewMenu", "PreviewPlatformSubMenu_ToolTip", "Sets the preview platform used by the main editor"),
-		FNewToolMenuDelegate::CreateLambda([](UToolMenu* InMenu) -> void {
-			FToolMenuSection& Section = InMenu->AddSection(
-				"EditorPreviewMode", LOCTEXT("EditorPreviewModeDevices", "Preview Devices"));
-			// Preview platforms discovered from ITargetPlatforms.
-			for (auto& Item : FLevelEditorCommands::Get().PreviewPlatformOverrides)
+	InSection.AddSubMenu("PerformanceAndScalability",
+		LOCTEXT("PerformanceAndScalabilityLabel", "Performance & Scalability"),
+		LOCTEXT("PerformanceAndScalabilityLabel", "Performance and scalability tools tied to this viewport."),
+		FNewToolMenuDelegate::CreateLambda([](UToolMenu* Submenu) -> void {
 			{
-				Section.AddMenuEntry(Item);
+				FToolMenuSection& UnnamedSection = Submenu->FindOrAddSection(NAME_None);
+
+				// Add realtime rendering toggle.
+				UnnamedSection.AddDynamicEntry("ToggleRealtimeDynamicSection",
+					FNewToolMenuSectionDelegate::CreateLambda([](FToolMenuSection& InnerSection) -> void {
+						FToolUIAction RealtimeToggleAction;
+						RealtimeToggleAction.ExecuteAction = FToolMenuExecuteAction::CreateLambda(
+							[](const FToolMenuContext& Context) -> void {
+								ULevelViewportContext* const LevelViewportContext =
+									Context.FindContext<ULevelViewportContext>();
+								if (!LevelViewportContext)
+								{
+									return;
+								}
+
+								if (const TSharedPtr<::SLevelViewport> LevelViewport =
+										LevelViewportContext->LevelViewport.Pin())
+								{
+									LevelViewport->OnToggleRealtime();
+									UToolMenus::Get()->RefreshAllWidgets();
+								}
+							});
+
+						RealtimeToggleAction.GetActionCheckState = FToolMenuGetActionCheckState::CreateLambda(
+							[](const FToolMenuContext& Context) -> ECheckBoxState {
+								ULevelViewportContext* const LevelViewportContext =
+									Context.FindContext<ULevelViewportContext>();
+								if (!LevelViewportContext)
+								{
+									return ECheckBoxState::Undetermined;
+								}
+
+								// Check if the realtime warn state is outdated and if so refresh widgets to update our
+								// top-level status.
+								if (const TOptional<bool> IsDirty = Private::IsDirtyRealtimeWarningFromContext(Context);
+									IsDirty.IsSet() && IsDirty.GetValue())
+								{
+									UToolMenus::Get()->RefreshAllWidgets();
+								}
+
+								if (const TSharedPtr<::SLevelViewport> LevelViewport =
+										LevelViewportContext->LevelViewport.Pin())
+								{
+									return LevelViewport->IsRealtime() ? ECheckBoxState::Checked
+																	   : ECheckBoxState::Unchecked;
+								}
+
+								return ECheckBoxState::Undetermined;
+							});
+
+						bool bDisplayTopLevel = false;
+						if (const TOptional<bool> ShouldWarn = Private::UpdateAndGetRealtimeWarningFromContext(
+								InnerSection.Context);
+							ShouldWarn.IsSet())
+						{
+							bDisplayTopLevel = ShouldWarn.GetValue();
+						}
+						else
+						{
+							// If we couldn't get the warn state, pretend we don't have to warn.
+							bDisplayTopLevel = false;
+						}
+
+						const FText Tooltip =
+							bDisplayTopLevel
+								? LOCTEXT("ToggleRealtimeTooltip_WarnRealtimeOff",
+									"This viewport is not updating in realtime.  Click to turn on realtime mode.")
+								: LOCTEXT("ToggleRealtimeTooltip", "Toggle realtime rendering of the viewport");
+
+						FToolMenuEntry ToggleRealtime = FToolMenuEntry::InitMenuEntry("ToggleRealtime",
+							LOCTEXT("ToggleRealtimeLabel", "Realtime Viewport"), Tooltip,
+							FSlateIcon(FAppStyle::GetAppStyleSetName(), "EditorViewport.ToggleRealTime"),
+							RealtimeToggleAction, EUserInterfaceActionType::ToggleButton);
+						ToggleRealtime.SetShowInToolbarTopLevel(bDisplayTopLevel);
+						InnerSection.AddEntry(ToggleRealtime);
+					}));
+			}
+
+			{
+				FToolMenuSection& PerformanceAndScalabilitySection = Submenu->FindOrAddSection(
+					"PerformanceAndScalability",
+					LOCTEXT("PerformanceAndScalabilitySectionLabel", "Performance & Scalability"));
+
+				AddFeatureLevelPreviewSubmenu(PerformanceAndScalabilitySection);
+
+				PerformanceAndScalabilitySection.AddSeparator("PerformanceAndScalabilitySettings");
+
+				PerformanceAndScalabilitySection.AddSubMenu("Scalability",
+					LOCTEXT("ScalabilitySubMenu", "Engine Scalability"),
+					LOCTEXT("ScalabilitySubMenu_ToolTip", "Open the engine scalability settings"),
+					FNewToolMenuDelegate::CreateLambda([](UToolMenu* InMenu) -> void {
+						FToolMenuSection& Section = InMenu->FindOrAddSection(NAME_None);
+						Section.AddEntry(
+							FToolMenuEntry::InitWidget("ScalabilitySettings", SNew(SScalabilitySettings), FText(), true));
+					}));
+
+				AddMaterialQualityLevelSubmenu(PerformanceAndScalabilitySection);
+
+				// FEditorViewportClient& ViewportClient = Viewport.Pin()->GetLevelViewportClient();
+				PerformanceAndScalabilitySection.AddSubMenu("ScreenPercentageSubMenu",
+					LOCTEXT("ScreenPercentageSubMenu", "Screen Percentage"),
+					LOCTEXT("ScreenPercentageSubMenu_ToolTip", "Customize the viewport's screen percentage"),
+					FNewToolMenuDelegate::CreateLambda([](UToolMenu* ScreenPercentageSubMenu) {
+						FToolMenuSection& UnnamedSection = ScreenPercentageSubMenu->FindOrAddSection(NAME_None);
+
+						ScreenPercentageSubMenu->AddDynamicSection(NAME_None,
+							FNewToolMenuDelegateLegacy::CreateLambda([](FMenuBuilder& MenuBuilder, UToolMenu* InMenu) -> void {
+								ULevelViewportContext* const LevelViewportContext =
+									InMenu->FindContext<ULevelViewportContext>();
+								if (!LevelViewportContext)
+								{
+									return;
+								}
+
+								if (const TSharedPtr<::SLevelViewport> LevelViewport =
+										LevelViewportContext->LevelViewport.Pin())
+								{
+									TSharedPtr<FEditorViewportClient> Client = LevelViewport->GetViewportClient();
+									SCommonEditorViewportToolbarBase::ConstructScreenPercentageMenu(
+										MenuBuilder, Client.Get());
+								}
+							}));
+					}));
 			}
 		}));
 }
@@ -252,81 +388,6 @@ void AddLevelEditorViewportToolbarSettingsSection(FToolMenuSection& InSection)
 		FNewToolMenuDelegate::CreateLambda([](UToolMenu* Submenu) -> void {
 			FToolMenuSection& UnnamedSection = Submenu->FindOrAddSection(NAME_None);
 
-			// Add realtime rendering toggle.
-			UnnamedSection.AddDynamicEntry("ToggleRealtimeDynamicSection",
-				FNewToolMenuSectionDelegate::CreateLambda([](FToolMenuSection& InnerSection) -> void {
-					FToolUIAction RealtimeToggleAction;
-					RealtimeToggleAction.ExecuteAction = FToolMenuExecuteAction::CreateLambda(
-						[](const FToolMenuContext& Context) -> void {
-							ULevelViewportContext* const LevelViewportContext =
-								Context.FindContext<ULevelViewportContext>();
-							if (!LevelViewportContext)
-							{
-								return;
-							}
-
-							if (const TSharedPtr<::SLevelViewport> LevelViewport = LevelViewportContext->LevelViewport.Pin())
-							{
-								LevelViewport->OnToggleRealtime();
-								UToolMenus::Get()->RefreshAllWidgets();
-							}
-						});
-
-					RealtimeToggleAction.GetActionCheckState = FToolMenuGetActionCheckState::CreateLambda(
-						[](const FToolMenuContext& Context) -> ECheckBoxState {
-							ULevelViewportContext* const LevelViewportContext =
-								Context.FindContext<ULevelViewportContext>();
-							if (!LevelViewportContext)
-							{
-								return ECheckBoxState::Undetermined;
-							}
-
-							// Check if the realtime warn state is outdated and if so refresh widgets to update our
-							// top-level status.
-							if (const TOptional<bool> IsDirty = Private::IsDirtyRealtimeWarningFromContext(Context);
-								IsDirty.IsSet() && IsDirty.GetValue())
-							{
-								UToolMenus::Get()->RefreshAllWidgets();
-							}
-
-							if (const TSharedPtr<::SLevelViewport> LevelViewport = LevelViewportContext->LevelViewport.Pin())
-							{
-								return LevelViewport->IsRealtime() ? ECheckBoxState::Checked : ECheckBoxState::Unchecked;
-							}
-
-							return ECheckBoxState::Undetermined;
-						});
-
-					bool bDisplayTopLevel = false;
-					if (const TOptional<bool> ShouldWarn = Private::UpdateAndGetRealtimeWarningFromContext(
-							InnerSection.Context);
-						ShouldWarn.IsSet())
-					{
-						bDisplayTopLevel = ShouldWarn.GetValue();
-					}
-					else
-					{
-						// If we couldn't get the warn state, pretend we don't have to warn.
-						bDisplayTopLevel = false;
-					}
-
-					const FText Tooltip =
-						bDisplayTopLevel
-							? LOCTEXT("ToggleRealtimeTooltip_WarnRealtimeOff",
-								"This viewport is not updating in realtime.  Click to turn on realtime mode.")
-							: LOCTEXT("ToggleRealtimeTooltip", "Toggle realtime rendering of the viewport");
-
-					FToolMenuEntry ToggleRealtime = FToolMenuEntry::InitMenuEntry("ToggleRealtime",
-						LOCTEXT("ToggleRealtimeLabel", "Realtime"), Tooltip,
-						FSlateIcon(FAppStyle::GetAppStyleSetName(), "EditorViewport.ToggleRealTime"),
-						RealtimeToggleAction, EUserInterfaceActionType::ToggleButton);
-					ToggleRealtime.SetShowInToolbarTopLevel(bDisplayTopLevel);
-					InnerSection.AddEntry(ToggleRealtime);
-				}));
-
-			AddMaterialQualityLevelSubmenu(UnnamedSection);
-			AddFeatureLevelPreviewSubmenu(UnnamedSection);
-
 			// Add maximize/restore viewport button.
 			{
 				FToolUIAction MaximizeRestoreAction;
diff --git a/Engine/Source/Editor/LevelEditor/Private/ViewportToolbar/LevelEditorViewportToolbarSections.h b/Engine/Source/Editor/LevelEditor/Private/ViewportToolbar/LevelEditorViewportToolbarSections.h
--- a/Engine/Source/Editor/LevelEditor/Private/ViewportToolbar/LevelEditorViewportToolbarSections.h
+++ b/Engine/Source/Editor/LevelEditor/Private/ViewportToolbar/LevelEditorViewportToolbarSections.h
@@ -12,10 +12,14 @@ struct FToolMenuSection;
 namespace UE::LevelEditor
 {
 
-void AddViewportToolbarTransformsSection(FToolMenuSection& InSection);
 bool ShowViewportRealtimeWarning(FLevelEditorViewportClient& ViewportClient);
-void AddMaterialQualityLevelSubmenu(FToolMenuSection& Section);
+
+void AddViewportToolbarTransformsSection(FToolMenuSection& InSection);
+
 void AddFeatureLevelPreviewSubmenu(FToolMenuSection& Section);
+void AddMaterialQualityLevelSubmenu(FToolMenuSection& Section);
+void AddViewportToolbarPerformanceAndScalabilitySubmenu(FToolMenuSection& InSection);
+
 void GenerateViewportLayoutsMenu(UToolMenu* InMenu, TSharedPtr<::SLevelViewport> InViewport);
 void AddLevelEditorViewportToolbarSettingsSection(FToolMenuSection& InSection);
 
