Update the "Settings" submenu of the new viewport toolbar

diff --git a/Engine/Source/Editor/LevelEditor/Private/LevelEditorToolBar.cpp b/Engine/Source/Editor/LevelEditor/Private/LevelEditorToolBar.cpp
--- a/Engine/Source/Editor/LevelEditor/Private/LevelEditorToolBar.cpp
+++ b/Engine/Source/Editor/LevelEditor/Private/LevelEditorToolBar.cpp
@@ -1755,22 +1755,9 @@ void FLevelEditorToolBar::RegisterQuickSettingsMenu()
 
 	{
 		FToolMenuSection& Section = Menu->AddSection("LevelEditorAudio", LOCTEXT("AudioHeading", "Real Time Audio") );
-		TSharedRef<SWidget> VolumeItem = SNew(SHorizontalBox)
-											+SHorizontalBox::Slot()
-											.FillWidth(0.9f)
-											.Padding( FMargin(2.0f, 0.0f, 0.0f, 0.0f) )
-											[
-												SNew(SVolumeControl)
-												.ToolTipText_Static(&FLevelEditorActionCallbacks::GetAudioVolumeToolTip)
-												.Volume_Static(&FLevelEditorActionCallbacks::GetAudioVolume)
-												.OnVolumeChanged_Static(&FLevelEditorActionCallbacks::OnAudioVolumeChanged)
-												.Muted_Static(&FLevelEditorActionCallbacks::GetAudioMuted)
-												.OnMuteChanged_Static(&FLevelEditorActionCallbacks::OnAudioMutedChanged)
-											]
-											+SHorizontalBox::Slot()
-											.FillWidth(0.1f);
-
-		Section.AddEntry(FToolMenuEntry::InitWidget("Volume", VolumeItem, LOCTEXT("VolumeControlLabel","Volume")));
+
+		Section.AddEntry(FToolMenuEntry::InitWidget(
+			"Volume", UE::LevelEditor::BuildVolumeControlCustomWidget(), LOCTEXT("VolumeControlLabel", "Volume")));
 	}
 
 	{
diff --git a/Engine/Source/Editor/LevelEditor/Private/SLevelViewport.cpp b/Engine/Source/Editor/LevelEditor/Private/SLevelViewport.cpp
--- a/Engine/Source/Editor/LevelEditor/Private/SLevelViewport.cpp
+++ b/Engine/Source/Editor/LevelEditor/Private/SLevelViewport.cpp
@@ -1934,7 +1934,7 @@ TSharedPtr<SWidget> SLevelViewport::MakeViewportToolbar()
 			{
 				FToolMenuSection& RightSection = ViewportToolbarMenu->FindOrAddSection("Right");
 				UE::LevelEditor::AddViewportToolbarPerformanceAndScalabilitySubmenu(RightSection);
-				UE::LevelEditor::AddLevelEditorViewportToolbarSettingsSection(RightSection);
+				UE::LevelEditor::AddLevelEditorViewportToolbarSettingsSubmenu(RightSection);
 			}
 		}
 	}
diff --git a/Engine/Source/Editor/LevelEditor/Private/ViewportToolbar/LevelEditorViewportToolbarSections.cpp b/Engine/Source/Editor/LevelEditor/Private/ViewportToolbar/LevelEditorViewportToolbarSections.cpp
--- a/Engine/Source/Editor/LevelEditor/Private/ViewportToolbar/LevelEditorViewportToolbarSections.cpp
+++ b/Engine/Source/Editor/LevelEditor/Private/ViewportToolbar/LevelEditorViewportToolbarSections.cpp
@@ -12,6 +12,8 @@
 #include "Templates/SharedPointer.h"
 #include "ToolMenu.h"
 #include "ToolMenus.h"
+#include "Widgets/Input/SVolumeControl.h"
+#include "Widgets/SBoxPanel.h"
 
 #define LOCTEXT_NAMESPACE "LevelEditorViewportToolbar"
 
@@ -381,19 +383,37 @@ void GenerateViewportLayoutsMenu(UToolMenu* InMenu, TSharedPtr<::SLevelViewport>
 	}
 }
 
-void AddLevelEditorViewportToolbarSettingsSection(FToolMenuSection& InSection)
+TSharedRef<SWidget> BuildVolumeControlCustomWidget()
+{
+	return SNew(SHorizontalBox)
+		+ SHorizontalBox::Slot()
+		.FillWidth(0.9f)
+		.Padding(FMargin(2.0f, 0.0f, 0.0f, 0.0f))
+		[
+		SNew(SVolumeControl)
+			.ToolTipText_Static(&FLevelEditorActionCallbacks::GetAudioVolumeToolTip)
+			.Volume_Static(&FLevelEditorActionCallbacks::GetAudioVolume)
+			.OnVolumeChanged_Static(&FLevelEditorActionCallbacks::OnAudioVolumeChanged)
+			.Muted_Static(&FLevelEditorActionCallbacks::GetAudioMuted)
+			.OnMuteChanged_Static(&FLevelEditorActionCallbacks::OnAudioMutedChanged)
+		]
+		+ SHorizontalBox::Slot()
+		.FillWidth(0.1f);
+}
+
+void AddLevelEditorViewportToolbarSettingsSubmenu(FToolMenuSection& InSection)
 {
 	InSection.AddSubMenu("Settings", LOCTEXT("SettingsSubmenuLabel", "Settings"),
 		LOCTEXT("SettingsSubmenuTooltip", "Viewport-related settings"),
 		FNewToolMenuDelegate::CreateLambda([](UToolMenu* Submenu) -> void {
-			FToolMenuSection& UnnamedSection = Submenu->FindOrAddSection(NAME_None);
-
-			// Add maximize/restore viewport button.
 			{
-				FToolUIAction MaximizeRestoreAction;
-				MaximizeRestoreAction.ExecuteAction = FToolMenuExecuteAction::CreateLambda(
-					[](const FToolMenuContext& Context) {
-						ULevelViewportContext* const LevelViewportContext = Context.FindContext<ULevelViewportContext>();
+				FToolMenuSection& ViewportControlsSection = Submenu->FindOrAddSection(
+					"ViewportControls", LOCTEXT("ViewportControlsSectionLabel", "Viewport Controls"));
+
+				ViewportControlsSection.AddSubMenu("ViewportLayouts", LOCTEXT("ViewportLayoutsLabel", "Layouts"),
+					LOCTEXT("ViewportLayoutsTooltip", "Configure the layouts of the viewport windows"),
+					FNewToolMenuDelegate::CreateLambda([](UToolMenu* InMenu) {
+						ULevelViewportContext* const LevelViewportContext = InMenu->FindContext<ULevelViewportContext>();
 						if (!LevelViewportContext)
 						{
 							return;
@@ -401,51 +421,69 @@ void AddLevelEditorViewportToolbarSettingsSection(FToolMenuSection& InSection)
 
 						if (const TSharedPtr<::SLevelViewport> LevelViewport = LevelViewportContext->LevelViewport.Pin())
 						{
-							LevelViewport->OnToggleMaximize();
+							GenerateViewportLayoutsMenu(InMenu, LevelViewport);
 						}
-					});
-				MaximizeRestoreAction.GetActionCheckState = FToolMenuGetActionCheckState::CreateLambda(
-					[](const FToolMenuContext& Context) -> ECheckBoxState {
+					}),
+					false, FSlateIcon(FAppStyle::Get().GetStyleSetName(), "Icons.Layout"));
+			}
+
+			{
+				FToolMenuSection& SettingsSection = Submenu->FindOrAddSection(
+					"Settings", LOCTEXT("SettingsSectionLabel", "Settings"));
+
+				SettingsSection.AddEntry(FToolMenuEntry::InitWidget(
+					"Volume", BuildVolumeControlCustomWidget(), LOCTEXT("VolumeControlLabel", "Volume")));
+
+				SettingsSection.AddSeparator("ViewportSizeSeparator");
+
+				SettingsSection.AddMenuEntry(FLevelViewportCommands::Get().ToggleImmersive);
+
+				{
+					FToolUIAction MaximizeRestoreAction;
+					MaximizeRestoreAction.ExecuteAction = FToolMenuExecuteAction::CreateLambda([](const FToolMenuContext& Context) {
 						ULevelViewportContext* const LevelViewportContext = Context.FindContext<ULevelViewportContext>();
 						if (!LevelViewportContext)
 						{
-							return ECheckBoxState::Undetermined;
+							return;
 						}
 
 						if (const TSharedPtr<::SLevelViewport> LevelViewport = LevelViewportContext->LevelViewport.Pin())
 						{
-							return LevelViewport->IsMaximized() ? ECheckBoxState::Checked : ECheckBoxState::Unchecked;
+							LevelViewport->OnToggleMaximize();
 						}
-
-						return ECheckBoxState::Undetermined;
 					});
+					MaximizeRestoreAction.GetActionCheckState = FToolMenuGetActionCheckState::CreateLambda(
+						[](const FToolMenuContext& Context) -> ECheckBoxState {
+							ULevelViewportContext* const LevelViewportContext =
+								Context.FindContext<ULevelViewportContext>();
+							if (!LevelViewportContext)
+							{
+								return ECheckBoxState::Undetermined;
+							}
 
-				UnnamedSection
-					.AddMenuEntry("MaximizeRestore", LOCTEXT("MaximizeRestoreLabel", "Maximize/restore"),
-						LOCTEXT("MaximizeRestoreTooltip", "Maximizes or restores this viewport"),
-						FSlateIcon(FAppStyle::GetAppStyleSetName(), "EditorViewportToolBar.Maximize.Normal"),
-						MaximizeRestoreAction, EUserInterfaceActionType::ToggleButton)
-					.SetShowInToolbarTopLevel(true);
-			}
+							if (const TSharedPtr<::SLevelViewport> LevelViewport = LevelViewportContext->LevelViewport.Pin())
+							{
+								return LevelViewport->IsMaximized() ? ECheckBoxState::Checked : ECheckBoxState::Unchecked;
+							}
 
-			// Add immersive mode toggle.
-			UnnamedSection.AddMenuEntry(FLevelViewportCommands::Get().ToggleImmersive);
-
-			UnnamedSection.AddSubMenu("ViewportLayouts", LOCTEXT("ViewportLayoutsLabel", "Layouts"),
-				LOCTEXT("ViewportLayoutsTooltip", "Configure the layouts of the viewport windows"),
-				FNewToolMenuDelegate::CreateLambda([](UToolMenu* InMenu) {
-					ULevelViewportContext* const LevelViewportContext = InMenu->FindContext<ULevelViewportContext>();
-					if (!LevelViewportContext)
-					{
-						return;
-					}
-
-					if (const TSharedPtr<::SLevelViewport> LevelViewport = LevelViewportContext->LevelViewport.Pin())
-					{
-						GenerateViewportLayoutsMenu(InMenu, LevelViewport);
-					}
-				}),
-				false, FSlateIcon(FAppStyle::Get().GetStyleSetName(), "Icons.Layout"));
+							return ECheckBoxState::Undetermined;
+						});
+
+					SettingsSection
+						.AddMenuEntry("MaximizeRestore", LOCTEXT("MaximizeRestoreLabel", "Maximize Viewport"),
+							LOCTEXT("MaximizeRestoreTooltip", "Maximizes or restores this viewport"),
+							FSlateIcon(FAppStyle::GetAppStyleSetName(), "EditorViewportToolBar.Maximize.Normal"),
+							MaximizeRestoreAction, EUserInterfaceActionType::ToggleButton)
+						.SetShowInToolbarTopLevel(true);
+				}
+
+				SettingsSection.AddSeparator("AdvancedSeparator");
+
+				{
+					const FLevelViewportCommands& LevelViewportActions = FLevelViewportCommands::Get();
+					SettingsSection.AddMenuEntry(LevelViewportActions.AdvancedSettings);
+				}
+			}
 		}));
 }
 
diff --git a/Engine/Source/Editor/LevelEditor/Private/ViewportToolbar/LevelEditorViewportToolbarSections.h b/Engine/Source/Editor/LevelEditor/Private/ViewportToolbar/LevelEditorViewportToolbarSections.h
--- a/Engine/Source/Editor/LevelEditor/Private/ViewportToolbar/LevelEditorViewportToolbarSections.h
+++ b/Engine/Source/Editor/LevelEditor/Private/ViewportToolbar/LevelEditorViewportToolbarSections.h
@@ -6,6 +6,7 @@
 
 class FLevelEditorViewportClient;
 class SLevelViewport;
+class SWidget;
 class UToolMenu;
 struct FToolMenuSection;
 
@@ -21,6 +22,7 @@ void AddMaterialQualityLevelSubmenu(FToolMenuSection& Section);
 void AddViewportToolbarPerformanceAndScalabilitySubmenu(FToolMenuSection& InSection);
 
 void GenerateViewportLayoutsMenu(UToolMenu* InMenu, TSharedPtr<::SLevelViewport> InViewport);
-void AddLevelEditorViewportToolbarSettingsSection(FToolMenuSection& InSection);
+TSharedRef<SWidget> BuildVolumeControlCustomWidget();
+void AddLevelEditorViewportToolbarSettingsSubmenu(FToolMenuSection& InSection);
 
 } // namespace UE::LevelEditor
