Support old legacy menu extenders in the new level editor viewport toolbar

The "View Modes" menu's "Visualizers" entry is an example of one that needs this.

diff --git a/Engine/Source/Editor/LevelEditor/Private/SLevelViewport.cpp b/Engine/Source/Editor/LevelEditor/Private/SLevelViewport.cpp
--- a/Engine/Source/Editor/LevelEditor/Private/SLevelViewport.cpp
+++ b/Engine/Source/Editor/LevelEditor/Private/SLevelViewport.cpp
@@ -1942,6 +1942,9 @@ TSharedPtr<SWidget> SLevelViewport::MakeViewportToolbar()
 	FToolMenuContext ViewportToolbarContext;
 	{
 		ViewportToolbarContext.AppendCommandList(GetCommandList());
+		// Stay backward-compatible with legacy view menu extenders. Note that these extenders can now leak between
+		// submenus of the viewport toolbar.
+		ViewportToolbarContext.AddExtender(UE::LevelEditor::GetViewModesLegacyExtenders());
 
 		ULevelViewportContext* const ContextObject = NewObject<ULevelViewportContext>();
 		ContextObject->LevelViewport = SharedThis(this);
diff --git a/Engine/Source/Editor/LevelEditor/Private/SLevelViewportToolBar.cpp b/Engine/Source/Editor/LevelEditor/Private/SLevelViewportToolBar.cpp
--- a/Engine/Source/Editor/LevelEditor/Private/SLevelViewportToolBar.cpp
+++ b/Engine/Source/Editor/LevelEditor/Private/SLevelViewportToolBar.cpp
@@ -241,7 +241,7 @@ void SLevelViewportToolBar::Construct( const FArguments& InArgs )
 					.Padding( ToolbarSlotPadding )
 					[
 						SNew( SLevelEditorViewportViewMenu, ViewportRef, SharedThis(this) )
-						.MenuExtenders(GetViewMenuExtender())
+						.MenuExtenders(UE::LevelEditor::GetViewModesLegacyExtenders())
 						.AddMetaData<FTagMetaData>(FTagMetaData(TEXT("ViewMenuButton")))
 					]
 					+ SHorizontalBox::Slot()
@@ -1573,12 +1573,6 @@ TWeakObjectPtr<UWorld> SLevelViewportToolBar::GetWorld() const
 	return NULL;
 }
 
-TSharedPtr<FExtender> SLevelViewportToolBar::GetViewMenuExtender()
-{
-	FLevelEditorModule& LevelEditorModule = FModuleManager::GetModuleChecked<FLevelEditorModule>(TEXT("LevelEditor"));
-	return LevelEditorModule.GetMenuExtensibilityManager()->GetAllExtenders();
-}
-
 void SLevelViewportToolBar::FillViewMenu(UToolMenu* Menu)
 {
 	FToolMenuInsert InsertPosition("ViewMode", EToolMenuInsertType::After);
@@ -1971,7 +1965,6 @@ void SLevelViewportToolBar::OnLandscapeLODChanged(int32 NewValue)
 	ViewportClient.Invalidate();
 }
 
-
 FReply SLevelViewportToolBar::OnRealtimeWarningClicked()
 {
 	FLevelEditorViewportClient& ViewportClient = Viewport.Pin()->GetLevelViewportClient();
diff --git a/Engine/Source/Editor/LevelEditor/Private/SLevelViewportToolBar.h b/Engine/Source/Editor/LevelEditor/Private/SLevelViewportToolBar.h
--- a/Engine/Source/Editor/LevelEditor/Private/SLevelViewportToolBar.h
+++ b/Engine/Source/Editor/LevelEditor/Private/SLevelViewportToolBar.h
@@ -10,7 +10,6 @@
 #include "SViewportToolBar.h"
 
 class ACameraActor;
-class FExtender;
 class SExtensionPanel;
 class SActionableMessageViewportWidget;
 class STransformViewportToolBar;
@@ -268,9 +267,6 @@ private:
 	/** Gets the world we are editing */
 	TWeakObjectPtr<UWorld> GetWorld() const;
 
-	/** Gets the extender for the view menu */
-	TSharedPtr<FExtender> GetViewMenuExtender();
-
 	/** Called when the user disables realtime override from the toolbar */
 	void OnDisableRealtimeOverride();
 	bool IsRealtimeOverrideToggleVisible() const;
diff --git a/Engine/Source/Editor/LevelEditor/Private/ViewportToolbar/LevelEditorViewportToolbarSections.cpp b/Engine/Source/Editor/LevelEditor/Private/ViewportToolbar/LevelEditorViewportToolbarSections.cpp
--- a/Engine/Source/Editor/LevelEditor/Private/ViewportToolbar/LevelEditorViewportToolbarSections.cpp
+++ b/Engine/Source/Editor/LevelEditor/Private/ViewportToolbar/LevelEditorViewportToolbarSections.cpp
@@ -3,6 +3,7 @@
 #include "ViewportToolbar/LevelEditorViewportToolbarSections.h"
 
 #include "EditorViewportCommands.h"
+#include "LevelEditor.h"
 #include "LevelEditorActions.h"
 #include "LevelViewportActions.h"
 #include "LevelViewportContext.h"
@@ -91,6 +92,12 @@ void AddViewportToolbarTransformsSection(FToolMenuSection& InSection)
 		}));
 }
 
+TSharedPtr<FExtender> GetViewModesLegacyExtenders()
+{
+	FLevelEditorModule& LevelEditorModule = FModuleManager::GetModuleChecked<FLevelEditorModule>(TEXT("LevelEditor"));
+	return LevelEditorModule.GetMenuExtensibilityManager()->GetAllExtenders();
+}
+
 void AddFeatureLevelPreviewSubmenu(FToolMenuSection& Section)
 {
 	Section.AddSubMenu("FeatureLevelPreview",
diff --git a/Engine/Source/Editor/LevelEditor/Private/ViewportToolbar/LevelEditorViewportToolbarSections.h b/Engine/Source/Editor/LevelEditor/Private/ViewportToolbar/LevelEditorViewportToolbarSections.h
--- a/Engine/Source/Editor/LevelEditor/Private/ViewportToolbar/LevelEditorViewportToolbarSections.h
+++ b/Engine/Source/Editor/LevelEditor/Private/ViewportToolbar/LevelEditorViewportToolbarSections.h
@@ -4,6 +4,7 @@
 
 #include "Templates/SharedPointerFwd.h"
 
+class FExtender;
 class FLevelEditorViewportClient;
 class SLevelViewport;
 class SWidget;
@@ -17,6 +18,8 @@ bool ShowViewportRealtimeWarning(FLevelEditorViewportClient& ViewportClient);
 
 void AddViewportToolbarTransformsSection(FToolMenuSection& InSection);
 
+TSharedPtr<FExtender> GetViewModesLegacyExtenders();
+
 void AddFeatureLevelPreviewSubmenu(FToolMenuSection& Section);
 void AddMaterialQualityLevelSubmenu(FToolMenuSection& Section);
 void AddViewportToolbarPerformanceAndScalabilitySubmenu(FToolMenuSection& InSection);
