Fix duplicated LOCTEXT key

diff --git a/Engine/Source/Editor/LevelEditor/Private/ViewportToolbar/LevelEditorViewportToolbarSections.cpp b/Engine/Source/Editor/LevelEditor/Private/ViewportToolbar/LevelEditorViewportToolbarSections.cpp
--- a/Engine/Source/Editor/LevelEditor/Private/ViewportToolbar/LevelEditorViewportToolbarSections.cpp
+++ b/Engine/Source/Editor/LevelEditor/Private/ViewportToolbar/LevelEditorViewportToolbarSections.cpp
@@ -137,7 +137,7 @@ void AddViewportToolbarPerformanceAndScalabilitySubmenu(FToolMenuSection& InSect
 {
 	InSection.AddSubMenu("PerformanceAndScalability",
 		LOCTEXT("PerformanceAndScalabilityLabel", "Performance & Scalability"),
-		LOCTEXT("PerformanceAndScalabilityLabel", "Performance and scalability tools tied to this viewport."),
+		LOCTEXT("PerformanceAndScalabilityTooltip", "Performance and scalability tools tied to this viewport."),
 		FNewToolMenuDelegate::CreateLambda([](UToolMenu* Submenu) -> void {
 			{
 				FToolMenuSection& UnnamedSection = Submenu->FindOrAddSection(NAME_None);
