Fix crash when the ToolMenu LevelEditor.LevelViewportToolBar.View is opened (or used as a parent) without its context object

diff --git a/Engine/Source/Editor/LevelEditor/Private/SLevelViewportToolBar.cpp b/Engine/Source/Editor/LevelEditor/Private/SLevelViewportToolBar.cpp
--- a/Engine/Source/Editor/LevelEditor/Private/SLevelViewportToolBar.cpp
+++ b/Engine/Source/Editor/LevelEditor/Private/SLevelViewportToolBar.cpp
@@ -79,11 +79,13 @@ public:
 		if (!UToolMenus::Get()->IsMenuRegistered("LevelEditor.LevelViewportToolBar.View"))
 		{
 			UToolMenu* Menu = UToolMenus::Get()->RegisterMenu("LevelEditor.LevelViewportToolBar.View", "UnrealEd.ViewportToolbar.View");
-			Menu->AddDynamicSection("LevelSection", FNewToolMenuDelegate::CreateLambda([](UToolMenu* InMenu)
-			{
-				UEditorViewportViewMenuContext* Context = InMenu->FindContext<UEditorViewportViewMenuContext>();
-				TSharedPtr<SLevelViewportToolBar> LevelViewportToolBar = StaticCastSharedPtr<SLevelViewportToolBar>(Context->EditorViewportViewMenu.Pin()->GetParentToolBar().Pin());
-				LevelViewportToolBar->FillViewMenu(InMenu);
+			Menu->AddDynamicSection("LevelSection", FNewToolMenuDelegate::CreateLambda([](UToolMenu* InMenu) {
+				if (UEditorViewportViewMenuContext* Context = InMenu->FindContext<UEditorViewportViewMenuContext>())
+				{
+					TSharedPtr<SLevelViewportToolBar> LevelViewportToolBar = StaticCastSharedPtr<SLevelViewportToolBar>(
+						Context->EditorViewportViewMenu.Pin()->GetParentToolBar().Pin());
+					LevelViewportToolBar->FillViewMenu(InMenu);
+				}
 			}));
 		}
 	}
