Extract level editor-specific code for populating the viewport toolbar "View Modes" menu

Also rename UE::UnrealEd::FillViewMenu to UE::UnrealEd::PopulateViewModesMenu to make it match the new UE::LevelEditor::PopulateViewModesMenu.

diff --git a/Engine/Source/Editor/LevelEditor/Private/SLevelViewportToolBar.cpp b/Engine/Source/Editor/LevelEditor/Private/SLevelViewportToolBar.cpp
--- a/Engine/Source/Editor/LevelEditor/Private/SLevelViewportToolBar.cpp
+++ b/Engine/Source/Editor/LevelEditor/Private/SLevelViewportToolBar.cpp
@@ -1577,310 +1577,9 @@ TWeakObjectPtr<UWorld> SLevelViewportToolBar::GetWorld() const
 
 void SLevelViewportToolBar::FillViewMenu(UToolMenu* Menu)
 {
-	FToolMenuInsert InsertPosition("ViewMode", EToolMenuInsertType::After);
-
-	{
-		FToolMenuSection& Section = Menu->AddSection("LevelViewportDeferredRendering", LOCTEXT("DeferredRenderingHeader", "Deferred Rendering"), InsertPosition);
-	}
-
-	{
-		FToolMenuSection& Section = Menu->FindOrAddSection("ViewMode");
-		Section.AddSubMenu(
-			"VisualizeBufferViewMode",
-			LOCTEXT("VisualizeBufferViewModeDisplayName", "Buffer Visualization"),
-			LOCTEXT("BufferVisualizationMenu_ToolTip", "Select a mode for buffer visualization"),
-			FNewMenuDelegate::CreateStatic(&FBufferVisualizationMenuCommands::BuildVisualisationSubMenu),
-			FUIAction(
-				FExecuteAction(),
-				FCanExecuteAction(),
-				FIsActionChecked::CreateLambda([this]()
-				{
-					const TSharedRef<SEditorViewport> ViewportRef = Viewport.Pin().ToSharedRef();
-					const TSharedPtr<FEditorViewportClient> ViewportClient = ViewportRef->GetViewportClient();
-					check(ViewportClient.IsValid());
-					return ViewportClient->IsViewModeEnabled(VMI_VisualizeBuffer);
-				})
-			),
-			EUserInterfaceActionType::RadioButton,
-			/* bInOpenSubMenuOnClick = */ false,
-			FSlateIcon(FAppStyle::GetAppStyleSetName(), "EditorViewport.VisualizeBufferMode")
-		);
-	}
-
-	{
-		FToolMenuSection& Section = Menu->FindOrAddSection("ViewMode");
-		Section.AddSubMenu(
-			"VisualizeNaniteViewMode",
-			LOCTEXT("VisualizeNaniteViewModeDisplayName", "Nanite Visualization"),
-			LOCTEXT("NaniteVisualizationMenu_ToolTip", "Select a mode for Nanite visualization"),
-			FNewMenuDelegate::CreateStatic(&FNaniteVisualizationMenuCommands::BuildVisualisationSubMenu),
-			FUIAction(
-				FExecuteAction(),
-				FCanExecuteAction(),
-				FIsActionChecked::CreateLambda([this]()
-				{
-					const TSharedRef<SEditorViewport> ViewportRef = Viewport.Pin().ToSharedRef();
-					const TSharedPtr<FEditorViewportClient> ViewportClient = ViewportRef->GetViewportClient();
-					check(ViewportClient.IsValid());
-					return ViewportClient->IsViewModeEnabled(VMI_VisualizeNanite);
-				})
-			),
-			EUserInterfaceActionType::RadioButton,
-			/* bInOpenSubMenuOnClick = */ false,
-			FSlateIcon(FAppStyle::GetAppStyleSetName(), "EditorViewport.VisualizeNaniteMode")
-		);
-	}
-
-	{
-		FToolMenuSection& Section = Menu->FindOrAddSection("ViewMode");
-		Section.AddSubMenu(
-			"VisualizeLumenViewMode",
-			LOCTEXT("VisualizeLumenViewModeDisplayName", "Lumen"),
-			LOCTEXT("LumenVisualizationMenu_ToolTip", "Select a mode for Lumen visualization"),
-			FNewMenuDelegate::CreateStatic(&FLumenVisualizationMenuCommands::BuildVisualisationSubMenu),
-			FUIAction(
-				FExecuteAction(),
-				FCanExecuteAction(),
-				FIsActionChecked::CreateLambda([this]()
-					{
-						const TSharedRef<SEditorViewport> ViewportRef = Viewport.Pin().ToSharedRef();
-						const TSharedPtr<FEditorViewportClient> ViewportClient = ViewportRef->GetViewportClient();
-						check(ViewportClient.IsValid());
-						return ViewportClient->IsViewModeEnabled(VMI_VisualizeLumen);
-					})
-			),
-			EUserInterfaceActionType::RadioButton,
-						/* bInOpenSubMenuOnClick = */ false,
-						FSlateIcon(FAppStyle::GetAppStyleSetName(), "EditorViewport.VisualizeLumenMode")
-						);
-	}
-
-	if (Substrate::IsSubstrateEnabled())
-	{
-		FToolMenuSection& Section = Menu->FindOrAddSection("ViewMode");
-		Section.AddSubMenu(
-			"VisualizeSubstrateViewMode",
-			LOCTEXT("VisualizeSubstrateViewModeDisplayName", "Substrate"),
-			LOCTEXT("SubstrateVisualizationMenu_ToolTip", "Select a mode for Substrate visualization"),
-			FNewMenuDelegate::CreateStatic(&FSubstrateVisualizationMenuCommands::BuildVisualisationSubMenu),
-			FUIAction(
-				FExecuteAction(),
-				FCanExecuteAction(),
-				FIsActionChecked::CreateLambda([this]()
-					{
-						const TSharedRef<SEditorViewport> ViewportRef = Viewport.Pin().ToSharedRef();
-						const TSharedPtr<FEditorViewportClient> ViewportClient = ViewportRef->GetViewportClient();
-						check(ViewportClient.IsValid());
-						return ViewportClient->IsViewModeEnabled(VMI_VisualizeSubstrate);
-					})
-			),
-			EUserInterfaceActionType::RadioButton,
-						/* bInOpenSubMenuOnClick = */ false,
-						FSlateIcon(FAppStyle::GetAppStyleSetName(), "EditorViewport.VisualizeSubstrateMode")
-						);
-	}
-
-	if (IsGroomEnabled())
-	{
-		FToolMenuSection& Section = Menu->FindOrAddSection("ViewMode");
-		Section.AddSubMenu(
-			"VisualizeGroomViewMode",
-			LOCTEXT("VisualizeGroomViewModeDisplayName", "Groom"),
-			LOCTEXT("GroomVisualizationMenu_ToolTip", "Select a mode for Groom visualization"),
-			FNewMenuDelegate::CreateStatic(&FGroomVisualizationMenuCommands::BuildVisualisationSubMenu),
-			FUIAction(
-				FExecuteAction(),
-				FCanExecuteAction(),
-				FIsActionChecked::CreateLambda([this]()
-					{
-						const TSharedRef<SEditorViewport> ViewportRef = Viewport.Pin().ToSharedRef();
-						const TSharedPtr<FEditorViewportClient> ViewportClient = ViewportRef->GetViewportClient();
-						check(ViewportClient.IsValid());
-						return ViewportClient->IsViewModeEnabled(VMI_VisualizeGroom);
-					})
-			),
-			EUserInterfaceActionType::RadioButton,
-						/* bInOpenSubMenuOnClick = */ false,
-						FSlateIcon(FAppStyle::GetAppStyleSetName(), "EditorViewport.VisualizeGroomMode")
-						);
-	}
-
-	{
-		FToolMenuSection& Section = Menu->FindOrAddSection("ViewMode");
-		Section.AddSubMenu(
-			"VisualizeVirtualShadowMapViewMode",
-			LOCTEXT("VisualizeVirtualShadowMapViewModeDisplayName", "Virtual Shadow Map"),
-			LOCTEXT("VirtualShadowMapVisualizationMenu_ToolTip", "Select a mode for virtual shadow map visualization. Select a light component in the world outliner to visualize that light."),
-			FNewMenuDelegate::CreateStatic(&FVirtualShadowMapVisualizationMenuCommands::BuildVisualisationSubMenu),
-			FUIAction(
-				FExecuteAction(),
-				FCanExecuteAction(),
-				FIsActionChecked::CreateLambda([this]()
-					{
-						const TSharedRef<SEditorViewport> ViewportRef = Viewport.Pin().ToSharedRef();
-						const TSharedPtr<FEditorViewportClient> ViewportClient = ViewportRef->GetViewportClient();
-						check(ViewportClient.IsValid());
-						return ViewportClient->IsViewModeEnabled(VMI_VisualizeVirtualShadowMap);
-					})
-			),
-			EUserInterfaceActionType::RadioButton,
-			/* bInOpenSubMenuOnClick = */ false,
-			FSlateIcon(FAppStyle::GetAppStyleSetName(), "EditorViewport.VisualizeVirtualShadowMapMode")
-			);
-	}
-
+	if (TSharedPtr<SLevelViewport> ViewportPinned = Viewport.Pin())
 	{
-		auto BuildActorColorationMenu = [this](UToolMenu* Menu, SLevelViewportToolBar* Toolbar)
-		{
-			FToolMenuSection& SubMenuSection = Menu->AddSection("LevelViewportActorColoration", LOCTEXT("ActorColorationHeader", "Actor Coloration"));
-
-			TArray<FActorPrimitiveColorHandler::FPrimitiveColorHandler> PrimitiveColorHandlers;
-			FActorPrimitiveColorHandler::Get().GetRegisteredPrimitiveColorHandlers(PrimitiveColorHandlers);
-
-			for (const FActorPrimitiveColorHandler::FPrimitiveColorHandler& PrimitiveColorHandler : PrimitiveColorHandlers)
-			{
-				if (!PrimitiveColorHandler.bAvailalbleInEditor)
-				{
-					continue;
-				}
-
-				SubMenuSection.AddMenuEntry(
-					NAME_None,
-					PrimitiveColorHandler.HandlerText,
-					FText(),
-					FSlateIcon(),
-					FUIAction(
-						FExecuteAction::CreateLambda([this, PrimitiveColorHandler]()
-						{
-							if (FLevelEditorViewportClient* ViewportClient = GetLevelViewportClient())
-							{
-								const bool bActorColorationEnabled = ViewportClient->HandleIsShowFlagEnabled(FEngineShowFlags::EShowFlag::SF_ActorColoration);
-
-								if (PrimitiveColorHandler.HandlerName.IsNone())
-								{
-									if (bActorColorationEnabled)
-									{
-										ViewportClient->HandleToggleShowFlag(FEngineShowFlags::EShowFlag::SF_ActorColoration);
-									}
-								}
-								else
-								{
-									if (!bActorColorationEnabled)
-									{
-										ViewportClient->HandleToggleShowFlag(FEngineShowFlags::EShowFlag::SF_ActorColoration);
-									}
-
-									FActorPrimitiveColorHandler::Get().SetActivePrimitiveColorHandler(PrimitiveColorHandler.HandlerName, GWorld);
-								}
-							}
-						}),
-						FCanExecuteAction::CreateLambda([this]()
-						{
-							if (FLevelEditorViewportClient* ViewportClient = GetLevelViewportClient())
-							{
-								return true;
-							}
-							return false;
-						}),
-						FGetActionCheckState::CreateLambda([this, PrimitiveColorHandler]()
-						{
-							if (FLevelEditorViewportClient* ViewportClient = GetLevelViewportClient())
-							{
-								const bool bActorColorationEnabled = ViewportClient->HandleIsShowFlagEnabled(FEngineShowFlags::EShowFlag::SF_ActorColoration);
-
-								if (PrimitiveColorHandler.HandlerName.IsNone())
-								{
-									return bActorColorationEnabled ? ECheckBoxState::Unchecked : ECheckBoxState::Checked;
-								}
-								else
-								{
-									if (bActorColorationEnabled)
-									{
-										return FActorPrimitiveColorHandler::Get().GetActivePrimitiveColorHandler() == PrimitiveColorHandler.HandlerName ? ECheckBoxState::Checked : ECheckBoxState::Unchecked;
-									}
-								}								
-							}
-
-							return ECheckBoxState::Unchecked;
-						})
-					),
-					EUserInterfaceActionType::RadioButton
-				);
-			}
-		};
-
-		FToolMenuSection& Section = Menu->FindOrAddSection("ViewMode");
-		Section.AddSubMenu(
-			"ActorColoration",
-			LOCTEXT("ActorColorationDisplayName", "Actor Coloration"),
-			LOCTEXT("ActorColorationMenu_ToolTip", "Override Actor Coloration mode"),
-			FNewToolMenuDelegate::CreateLambda(BuildActorColorationMenu, this),
-			FUIAction(
-				FExecuteAction(),
-				FCanExecuteAction(),
-				FIsActionChecked::CreateLambda([this]()
-				{
-					if (FLevelEditorViewportClient* ViewportClient = GetLevelViewportClient())
-					{
-						return ViewportClient->HandleIsShowFlagEnabled(FEngineShowFlags::EShowFlag::SF_ActorColoration);
-					}
-					return false;
-				})
-			),
-			EUserInterfaceActionType::RadioButton,
-			/*bInOpenSubMenuOnClick=*/ false,
-			FSlateIcon(FAppStyle::GetAppStyleSetName(), "EditorViewport.LODColorationMode")
-		);
-	}
-
-	{
-		FToolMenuSection& Section = Menu->AddSection("LevelViewportLandscape", LOCTEXT("LandscapeHeader", "Landscape"), InsertPosition);
-
-		auto BuildLandscapeLODMenu = [](UToolMenu* Menu, SLevelViewportToolBar* Toolbar)
-		{
-			FToolMenuSection& SubMenuSection = Menu->AddSection("LevelViewportLandScapeLOD", LOCTEXT("LandscapeLODHeader", "Landscape LOD"));
-
-			SubMenuSection.AddMenuEntry(
-				"LandscapeLODAuto",
-				LOCTEXT("LandscapeLODAuto", "Auto"),
-				FText(),
-				FSlateIcon(),
-				FUIAction(
-					FExecuteAction::CreateSP(Toolbar, &SLevelViewportToolBar::OnLandscapeLODChanged, -1),
-					FCanExecuteAction(),
-					FIsActionChecked::CreateSP(Toolbar, &SLevelViewportToolBar::IsLandscapeLODSettingChecked, -1)
-				),
-				EUserInterfaceActionType::RadioButton
-			);
-
-			SubMenuSection.AddSeparator("LandscapeLODSeparator");
-
-			static const FText FormatString = LOCTEXT("LandscapeLODFixed", "Fixed at {0}");
-			for (int32 i = 0; i < 8; ++i)
-			{
-				SubMenuSection.AddMenuEntry(
-					NAME_None,
-					FText::Format(FormatString, FText::AsNumber(i)),
-					FText(),
-					FSlateIcon(),
-					FUIAction(
-						FExecuteAction::CreateSP(Toolbar, &SLevelViewportToolBar::OnLandscapeLODChanged, i),
-						FCanExecuteAction(),
-						FIsActionChecked::CreateSP(Toolbar, &SLevelViewportToolBar::IsLandscapeLODSettingChecked, i)
-					),
-					EUserInterfaceActionType::RadioButton
-				);
-			}
-		};
-
-		Section.AddSubMenu(
-			"LandscapeLOD",
-			LOCTEXT("LandscapeLODDisplayName", "LOD"),
-			LOCTEXT("LandscapeLODMenu_ToolTip", "Override Landscape LOD in this viewport"),
-			FNewToolMenuDelegate::CreateLambda(BuildLandscapeLODMenu, this),
-			/*bInOpenSubMenuOnClick=*/ false,
-			FSlateIcon(FAppStyle::GetAppStyleSetName(), "EditorViewport.LOD")
-		);
+		UE::LevelEditor::PopulateViewModesMenu(Menu, ViewportPinned.ToSharedRef());
 	}
 }
 
@@ -1955,18 +1654,6 @@ float SLevelViewportToolBar::GetTransformToolbarWidth() const
 	return 0.0f;
 }
 
-bool SLevelViewportToolBar::IsLandscapeLODSettingChecked(int32 Value) const
-{
-	return Viewport.Pin()->GetLevelViewportClient().LandscapeLODOverride == Value;
-}
-
-void SLevelViewportToolBar::OnLandscapeLODChanged(int32 NewValue)
-{
-	FLevelEditorViewportClient& ViewportClient = Viewport.Pin()->GetLevelViewportClient();
-	ViewportClient.LandscapeLODOverride = NewValue;
-	ViewportClient.Invalidate();
-}
-
 FReply SLevelViewportToolBar::OnRealtimeWarningClicked()
 {
 	FLevelEditorViewportClient& ViewportClient = Viewport.Pin()->GetLevelViewportClient();
diff --git a/Engine/Source/Editor/LevelEditor/Private/SLevelViewportToolBar.h b/Engine/Source/Editor/LevelEditor/Private/SLevelViewportToolBar.h
--- a/Engine/Source/Editor/LevelEditor/Private/SLevelViewportToolBar.h
+++ b/Engine/Source/Editor/LevelEditor/Private/SLevelViewportToolBar.h
@@ -225,9 +225,6 @@ private:
 	/** Called when the far view plane slider is adjusted in the perspective viewport */
 	void OnFarViewPlaneValueChanged( float NewValue );
 
-	bool IsLandscapeLODSettingChecked(int32 Value) const;
-	void OnLandscapeLODChanged(int32 NewValue);
-
 	FReply OnRealtimeWarningClicked();
 	EVisibility GetRealtimeWarningVisibility() const;
 
diff --git a/Engine/Source/Editor/LevelEditor/Private/ViewportToolbar/LevelEditorViewportToolbarSections.cpp b/Engine/Source/Editor/LevelEditor/Private/ViewportToolbar/LevelEditorViewportToolbarSections.cpp
--- a/Engine/Source/Editor/LevelEditor/Private/ViewportToolbar/LevelEditorViewportToolbarSections.cpp
+++ b/Engine/Source/Editor/LevelEditor/Private/ViewportToolbar/LevelEditorViewportToolbarSections.cpp
@@ -3,6 +3,8 @@
 #include "ViewportToolbar/LevelEditorViewportToolbarSections.h"
 
 #include "EditorViewportCommands.h"
+#include "GameFramework/ActorPrimitiveColorHandler.h"
+#include "GroomVisualizationData.h"
 #include "LevelEditor.h"
 #include "LevelEditorActions.h"
 #include "LevelViewportActions.h"
@@ -21,6 +23,17 @@
 namespace UE::LevelEditor::Private
 {
 
+bool IsLandscapeLODSettingChecked(FLevelEditorViewportClient& ViewportClient, int32 Value)
+{
+	return ViewportClient.LandscapeLODOverride == Value;
+}
+
+void OnLandscapeLODChanged(FLevelEditorViewportClient& ViewportClient, int32 NewValue)
+{
+	ViewportClient.LandscapeLODOverride = NewValue;
+	ViewportClient.Invalidate();
+}
+
 TOptional<bool> UpdateAndGetRealtimeWarningFromContext(const FToolMenuContext& Context)
 {
 	if (ULevelViewportContext* const LevelViewportContext = Context.FindContext<ULevelViewportContext>())
@@ -98,6 +111,271 @@ TSharedPtr<FExtender> GetViewModesLegacyExtenders()
 	return LevelEditorModule.GetMenuExtensibilityManager()->GetAllExtenders();
 }
 
+void PopulateViewModesMenu(UToolMenu* InMenu, TSharedRef<::SLevelViewport> InViewport)
+{
+	FToolMenuInsert InsertPosition("ViewMode", EToolMenuInsertType::After);
+
+	{
+		FToolMenuSection& Section = InMenu->AddSection(
+			"LevelViewportDeferredRendering", LOCTEXT("DeferredRenderingHeader", "Deferred Rendering"), InsertPosition);
+	}
+
+	{
+		FToolMenuSection& Section = InMenu->FindOrAddSection("ViewMode");
+		Section.AddSubMenu("VisualizeBufferViewMode",
+			LOCTEXT("VisualizeBufferViewModeDisplayName", "Buffer Visualization"),
+			LOCTEXT("BufferVisualizationMenu_ToolTip", "Select a mode for buffer visualization"),
+			FNewMenuDelegate::CreateStatic(&FBufferVisualizationMenuCommands::BuildVisualisationSubMenu),
+			FUIAction(FExecuteAction(), FCanExecuteAction(),
+				FIsActionChecked::CreateLambda([WeakViewport = InViewport.ToWeakPtr()]() {
+					const TSharedPtr<::SLevelViewport> Viewport = WeakViewport.Pin();
+					check(Viewport.IsValid());
+					FLevelEditorViewportClient& ViewportClient = Viewport->GetLevelViewportClient();
+					return ViewportClient.IsViewModeEnabled(VMI_VisualizeBuffer);
+				})),
+			EUserInterfaceActionType::RadioButton,
+			/* bInOpenSubMenuOnClick = */ false,
+			FSlateIcon(FAppStyle::GetAppStyleSetName(), "EditorViewport.VisualizeBufferMode"));
+	}
+
+	{
+		FToolMenuSection& Section = InMenu->FindOrAddSection("ViewMode");
+		Section.AddSubMenu("VisualizeNaniteViewMode",
+			LOCTEXT("VisualizeNaniteViewModeDisplayName", "Nanite Visualization"),
+			LOCTEXT("NaniteVisualizationMenu_ToolTip", "Select a mode for Nanite visualization"),
+			FNewMenuDelegate::CreateStatic(&FNaniteVisualizationMenuCommands::BuildVisualisationSubMenu),
+			FUIAction(FExecuteAction(), FCanExecuteAction(),
+				FIsActionChecked::CreateLambda([WeakViewport = InViewport.ToWeakPtr()]() {
+					const TSharedPtr<::SLevelViewport> Viewport = WeakViewport.Pin();
+					check(Viewport.IsValid());
+					FLevelEditorViewportClient& ViewportClient = Viewport->GetLevelViewportClient();
+					return ViewportClient.IsViewModeEnabled(VMI_VisualizeNanite);
+				})),
+			EUserInterfaceActionType::RadioButton,
+			/* bInOpenSubMenuOnClick = */ false,
+			FSlateIcon(FAppStyle::GetAppStyleSetName(), "EditorViewport.VisualizeNaniteMode"));
+	}
+
+	{
+		FToolMenuSection& Section = InMenu->FindOrAddSection("ViewMode");
+		Section.AddSubMenu("VisualizeLumenViewMode", LOCTEXT("VisualizeLumenViewModeDisplayName", "Lumen"),
+			LOCTEXT("LumenVisualizationMenu_ToolTip", "Select a mode for Lumen visualization"),
+			FNewMenuDelegate::CreateStatic(&FLumenVisualizationMenuCommands::BuildVisualisationSubMenu),
+			FUIAction(FExecuteAction(), FCanExecuteAction(),
+				FIsActionChecked::CreateLambda([WeakViewport = InViewport.ToWeakPtr()]() {
+					const TSharedPtr<::SLevelViewport> Viewport = WeakViewport.Pin();
+					check(Viewport.IsValid());
+					FLevelEditorViewportClient& ViewportClient = Viewport->GetLevelViewportClient();
+					return ViewportClient.IsViewModeEnabled(VMI_VisualizeLumen);
+				})),
+			EUserInterfaceActionType::RadioButton,
+			/* bInOpenSubMenuOnClick = */ false,
+			FSlateIcon(FAppStyle::GetAppStyleSetName(), "EditorViewport.VisualizeLumenMode"));
+	}
+
+	if (Substrate::IsSubstrateEnabled())
+	{
+		FToolMenuSection& Section = InMenu->FindOrAddSection("ViewMode");
+		Section.AddSubMenu("VisualizeSubstrateViewMode", LOCTEXT("VisualizeSubstrateViewModeDisplayName", "Substrate"),
+			LOCTEXT("SubstrateVisualizationMenu_ToolTip", "Select a mode for Substrate visualization"),
+			FNewMenuDelegate::CreateStatic(&FSubstrateVisualizationMenuCommands::BuildVisualisationSubMenu),
+			FUIAction(FExecuteAction(), FCanExecuteAction(),
+				FIsActionChecked::CreateLambda([WeakViewport = InViewport.ToWeakPtr()]() {
+					const TSharedPtr<::SLevelViewport> Viewport = WeakViewport.Pin();
+					check(Viewport.IsValid());
+					FLevelEditorViewportClient& ViewportClient = Viewport->GetLevelViewportClient();
+					return ViewportClient.IsViewModeEnabled(VMI_VisualizeSubstrate);
+				})),
+			EUserInterfaceActionType::RadioButton,
+			/* bInOpenSubMenuOnClick = */ false,
+			FSlateIcon(FAppStyle::GetAppStyleSetName(), "EditorViewport.VisualizeSubstrateMode"));
+	}
+
+	if (IsGroomEnabled())
+	{
+		FToolMenuSection& Section = InMenu->FindOrAddSection("ViewMode");
+		Section.AddSubMenu("VisualizeGroomViewMode", LOCTEXT("VisualizeGroomViewModeDisplayName", "Groom"),
+			LOCTEXT("GroomVisualizationMenu_ToolTip", "Select a mode for Groom visualization"),
+			FNewMenuDelegate::CreateStatic(&FGroomVisualizationMenuCommands::BuildVisualisationSubMenu),
+			FUIAction(FExecuteAction(), FCanExecuteAction(),
+				FIsActionChecked::CreateLambda([WeakViewport = InViewport.ToWeakPtr()]() {
+					const TSharedPtr<::SLevelViewport> Viewport = WeakViewport.Pin();
+					check(Viewport.IsValid());
+					FLevelEditorViewportClient& ViewportClient = Viewport->GetLevelViewportClient();
+					return ViewportClient.IsViewModeEnabled(VMI_VisualizeGroom);
+				})),
+			EUserInterfaceActionType::RadioButton,
+			/* bInOpenSubMenuOnClick = */ false,
+			FSlateIcon(FAppStyle::GetAppStyleSetName(), "EditorViewport.VisualizeGroomMode"));
+	}
+
+	{
+		FToolMenuSection& Section = InMenu->FindOrAddSection("ViewMode");
+		Section.AddSubMenu("VisualizeVirtualShadowMapViewMode",
+			LOCTEXT("VisualizeVirtualShadowMapViewModeDisplayName", "Virtual Shadow Map"),
+			LOCTEXT("VirtualShadowMapVisualizationMenu_ToolTip",
+				"Select a mode for virtual shadow map visualization. Select a light component in the world outliner to "
+				"visualize that light."),
+			FNewMenuDelegate::CreateStatic(&FVirtualShadowMapVisualizationMenuCommands::BuildVisualisationSubMenu),
+			FUIAction(FExecuteAction(), FCanExecuteAction(),
+				FIsActionChecked::CreateLambda([WeakViewport = InViewport.ToWeakPtr()]() {
+					const TSharedPtr<::SLevelViewport> Viewport = WeakViewport.Pin();
+					check(Viewport.IsValid());
+					FLevelEditorViewportClient& ViewportClient = Viewport->GetLevelViewportClient();
+					return ViewportClient.IsViewModeEnabled(VMI_VisualizeVirtualShadowMap);
+				})),
+			EUserInterfaceActionType::RadioButton,
+			/* bInOpenSubMenuOnClick = */ false,
+			FSlateIcon(FAppStyle::GetAppStyleSetName(), "EditorViewport.VisualizeVirtualShadowMapMode"));
+	}
+
+	{
+		auto BuildActorColorationMenu = [WeakViewport = InViewport.ToWeakPtr()](UToolMenu* InMenu) {
+			FToolMenuSection& SubMenuSection = InMenu->AddSection(
+				"LevelViewportActorColoration", LOCTEXT("ActorColorationHeader", "Actor Coloration"));
+
+			TArray<FActorPrimitiveColorHandler::FPrimitiveColorHandler> PrimitiveColorHandlers;
+			FActorPrimitiveColorHandler::Get().GetRegisteredPrimitiveColorHandlers(PrimitiveColorHandlers);
+
+			for (const FActorPrimitiveColorHandler::FPrimitiveColorHandler& PrimitiveColorHandler : PrimitiveColorHandlers)
+			{
+				if (!PrimitiveColorHandler.bAvailalbleInEditor)
+				{
+					continue;
+				}
+
+				SubMenuSection.AddMenuEntry(NAME_None, PrimitiveColorHandler.HandlerText, FText(), FSlateIcon(),
+					FUIAction(FExecuteAction::CreateLambda([WeakViewport, PrimitiveColorHandler]() {
+						if (TSharedPtr<::SLevelViewport> Viewport = WeakViewport.Pin())
+						{
+							const bool bActorColorationEnabled = Viewport->GetLevelViewportClient().HandleIsShowFlagEnabled(
+								FEngineShowFlags::EShowFlag::SF_ActorColoration);
+
+							if (PrimitiveColorHandler.HandlerName.IsNone())
+							{
+								if (bActorColorationEnabled)
+								{
+									Viewport->GetLevelViewportClient().HandleToggleShowFlag(
+										FEngineShowFlags::EShowFlag::SF_ActorColoration);
+								}
+							}
+							else
+							{
+								if (!bActorColorationEnabled)
+								{
+									Viewport->GetLevelViewportClient().HandleToggleShowFlag(
+										FEngineShowFlags::EShowFlag::SF_ActorColoration);
+								}
+
+								FActorPrimitiveColorHandler::Get().SetActivePrimitiveColorHandler(
+									PrimitiveColorHandler.HandlerName, GWorld);
+							}
+						}
+					}),
+						FCanExecuteAction::CreateLambda([WeakViewport]() {
+							if (TSharedPtr<::SLevelViewport> Viewport = WeakViewport.Pin())
+							{
+								return true;
+							}
+							return false;
+						}),
+						FGetActionCheckState::CreateLambda([WeakViewport, PrimitiveColorHandler]() {
+							if (TSharedPtr<::SLevelViewport> Viewport = WeakViewport.Pin())
+							{
+								const bool bActorColorationEnabled =
+									Viewport->GetLevelViewportClient().HandleIsShowFlagEnabled(
+										FEngineShowFlags::EShowFlag::SF_ActorColoration);
+
+								if (PrimitiveColorHandler.HandlerName.IsNone())
+								{
+									return bActorColorationEnabled ? ECheckBoxState::Unchecked : ECheckBoxState::Checked;
+								}
+								else
+								{
+									if (bActorColorationEnabled)
+									{
+										return FActorPrimitiveColorHandler::Get().GetActivePrimitiveColorHandler()
+													== PrimitiveColorHandler.HandlerName
+												 ? ECheckBoxState::Checked
+												 : ECheckBoxState::Unchecked;
+									}
+								}
+							}
+
+							return ECheckBoxState::Unchecked;
+						})),
+					EUserInterfaceActionType::RadioButton);
+			}
+		};
+
+		FToolMenuSection& Section = InMenu->FindOrAddSection("ViewMode");
+		Section.AddSubMenu("ActorColoration", LOCTEXT("ActorColorationDisplayName", "Actor Coloration"),
+			LOCTEXT("ActorColorationMenu_ToolTip", "Override Actor Coloration mode"),
+			FNewToolMenuDelegate::CreateLambda(BuildActorColorationMenu),
+			FUIAction(FExecuteAction(), FCanExecuteAction(),
+				FIsActionChecked::CreateLambda([WeakViewport = InViewport.ToWeakPtr()]() {
+					if (const TSharedPtr<::SLevelViewport> Viewport = WeakViewport.Pin())
+					{
+						return Viewport->GetLevelViewportClient().HandleIsShowFlagEnabled(
+							FEngineShowFlags::EShowFlag::SF_ActorColoration);
+					}
+					return false;
+				})),
+			EUserInterfaceActionType::RadioButton,
+			/*bInOpenSubMenuOnClick=*/false,
+			FSlateIcon(FAppStyle::GetAppStyleSetName(), "EditorViewport.LODColorationMode"));
+	}
+
+	{
+		FToolMenuSection& Section = InMenu->AddSection(
+			"LevelViewportLandscape", LOCTEXT("LandscapeHeader", "Landscape"), InsertPosition);
+
+		auto BuildLandscapeLODMenu = [WeakViewport = InViewport.ToWeakPtr()](UToolMenu* InMenu) {
+			FToolMenuSection& SubMenuSection = InMenu->AddSection(
+				"LevelViewportLandScapeLOD", LOCTEXT("LandscapeLODHeader", "Landscape LOD"));
+
+			auto CreateLandscapeLODAction = [WeakViewport](int32 LODValue) {
+				FUIAction LandscapeLODAction;
+				LandscapeLODAction.ExecuteAction = FExecuteAction::CreateLambda([WeakViewport, LODValue]() {
+					if (const TSharedPtr<::SLevelViewport> Viewport = WeakViewport.Pin())
+					{
+						UE::LevelEditor::Private::OnLandscapeLODChanged(Viewport->GetLevelViewportClient(), LODValue);
+					}
+				});
+				LandscapeLODAction.GetActionCheckState = FGetActionCheckState::CreateLambda(
+					[WeakViewport, LODValue]() -> ECheckBoxState {
+						bool bChecked = false;
+						if (const TSharedPtr<::SLevelViewport> Viewport = WeakViewport.Pin())
+						{
+							bChecked = UE::LevelEditor::Private::IsLandscapeLODSettingChecked(
+								Viewport->GetLevelViewportClient(), LODValue);
+						}
+						return bChecked ? ECheckBoxState::Checked : ECheckBoxState::Unchecked;
+					});
+
+				return LandscapeLODAction;
+			};
+
+			SubMenuSection.AddMenuEntry("LandscapeLODAuto", LOCTEXT("LandscapeLODAuto", "Auto"), FText(), FSlateIcon(),
+				CreateLandscapeLODAction(-1), EUserInterfaceActionType::RadioButton);
+
+			SubMenuSection.AddSeparator("LandscapeLODSeparator");
+
+			static const FText FormatString = LOCTEXT("LandscapeLODFixed", "Fixed at {0}");
+			for (int32 i = 0; i < 8; ++i)
+			{
+				SubMenuSection.AddMenuEntry(NAME_None, FText::Format(FormatString, FText::AsNumber(i)), FText(),
+					FSlateIcon(), CreateLandscapeLODAction(i), EUserInterfaceActionType::RadioButton);
+			}
+		};
+
+		Section.AddSubMenu("LandscapeLOD", LOCTEXT("LandscapeLODDisplayName", "LOD"),
+			LOCTEXT("LandscapeLODMenu_ToolTip", "Override Landscape LOD in this viewport"),
+			FNewToolMenuDelegate::CreateLambda(BuildLandscapeLODMenu),
+			/*bInOpenSubMenuOnClick=*/false, FSlateIcon(FAppStyle::GetAppStyleSetName(), "EditorViewport.LOD"));
+	}
+}
+
 void AddFeatureLevelPreviewSubmenu(FToolMenuSection& Section)
 {
 	Section.AddSubMenu("FeatureLevelPreview",
diff --git a/Engine/Source/Editor/LevelEditor/Private/ViewportToolbar/LevelEditorViewportToolbarSections.h b/Engine/Source/Editor/LevelEditor/Private/ViewportToolbar/LevelEditorViewportToolbarSections.h
--- a/Engine/Source/Editor/LevelEditor/Private/ViewportToolbar/LevelEditorViewportToolbarSections.h
+++ b/Engine/Source/Editor/LevelEditor/Private/ViewportToolbar/LevelEditorViewportToolbarSections.h
@@ -19,6 +19,7 @@ bool ShowViewportRealtimeWarning(FLevelEditorViewportClient& ViewportClient);
 void AddViewportToolbarTransformsSection(FToolMenuSection& InSection);
 
 TSharedPtr<FExtender> GetViewModesLegacyExtenders();
+void PopulateViewModesMenu(UToolMenu* InMenu, TSharedRef<SLevelViewport> InViewport);
 
 void AddFeatureLevelPreviewSubmenu(FToolMenuSection& Section);
 void AddMaterialQualityLevelSubmenu(FToolMenuSection& Section);
diff --git a/Engine/Source/Editor/UnrealEd/Private/SEditorViewportViewMenu.cpp b/Engine/Source/Editor/UnrealEd/Private/SEditorViewportViewMenu.cpp
--- a/Engine/Source/Editor/UnrealEd/Private/SEditorViewportViewMenu.cpp
+++ b/Engine/Source/Editor/UnrealEd/Private/SEditorViewportViewMenu.cpp
@@ -89,7 +89,7 @@ void SEditorViewportViewMenu::FillViewMenu(UToolMenu* Menu) const
 			return true;
 		});
 
-	UE::UnrealEd::FillViewMenu(Menu, Viewport.Pin().ToSharedRef(), IsViewModeSupported);
+	UE::UnrealEd::PopulateViewModesMenu(Menu, Viewport.Pin().ToSharedRef(), IsViewModeSupported);
 }
 
 #undef LOCTEXT_NAMESPACE
diff --git a/Engine/Source/Editor/UnrealEd/Private/ViewportToolbar/UnrealEdViewportToolbar.cpp b/Engine/Source/Editor/UnrealEd/Private/ViewportToolbar/UnrealEdViewportToolbar.cpp
--- a/Engine/Source/Editor/UnrealEd/Private/ViewportToolbar/UnrealEdViewportToolbar.cpp
+++ b/Engine/Source/Editor/UnrealEd/Private/ViewportToolbar/UnrealEdViewportToolbar.cpp
@@ -66,7 +66,7 @@ FText GetViewModesSubmenuLabel(TWeakPtr<SEditorViewport> InViewport)
 	return Label;
 }
 
-void FillViewMenu(
+void PopulateViewModesMenu(
 	UToolMenu* InMenu, TSharedRef<SEditorViewport> InViewport, IsViewModeSupportedDelegate InIsViewModeSupported)
 {
 	const FEditorViewportCommands& BaseViewportActions = FEditorViewportCommands::Get();
diff --git a/Engine/Source/Editor/UnrealEd/Public/ViewportToolbar/UnrealEdViewportToolbar.h b/Engine/Source/Editor/UnrealEd/Public/ViewportToolbar/UnrealEdViewportToolbar.h
--- a/Engine/Source/Editor/UnrealEd/Public/ViewportToolbar/UnrealEdViewportToolbar.h
+++ b/Engine/Source/Editor/UnrealEd/Public/ViewportToolbar/UnrealEdViewportToolbar.h
@@ -24,7 +24,7 @@ DECLARE_DELEGATE_RetVal_OneParam(bool, IsViewModeSupportedDelegate, EViewModeInd
  * @param InViewport The viewport associated with this viewport toolbar.
  * @param InIsViewModeSupported Optional delegate to filter which view modes are added to the list.
  */
-UNREALED_API void FillViewMenu(UToolMenu* InMenu,
+UNREALED_API void PopulateViewModesMenu(UToolMenu* InMenu,
 	TSharedRef<SEditorViewport> InViewport,
 	IsViewModeSupportedDelegate InIsViewModeSupported = IsViewModeSupportedDelegate());
 
