Tweaked formatting of project badge in editor to avoid unnecessary spaces
#rnx

diff --git a/Engine/Source/Editor/LevelEditor/Private/LevelEditor.cpp b/Engine/Source/Editor/LevelEditor/Private/LevelEditor.cpp
--- a/Engine/Source/Editor/LevelEditor/Private/LevelEditor.cpp
+++ b/Engine/Source/Editor/LevelEditor/Private/LevelEditor.cpp
@@ -110,26 +110,17 @@ public:
 
 		const FString EngineVersionString = FEngineVersion::Current().ToString(FEngineVersion::Current().HasChangelist() ? EVersionComponent::Changelist : EVersionComponent::Patch);
 
-		FFormatNamedArguments Args;
+		const EBuildConfiguration BuildConfig = FApp::GetBuildConfiguration();
+		const bool bAppendBuildConfig = (BuildConfig != EBuildConfiguration::Shipping && BuildConfig != EBuildConfiguration::Development && BuildConfig != EBuildConfiguration::Unknown);
 
+		FFormatNamedArguments Args;
 		Args.Add(TEXT("ProjectNameWatermarkPrefix"), FText::FromString(ProjectNameWatermarkPrefix));
-		Args.Add(TEXT("Branch"), FEngineBuildSettings::IsPerforceBuild() ? FText::FromString(FEngineVersion::Current().GetBranch()) : FText::GetEmpty());
-		Args.Add(TEXT("GameName"), FText::FromString(FString(FApp::GetProjectName())));
-		Args.Add(TEXT("EngineVersion"), (GetDefault<UEditorPerProjectUserSettings>()->bDisplayEngineVersionInBadge) ? FText::FromString("(" + EngineVersionString + ")") : FText());
+		Args.Add(TEXT("Branch"), FEngineBuildSettings::IsPerforceBuild() ? FText::FromString(TEXT(" ") + FEngineVersion::Current().GetBranch()) : FText::GetEmpty());
+		Args.Add(TEXT("GameName"), FText::FromString(TEXT(" ") + FString(FApp::GetProjectName())));
+		Args.Add(TEXT("EngineVersion"), (GetDefault<UEditorPerProjectUserSettings>()->bDisplayEngineVersionInBadge) ? FText::FromString(TEXT(" (") + EngineVersionString + TEXT(")")) : FText::GetEmpty());
+		Args.Add(TEXT("Config"), bAppendBuildConfig ? FText::Format(NSLOCTEXT("UnrealEditor", "TitleBarConfig", " [{0}]"), EBuildConfigurations::ToText(BuildConfig)) : FText::GetEmpty());
 
-		FText RightContentText;
-		FText RightContentTooltip;
-
-		const EBuildConfiguration BuildConfig = FApp::GetBuildConfiguration();
-		if (BuildConfig != EBuildConfiguration::Shipping && BuildConfig != EBuildConfiguration::Development && BuildConfig != EBuildConfiguration::Unknown)
-		{
-			Args.Add(TEXT("Config"), EBuildConfigurations::ToText(BuildConfig));
-			RightContentText = FText::Format(NSLOCTEXT("UnrealEditor", "TitleBarRightContentAndConfig", "{ProjectNameWatermarkPrefix} {GameName} [{Config}] {Branch} {EngineVersion}"), Args);
-		}
-		else
-		{
-			RightContentText = FText::Format(NSLOCTEXT("UnrealEditor", "TitleBarRightContent", "{ProjectNameWatermarkPrefix} {GameName} {Branch} {EngineVersion}"), Args);
-		}
+		FText RightContentText = FText::Format(NSLOCTEXT("UnrealEditor", "TitleBarRightContentAndConfig", "{ProjectNameWatermarkPrefix}{GameName}{Config}{Branch}{EngineVersion}"), Args);
 
 		// Create the tooltip showing more detailed information
 		FFormatNamedArguments TooltipArgs;
@@ -140,7 +131,7 @@ public:
 		TooltipArgs.Add(TEXT("BuildDate"), FText::FromString(FApp::GetBuildDate()));
 		TooltipArgs.Add(TEXT("GraphicsRHI"), FText::FromString(FApp::GetGraphicsRHI()));
 
-		RightContentTooltip = FText::Format(NSLOCTEXT("UnrealEditor", "TitleBarRightContentTooltip", "Version: {Version}\nBranch: {Branch}\nBuild Configuration: {BuildConfiguration}\nBuild Date: {BuildDate}\nGraphics RHI: {GraphicsRHI}"), TooltipArgs);
+		FText RightContentTooltip = FText::Format(NSLOCTEXT("UnrealEditor", "TitleBarRightContentTooltip", "Version: {Version}\nBranch: {Branch}\nBuild Configuration: {BuildConfiguration}\nBuild Date: {BuildDate}\nGraphics RHI: {GraphicsRHI}"), TooltipArgs);
 
 		SetToolTipText(RightContentTooltip);
 
