Add "View Modes" menu to the new viewport toolbar

diff --git a/Engine/Source/Editor/LevelEditor/Private/SLevelViewport.cpp b/Engine/Source/Editor/LevelEditor/Private/SLevelViewport.cpp
--- a/Engine/Source/Editor/LevelEditor/Private/SLevelViewport.cpp
+++ b/Engine/Source/Editor/LevelEditor/Private/SLevelViewport.cpp
@@ -1933,6 +1933,14 @@ TSharedPtr<SWidget> SLevelViewport::MakeViewportToolbar()
 
 			{
 				FToolMenuSection& RightSection = ViewportToolbarMenu->FindOrAddSection("Right");
+
+				{
+					// Stay backward-compatible with the old viewport toolbar.
+					UToolMenus::Get()->RegisterMenu(
+						"LevelEditor.ViewportToolbar.ViewModes", "LevelEditor.LevelViewportToolbar.View");
+					UE::LevelEditor::AddViewportToolbarViewModesSubmenu(RightSection);
+				}
+
 				UE::LevelEditor::AddViewportToolbarPerformanceAndScalabilitySubmenu(RightSection);
 				UE::LevelEditor::AddLevelEditorViewportToolbarSettingsSubmenu(RightSection);
 			}
diff --git a/Engine/Source/Editor/LevelEditor/Private/ViewportToolbar/LevelEditorViewportToolbarSections.cpp b/Engine/Source/Editor/LevelEditor/Private/ViewportToolbar/LevelEditorViewportToolbarSections.cpp
--- a/Engine/Source/Editor/LevelEditor/Private/ViewportToolbar/LevelEditorViewportToolbarSections.cpp
+++ b/Engine/Source/Editor/LevelEditor/Private/ViewportToolbar/LevelEditorViewportToolbarSections.cpp
@@ -15,6 +15,7 @@
 #include "Templates/SharedPointer.h"
 #include "ToolMenu.h"
 #include "ToolMenus.h"
+#include "ViewportToolbar/UnrealEdViewportToolbar.h"
 #include "Widgets/Input/SVolumeControl.h"
 #include "Widgets/SBoxPanel.h"
 
@@ -376,6 +377,37 @@ void PopulateViewModesMenu(UToolMenu* InMenu, TSharedRef<::SLevelViewport> InVie
 	}
 }
 
+void AddViewportToolbarViewModesSubmenu(FToolMenuSection& InSection)
+{
+	// This has to be a dynamic entry for the ViewModes submenu's label to be able to access the context.
+	InSection.AddDynamicEntry(
+		"DynamicViewModes", FNewToolMenuSectionDelegate::CreateLambda([](FToolMenuSection& InDynamicSection) -> void {
+			TAttribute<FText> LabelAttribute = UE::UnrealEd::GetViewModesSubmenuLabel(nullptr);
+			if (ULevelViewportContext* const LevelViewportContext = InDynamicSection.FindContext<ULevelViewportContext>())
+			{
+				TWeakPtr<SEditorViewport> EditorViewport = LevelViewportContext->LevelViewport;
+				LabelAttribute = TAttribute<FText>::CreateLambda(
+					[EditorViewport]() { return UE::UnrealEd::GetViewModesSubmenuLabel(EditorViewport); });
+			}
+
+			InDynamicSection.AddSubMenu("ViewModes", LabelAttribute,
+				LOCTEXT("ViewModesSubmenuTooltip", "View mode settings for the current viewport."),
+				FNewToolMenuDelegate::CreateLambda([](UToolMenu* Submenu) -> void {
+					ULevelViewportContext* const LevelViewportContext = Submenu->FindContext<ULevelViewportContext>();
+					if (!LevelViewportContext)
+					{
+						return;
+					}
+
+					if (const TSharedPtr<::SLevelViewport> LevelViewport = LevelViewportContext->LevelViewport.Pin())
+					{
+						UE::UnrealEd::PopulateViewModesMenu(Submenu, LevelViewport.ToSharedRef());
+						PopulateViewModesMenu(Submenu, LevelViewport.ToSharedRef());
+					}
+				}));
+		}));
+}
+
 void AddFeatureLevelPreviewSubmenu(FToolMenuSection& Section)
 {
 	Section.AddSubMenu("FeatureLevelPreview",
diff --git a/Engine/Source/Editor/LevelEditor/Private/ViewportToolbar/LevelEditorViewportToolbarSections.h b/Engine/Source/Editor/LevelEditor/Private/ViewportToolbar/LevelEditorViewportToolbarSections.h
--- a/Engine/Source/Editor/LevelEditor/Private/ViewportToolbar/LevelEditorViewportToolbarSections.h
+++ b/Engine/Source/Editor/LevelEditor/Private/ViewportToolbar/LevelEditorViewportToolbarSections.h
@@ -20,6 +20,7 @@ void AddViewportToolbarTransformsSection(FToolMenuSection& InSection);
 
 TSharedPtr<FExtender> GetViewModesLegacyExtenders();
 void PopulateViewModesMenu(UToolMenu* InMenu, TSharedRef<SLevelViewport> InViewport);
+void AddViewportToolbarViewModesSubmenu(FToolMenuSection& InSection);
 
 void AddFeatureLevelPreviewSubmenu(FToolMenuSection& Section);
 void AddMaterialQualityLevelSubmenu(FToolMenuSection& Section);
