Fixing crash when saving an Actor to a GroundTruthData asset for test validation as the Actor does not have a valid level associated with it.

#rnx

#virtualized

diff --git a/Engine/Source/Developer/FunctionalTesting/Private/GroundTruthData.cpp b/Engine/Source/Developer/FunctionalTesting/Private/GroundTruthData.cpp
--- a/Engine/Source/Developer/FunctionalTesting/Private/GroundTruthData.cpp
+++ b/Engine/Source/Developer/FunctionalTesting/Private/GroundTruthData.cpp
@@ -69,8 +69,10 @@ void UGroundTruthData::SaveObject(UObject* GroundTruth)
 		SourceControlProvider.Execute(ISourceControlOperation::Create<FCheckOut>(), GroundTruthPackage);
 	}
 
-	ObjectData = GroundTruth;
-	GroundTruth->Rename(nullptr, this);
+	// There is an edge case for Actors where calling `Rename` will transfer ownership from the level to the GroundTruthData and will remove the Actor from the current level
+	// To avoid any issues this process causes, we want to duplicate our object with a transient package before performing the rename which will change the ownership from the transient package to the GroundTruthData
+	ObjectData = DuplicateObject(GroundTruth, nullptr);
+	ObjectData->Rename(nullptr, this, REN_SkipComponentRegWork);
 	MarkPackageDirty();
 
 	FSavePackageArgs SaveArgs;
diff --git a/Engine/Source/Developer/FunctionalTesting/Public/GroundTruthData.h b/Engine/Source/Developer/FunctionalTesting/Public/GroundTruthData.h
--- a/Engine/Source/Developer/FunctionalTesting/Public/GroundTruthData.h
+++ b/Engine/Source/Developer/FunctionalTesting/Public/GroundTruthData.h
@@ -47,6 +47,6 @@ public:
 
 protected:
 	
-	UPROPERTY(VisibleAnywhere, Instanced, Category=Data)
+	UPROPERTY(VisibleAnywhere, Export, Category=Data)
 	TObjectPtr<UObject> ObjectData;
 };
diff --git a/Engine/Source/Editor/LevelEditor/Private/Elements/Actor/ActorElementLevelEditorSelectionCustomization.cpp b/Engine/Source/Editor/LevelEditor/Private/Elements/Actor/ActorElementLevelEditorSelectionCustomization.cpp
--- a/Engine/Source/Editor/LevelEditor/Private/Elements/Actor/ActorElementLevelEditorSelectionCustomization.cpp
+++ b/Engine/Source/Editor/LevelEditor/Private/Elements/Actor/ActorElementLevelEditorSelectionCustomization.cpp
@@ -90,7 +90,7 @@ bool IsActorReachable(const AActor* Actor)
 	// Ensure that neither the level nor the actor is being destroyed or is unreachable
 	const EObjectFlags InvalidSelectableFlags = RF_BeginDestroyed;
 	bool ReturnValue = true;
-	if (Actor->GetLevel()->HasAnyFlags(InvalidSelectableFlags) || (!GIsTransacting && (!IsValidChecked(Actor->GetLevel()) || Actor->GetLevel()->IsUnreachable())))
+	if (!Actor->GetLevel() || Actor->GetLevel()->HasAnyFlags(InvalidSelectableFlags) || (!GIsTransacting && (!IsValidChecked(Actor->GetLevel()) || Actor->GetLevel()->IsUnreachable())))
 	{
 		UE_LOG(LogActorLevelEditorSelection, Warning, TEXT("SelectActor: %s (%s) Flags: (0x%08X)"),
 			TEXT("The requested operation could not be completed because the level has invalid flags."),
