Run clang-format on viewport toolbar helper code

Also add // clang-format off/on comments around the Slate code in BuildVolumeControlCustomWidget.

diff --git a/Engine/Source/Editor/LevelEditor/Private/ViewportToolbar/LevelEditorViewportToolbarSections.cpp b/Engine/Source/Editor/LevelEditor/Private/ViewportToolbar/LevelEditorViewportToolbarSections.cpp
--- a/Engine/Source/Editor/LevelEditor/Private/ViewportToolbar/LevelEditorViewportToolbarSections.cpp
+++ b/Engine/Source/Editor/LevelEditor/Private/ViewportToolbar/LevelEditorViewportToolbarSections.cpp
@@ -79,31 +79,37 @@ bool ShowViewportRealtimeWarning(FLevelEditorViewportClient& ViewportClient)
 // TODO: Move this outside the level editor and make it publicly available to anyone building a viewport toolbar.
 void AddViewportToolbarTransformsSection(FToolMenuSection& InSection)
 {
-	InSection.AddSubMenu("Transforms", LOCTEXT("TransformsSubmenuLabel", "Transforms"),
+	InSection.AddSubMenu(
+		"Transforms",
+		LOCTEXT("TransformsSubmenuLabel", "Transforms"),
 		LOCTEXT("TransformsSubmenuTooltip", "Viewport-related transforms tools"),
-		FNewToolMenuDelegate::CreateLambda([](UToolMenu* Submenu) -> void {
-			FToolMenuSection& Section = Submenu->FindOrAddSection(NAME_None);
-
-			FToolMenuEntry SelectMode = FToolMenuEntry::InitMenuEntry(FEditorViewportCommands::Get().SelectMode);
-			SelectMode.UserInterfaceActionType = EUserInterfaceActionType::RadioButton;
-			SelectMode.SetShowInToolbarTopLevel(true);
-			Section.AddEntry(SelectMode);
-
-			FToolMenuEntry TranslateMode = FToolMenuEntry::InitMenuEntry(FEditorViewportCommands::Get().TranslateMode);
-			TranslateMode.UserInterfaceActionType = EUserInterfaceActionType::RadioButton;
-			TranslateMode.SetShowInToolbarTopLevel(true);
-			Section.AddEntry(TranslateMode);
-
-			FToolMenuEntry RotateMode = FToolMenuEntry::InitMenuEntry(FEditorViewportCommands::Get().RotateMode);
-			RotateMode.UserInterfaceActionType = EUserInterfaceActionType::RadioButton;
-			RotateMode.SetShowInToolbarTopLevel(true);
-			Section.AddEntry(RotateMode);
-
-			FToolMenuEntry ScaleMode = FToolMenuEntry::InitMenuEntry(FEditorViewportCommands::Get().ScaleMode);
-			ScaleMode.UserInterfaceActionType = EUserInterfaceActionType::RadioButton;
-			ScaleMode.SetShowInToolbarTopLevel(true);
-			Section.AddEntry(ScaleMode);
-		}));
+		FNewToolMenuDelegate::CreateLambda(
+			[](UToolMenu* Submenu) -> void
+			{
+				FToolMenuSection& Section = Submenu->FindOrAddSection(NAME_None);
+
+				FToolMenuEntry SelectMode = FToolMenuEntry::InitMenuEntry(FEditorViewportCommands::Get().SelectMode);
+				SelectMode.UserInterfaceActionType = EUserInterfaceActionType::RadioButton;
+				SelectMode.SetShowInToolbarTopLevel(true);
+				Section.AddEntry(SelectMode);
+
+				FToolMenuEntry TranslateMode = FToolMenuEntry::InitMenuEntry(FEditorViewportCommands::Get().TranslateMode);
+				TranslateMode.UserInterfaceActionType = EUserInterfaceActionType::RadioButton;
+				TranslateMode.SetShowInToolbarTopLevel(true);
+				Section.AddEntry(TranslateMode);
+
+				FToolMenuEntry RotateMode = FToolMenuEntry::InitMenuEntry(FEditorViewportCommands::Get().RotateMode);
+				RotateMode.UserInterfaceActionType = EUserInterfaceActionType::RadioButton;
+				RotateMode.SetShowInToolbarTopLevel(true);
+				Section.AddEntry(RotateMode);
+
+				FToolMenuEntry ScaleMode = FToolMenuEntry::InitMenuEntry(FEditorViewportCommands::Get().ScaleMode);
+				ScaleMode.UserInterfaceActionType = EUserInterfaceActionType::RadioButton;
+				ScaleMode.SetShowInToolbarTopLevel(true);
+				Section.AddEntry(ScaleMode);
+			}
+		)
+	);
 }
 
 TSharedPtr<FExtender> GetViewModesLegacyExtenders()
@@ -118,122 +124,177 @@ void PopulateViewModesMenu(UToolMenu* InMenu, TSharedRef<::SLevelViewport> InVie
 
 	{
 		FToolMenuSection& Section = InMenu->AddSection(
-			"LevelViewportDeferredRendering", LOCTEXT("DeferredRenderingHeader", "Deferred Rendering"), InsertPosition);
+			"LevelViewportDeferredRendering", LOCTEXT("DeferredRenderingHeader", "Deferred Rendering"), InsertPosition
+		);
 	}
 
 	{
 		FToolMenuSection& Section = InMenu->FindOrAddSection("ViewMode");
-		Section.AddSubMenu("VisualizeBufferViewMode",
+		Section.AddSubMenu(
+			"VisualizeBufferViewMode",
 			LOCTEXT("VisualizeBufferViewModeDisplayName", "Buffer Visualization"),
 			LOCTEXT("BufferVisualizationMenu_ToolTip", "Select a mode for buffer visualization"),
 			FNewMenuDelegate::CreateStatic(&FBufferVisualizationMenuCommands::BuildVisualisationSubMenu),
-			FUIAction(FExecuteAction(), FCanExecuteAction(),
-				FIsActionChecked::CreateLambda([WeakViewport = InViewport.ToWeakPtr()]() {
-					const TSharedPtr<::SLevelViewport> Viewport = WeakViewport.Pin();
-					check(Viewport.IsValid());
-					FLevelEditorViewportClient& ViewportClient = Viewport->GetLevelViewportClient();
-					return ViewportClient.IsViewModeEnabled(VMI_VisualizeBuffer);
-				})),
+			FUIAction(
+				FExecuteAction(),
+				FCanExecuteAction(),
+				FIsActionChecked::CreateLambda(
+					[WeakViewport = InViewport.ToWeakPtr()]()
+					{
+						const TSharedPtr<::SLevelViewport> Viewport = WeakViewport.Pin();
+						check(Viewport.IsValid());
+						FLevelEditorViewportClient& ViewportClient = Viewport->GetLevelViewportClient();
+						return ViewportClient.IsViewModeEnabled(VMI_VisualizeBuffer);
+					}
+				)
+			),
 			EUserInterfaceActionType::RadioButton,
 			/* bInOpenSubMenuOnClick = */ false,
-			FSlateIcon(FAppStyle::GetAppStyleSetName(), "EditorViewport.VisualizeBufferMode"));
+			FSlateIcon(FAppStyle::GetAppStyleSetName(), "EditorViewport.VisualizeBufferMode")
+		);
 	}
 
 	{
 		FToolMenuSection& Section = InMenu->FindOrAddSection("ViewMode");
-		Section.AddSubMenu("VisualizeNaniteViewMode",
+		Section.AddSubMenu(
+			"VisualizeNaniteViewMode",
 			LOCTEXT("VisualizeNaniteViewModeDisplayName", "Nanite Visualization"),
 			LOCTEXT("NaniteVisualizationMenu_ToolTip", "Select a mode for Nanite visualization"),
 			FNewMenuDelegate::CreateStatic(&FNaniteVisualizationMenuCommands::BuildVisualisationSubMenu),
-			FUIAction(FExecuteAction(), FCanExecuteAction(),
-				FIsActionChecked::CreateLambda([WeakViewport = InViewport.ToWeakPtr()]() {
-					const TSharedPtr<::SLevelViewport> Viewport = WeakViewport.Pin();
-					check(Viewport.IsValid());
-					FLevelEditorViewportClient& ViewportClient = Viewport->GetLevelViewportClient();
-					return ViewportClient.IsViewModeEnabled(VMI_VisualizeNanite);
-				})),
+			FUIAction(
+				FExecuteAction(),
+				FCanExecuteAction(),
+				FIsActionChecked::CreateLambda(
+					[WeakViewport = InViewport.ToWeakPtr()]()
+					{
+						const TSharedPtr<::SLevelViewport> Viewport = WeakViewport.Pin();
+						check(Viewport.IsValid());
+						FLevelEditorViewportClient& ViewportClient = Viewport->GetLevelViewportClient();
+						return ViewportClient.IsViewModeEnabled(VMI_VisualizeNanite);
+					}
+				)
+			),
 			EUserInterfaceActionType::RadioButton,
 			/* bInOpenSubMenuOnClick = */ false,
-			FSlateIcon(FAppStyle::GetAppStyleSetName(), "EditorViewport.VisualizeNaniteMode"));
+			FSlateIcon(FAppStyle::GetAppStyleSetName(), "EditorViewport.VisualizeNaniteMode")
+		);
 	}
 
 	{
 		FToolMenuSection& Section = InMenu->FindOrAddSection("ViewMode");
-		Section.AddSubMenu("VisualizeLumenViewMode", LOCTEXT("VisualizeLumenViewModeDisplayName", "Lumen"),
+		Section.AddSubMenu(
+			"VisualizeLumenViewMode",
+			LOCTEXT("VisualizeLumenViewModeDisplayName", "Lumen"),
 			LOCTEXT("LumenVisualizationMenu_ToolTip", "Select a mode for Lumen visualization"),
 			FNewMenuDelegate::CreateStatic(&FLumenVisualizationMenuCommands::BuildVisualisationSubMenu),
-			FUIAction(FExecuteAction(), FCanExecuteAction(),
-				FIsActionChecked::CreateLambda([WeakViewport = InViewport.ToWeakPtr()]() {
-					const TSharedPtr<::SLevelViewport> Viewport = WeakViewport.Pin();
-					check(Viewport.IsValid());
-					FLevelEditorViewportClient& ViewportClient = Viewport->GetLevelViewportClient();
-					return ViewportClient.IsViewModeEnabled(VMI_VisualizeLumen);
-				})),
+			FUIAction(
+				FExecuteAction(),
+				FCanExecuteAction(),
+				FIsActionChecked::CreateLambda(
+					[WeakViewport = InViewport.ToWeakPtr()]()
+					{
+						const TSharedPtr<::SLevelViewport> Viewport = WeakViewport.Pin();
+						check(Viewport.IsValid());
+						FLevelEditorViewportClient& ViewportClient = Viewport->GetLevelViewportClient();
+						return ViewportClient.IsViewModeEnabled(VMI_VisualizeLumen);
+					}
+				)
+			),
 			EUserInterfaceActionType::RadioButton,
 			/* bInOpenSubMenuOnClick = */ false,
-			FSlateIcon(FAppStyle::GetAppStyleSetName(), "EditorViewport.VisualizeLumenMode"));
+			FSlateIcon(FAppStyle::GetAppStyleSetName(), "EditorViewport.VisualizeLumenMode")
+		);
 	}
 
 	if (Substrate::IsSubstrateEnabled())
 	{
 		FToolMenuSection& Section = InMenu->FindOrAddSection("ViewMode");
-		Section.AddSubMenu("VisualizeSubstrateViewMode", LOCTEXT("VisualizeSubstrateViewModeDisplayName", "Substrate"),
+		Section.AddSubMenu(
+			"VisualizeSubstrateViewMode",
+			LOCTEXT("VisualizeSubstrateViewModeDisplayName", "Substrate"),
 			LOCTEXT("SubstrateVisualizationMenu_ToolTip", "Select a mode for Substrate visualization"),
 			FNewMenuDelegate::CreateStatic(&FSubstrateVisualizationMenuCommands::BuildVisualisationSubMenu),
-			FUIAction(FExecuteAction(), FCanExecuteAction(),
-				FIsActionChecked::CreateLambda([WeakViewport = InViewport.ToWeakPtr()]() {
-					const TSharedPtr<::SLevelViewport> Viewport = WeakViewport.Pin();
-					check(Viewport.IsValid());
-					FLevelEditorViewportClient& ViewportClient = Viewport->GetLevelViewportClient();
-					return ViewportClient.IsViewModeEnabled(VMI_VisualizeSubstrate);
-				})),
+			FUIAction(
+				FExecuteAction(),
+				FCanExecuteAction(),
+				FIsActionChecked::CreateLambda(
+					[WeakViewport = InViewport.ToWeakPtr()]()
+					{
+						const TSharedPtr<::SLevelViewport> Viewport = WeakViewport.Pin();
+						check(Viewport.IsValid());
+						FLevelEditorViewportClient& ViewportClient = Viewport->GetLevelViewportClient();
+						return ViewportClient.IsViewModeEnabled(VMI_VisualizeSubstrate);
+					}
+				)
+			),
 			EUserInterfaceActionType::RadioButton,
 			/* bInOpenSubMenuOnClick = */ false,
-			FSlateIcon(FAppStyle::GetAppStyleSetName(), "EditorViewport.VisualizeSubstrateMode"));
+			FSlateIcon(FAppStyle::GetAppStyleSetName(), "EditorViewport.VisualizeSubstrateMode")
+		);
 	}
 
 	if (IsGroomEnabled())
 	{
 		FToolMenuSection& Section = InMenu->FindOrAddSection("ViewMode");
-		Section.AddSubMenu("VisualizeGroomViewMode", LOCTEXT("VisualizeGroomViewModeDisplayName", "Groom"),
+		Section.AddSubMenu(
+			"VisualizeGroomViewMode",
+			LOCTEXT("VisualizeGroomViewModeDisplayName", "Groom"),
 			LOCTEXT("GroomVisualizationMenu_ToolTip", "Select a mode for Groom visualization"),
 			FNewMenuDelegate::CreateStatic(&FGroomVisualizationMenuCommands::BuildVisualisationSubMenu),
-			FUIAction(FExecuteAction(), FCanExecuteAction(),
-				FIsActionChecked::CreateLambda([WeakViewport = InViewport.ToWeakPtr()]() {
-					const TSharedPtr<::SLevelViewport> Viewport = WeakViewport.Pin();
-					check(Viewport.IsValid());
-					FLevelEditorViewportClient& ViewportClient = Viewport->GetLevelViewportClient();
-					return ViewportClient.IsViewModeEnabled(VMI_VisualizeGroom);
-				})),
+			FUIAction(
+				FExecuteAction(),
+				FCanExecuteAction(),
+				FIsActionChecked::CreateLambda(
+					[WeakViewport = InViewport.ToWeakPtr()]()
+					{
+						const TSharedPtr<::SLevelViewport> Viewport = WeakViewport.Pin();
+						check(Viewport.IsValid());
+						FLevelEditorViewportClient& ViewportClient = Viewport->GetLevelViewportClient();
+						return ViewportClient.IsViewModeEnabled(VMI_VisualizeGroom);
+					}
+				)
+			),
 			EUserInterfaceActionType::RadioButton,
 			/* bInOpenSubMenuOnClick = */ false,
-			FSlateIcon(FAppStyle::GetAppStyleSetName(), "EditorViewport.VisualizeGroomMode"));
+			FSlateIcon(FAppStyle::GetAppStyleSetName(), "EditorViewport.VisualizeGroomMode")
+		);
 	}
 
 	{
 		FToolMenuSection& Section = InMenu->FindOrAddSection("ViewMode");
-		Section.AddSubMenu("VisualizeVirtualShadowMapViewMode",
+		Section.AddSubMenu(
+			"VisualizeVirtualShadowMapViewMode",
 			LOCTEXT("VisualizeVirtualShadowMapViewModeDisplayName", "Virtual Shadow Map"),
-			LOCTEXT("VirtualShadowMapVisualizationMenu_ToolTip",
+			LOCTEXT(
+				"VirtualShadowMapVisualizationMenu_ToolTip",
 				"Select a mode for virtual shadow map visualization. Select a light component in the world outliner to "
-				"visualize that light."),
+				"visualize that light."
+			),
 			FNewMenuDelegate::CreateStatic(&FVirtualShadowMapVisualizationMenuCommands::BuildVisualisationSubMenu),
-			FUIAction(FExecuteAction(), FCanExecuteAction(),
-				FIsActionChecked::CreateLambda([WeakViewport = InViewport.ToWeakPtr()]() {
-					const TSharedPtr<::SLevelViewport> Viewport = WeakViewport.Pin();
-					check(Viewport.IsValid());
-					FLevelEditorViewportClient& ViewportClient = Viewport->GetLevelViewportClient();
-					return ViewportClient.IsViewModeEnabled(VMI_VisualizeVirtualShadowMap);
-				})),
+			FUIAction(
+				FExecuteAction(),
+				FCanExecuteAction(),
+				FIsActionChecked::CreateLambda(
+					[WeakViewport = InViewport.ToWeakPtr()]()
+					{
+						const TSharedPtr<::SLevelViewport> Viewport = WeakViewport.Pin();
+						check(Viewport.IsValid());
+						FLevelEditorViewportClient& ViewportClient = Viewport->GetLevelViewportClient();
+						return ViewportClient.IsViewModeEnabled(VMI_VisualizeVirtualShadowMap);
+					}
+				)
+			),
 			EUserInterfaceActionType::RadioButton,
 			/* bInOpenSubMenuOnClick = */ false,
-			FSlateIcon(FAppStyle::GetAppStyleSetName(), "EditorViewport.VisualizeVirtualShadowMapMode"));
+			FSlateIcon(FAppStyle::GetAppStyleSetName(), "EditorViewport.VisualizeVirtualShadowMapMode")
+		);
 	}
 
 	{
-		auto BuildActorColorationMenu = [WeakViewport = InViewport.ToWeakPtr()](UToolMenu* InMenu) {
-			FToolMenuSection& SubMenuSection = InMenu->AddSection(
-				"LevelViewportActorColoration", LOCTEXT("ActorColorationHeader", "Actor Coloration"));
+		auto BuildActorColorationMenu = [WeakViewport = InViewport.ToWeakPtr()](UToolMenu* InMenu)
+		{
+			FToolMenuSection& SubMenuSection =
+				InMenu->AddSection("LevelViewportActorColoration", LOCTEXT("ActorColorationHeader", "Actor Coloration"));
 
 			TArray<FActorPrimitiveColorHandler::FPrimitiveColorHandler> PrimitiveColorHandlers;
 			FActorPrimitiveColorHandler::Get().GetRegisteredPrimitiveColorHandlers(PrimitiveColorHandlers);
@@ -245,135 +306,192 @@ void PopulateViewModesMenu(UToolMenu* InMenu, TSharedRef<::SLevelViewport> InVie
 					continue;
 				}
 
-				SubMenuSection.AddMenuEntry(NAME_None, PrimitiveColorHandler.HandlerText, FText(), FSlateIcon(),
-					FUIAction(FExecuteAction::CreateLambda([WeakViewport, PrimitiveColorHandler]() {
-						if (TSharedPtr<::SLevelViewport> Viewport = WeakViewport.Pin())
-						{
-							const bool bActorColorationEnabled = Viewport->GetLevelViewportClient().HandleIsShowFlagEnabled(
-								FEngineShowFlags::EShowFlag::SF_ActorColoration);
-
-							if (PrimitiveColorHandler.HandlerName.IsNone())
+				SubMenuSection.AddMenuEntry(
+					NAME_None,
+					PrimitiveColorHandler.HandlerText,
+					FText(),
+					FSlateIcon(),
+					FUIAction(
+						FExecuteAction::CreateLambda(
+							[WeakViewport, PrimitiveColorHandler]()
 							{
-								if (bActorColorationEnabled)
+								if (TSharedPtr<::SLevelViewport> Viewport = WeakViewport.Pin())
 								{
-									Viewport->GetLevelViewportClient().HandleToggleShowFlag(
-										FEngineShowFlags::EShowFlag::SF_ActorColoration);
+									const bool bActorColorationEnabled =
+										Viewport->GetLevelViewportClient().HandleIsShowFlagEnabled(
+											FEngineShowFlags::EShowFlag::SF_ActorColoration
+										);
+
+									if (PrimitiveColorHandler.HandlerName.IsNone())
+									{
+										if (bActorColorationEnabled)
+										{
+											Viewport->GetLevelViewportClient().HandleToggleShowFlag(
+												FEngineShowFlags::EShowFlag::SF_ActorColoration
+											);
+										}
+									}
+									else
+									{
+										if (!bActorColorationEnabled)
+										{
+											Viewport->GetLevelViewportClient().HandleToggleShowFlag(
+												FEngineShowFlags::EShowFlag::SF_ActorColoration
+											);
+										}
+
+										FActorPrimitiveColorHandler::Get().SetActivePrimitiveColorHandler(
+											PrimitiveColorHandler.HandlerName, GWorld
+										);
+									}
 								}
 							}
-							else
+						),
+						FCanExecuteAction::CreateLambda(
+							[WeakViewport]()
 							{
-								if (!bActorColorationEnabled)
+								if (TSharedPtr<::SLevelViewport> Viewport = WeakViewport.Pin())
 								{
-									Viewport->GetLevelViewportClient().HandleToggleShowFlag(
-										FEngineShowFlags::EShowFlag::SF_ActorColoration);
+									return true;
 								}
-
-								FActorPrimitiveColorHandler::Get().SetActivePrimitiveColorHandler(
-									PrimitiveColorHandler.HandlerName, GWorld);
+								return false;
 							}
-						}
-					}),
-						FCanExecuteAction::CreateLambda([WeakViewport]() {
-							if (TSharedPtr<::SLevelViewport> Viewport = WeakViewport.Pin())
+						),
+						FGetActionCheckState::CreateLambda(
+							[WeakViewport, PrimitiveColorHandler]()
 							{
-								return true;
-							}
-							return false;
-						}),
-						FGetActionCheckState::CreateLambda([WeakViewport, PrimitiveColorHandler]() {
-							if (TSharedPtr<::SLevelViewport> Viewport = WeakViewport.Pin())
-							{
-								const bool bActorColorationEnabled =
-									Viewport->GetLevelViewportClient().HandleIsShowFlagEnabled(
-										FEngineShowFlags::EShowFlag::SF_ActorColoration);
-
-								if (PrimitiveColorHandler.HandlerName.IsNone())
-								{
-									return bActorColorationEnabled ? ECheckBoxState::Unchecked : ECheckBoxState::Checked;
-								}
-								else
+								if (TSharedPtr<::SLevelViewport> Viewport = WeakViewport.Pin())
 								{
-									if (bActorColorationEnabled)
+									const bool bActorColorationEnabled =
+										Viewport->GetLevelViewportClient().HandleIsShowFlagEnabled(
+											FEngineShowFlags::EShowFlag::SF_ActorColoration
+										);
+
+									if (PrimitiveColorHandler.HandlerName.IsNone())
+									{
+										return bActorColorationEnabled ? ECheckBoxState::Unchecked
+																	   : ECheckBoxState::Checked;
+									}
+									else
 									{
-										return FActorPrimitiveColorHandler::Get().GetActivePrimitiveColorHandler()
-													== PrimitiveColorHandler.HandlerName
-												 ? ECheckBoxState::Checked
-												 : ECheckBoxState::Unchecked;
+										if (bActorColorationEnabled)
+										{
+											return FActorPrimitiveColorHandler::Get().GetActivePrimitiveColorHandler()
+														== PrimitiveColorHandler.HandlerName
+													 ? ECheckBoxState::Checked
+													 : ECheckBoxState::Unchecked;
+										}
 									}
 								}
-							}
 
-							return ECheckBoxState::Unchecked;
-						})),
-					EUserInterfaceActionType::RadioButton);
+								return ECheckBoxState::Unchecked;
+							}
+						)
+					),
+					EUserInterfaceActionType::RadioButton
+				);
 			}
 		};
 
 		FToolMenuSection& Section = InMenu->FindOrAddSection("ViewMode");
-		Section.AddSubMenu("ActorColoration", LOCTEXT("ActorColorationDisplayName", "Actor Coloration"),
+		Section.AddSubMenu(
+			"ActorColoration",
+			LOCTEXT("ActorColorationDisplayName", "Actor Coloration"),
 			LOCTEXT("ActorColorationMenu_ToolTip", "Override Actor Coloration mode"),
 			FNewToolMenuDelegate::CreateLambda(BuildActorColorationMenu),
-			FUIAction(FExecuteAction(), FCanExecuteAction(),
-				FIsActionChecked::CreateLambda([WeakViewport = InViewport.ToWeakPtr()]() {
-					if (const TSharedPtr<::SLevelViewport> Viewport = WeakViewport.Pin())
+			FUIAction(
+				FExecuteAction(),
+				FCanExecuteAction(),
+				FIsActionChecked::CreateLambda(
+					[WeakViewport = InViewport.ToWeakPtr()]()
 					{
-						return Viewport->GetLevelViewportClient().HandleIsShowFlagEnabled(
-							FEngineShowFlags::EShowFlag::SF_ActorColoration);
+						if (const TSharedPtr<::SLevelViewport> Viewport = WeakViewport.Pin())
+						{
+							return Viewport->GetLevelViewportClient().HandleIsShowFlagEnabled(
+								FEngineShowFlags::EShowFlag::SF_ActorColoration
+							);
+						}
+						return false;
 					}
-					return false;
-				})),
+				)
+			),
 			EUserInterfaceActionType::RadioButton,
 			/*bInOpenSubMenuOnClick=*/false,
-			FSlateIcon(FAppStyle::GetAppStyleSetName(), "EditorViewport.LODColorationMode"));
+			FSlateIcon(FAppStyle::GetAppStyleSetName(), "EditorViewport.LODColorationMode")
+		);
 	}
 
 	{
-		FToolMenuSection& Section = InMenu->AddSection(
-			"LevelViewportLandscape", LOCTEXT("LandscapeHeader", "Landscape"), InsertPosition);
+		FToolMenuSection& Section =
+			InMenu->AddSection("LevelViewportLandscape", LOCTEXT("LandscapeHeader", "Landscape"), InsertPosition);
 
-		auto BuildLandscapeLODMenu = [WeakViewport = InViewport.ToWeakPtr()](UToolMenu* InMenu) {
-			FToolMenuSection& SubMenuSection = InMenu->AddSection(
-				"LevelViewportLandScapeLOD", LOCTEXT("LandscapeLODHeader", "Landscape LOD"));
+		auto BuildLandscapeLODMenu = [WeakViewport = InViewport.ToWeakPtr()](UToolMenu* InMenu)
+		{
+			FToolMenuSection& SubMenuSection =
+				InMenu->AddSection("LevelViewportLandScapeLOD", LOCTEXT("LandscapeLODHeader", "Landscape LOD"));
 
-			auto CreateLandscapeLODAction = [WeakViewport](int32 LODValue) {
+			auto CreateLandscapeLODAction = [WeakViewport](int32 LODValue)
+			{
 				FUIAction LandscapeLODAction;
-				LandscapeLODAction.ExecuteAction = FExecuteAction::CreateLambda([WeakViewport, LODValue]() {
-					if (const TSharedPtr<::SLevelViewport> Viewport = WeakViewport.Pin())
+				LandscapeLODAction.ExecuteAction = FExecuteAction::CreateLambda(
+					[WeakViewport, LODValue]()
 					{
-						UE::LevelEditor::Private::OnLandscapeLODChanged(Viewport->GetLevelViewportClient(), LODValue);
+						if (const TSharedPtr<::SLevelViewport> Viewport = WeakViewport.Pin())
+						{
+							UE::LevelEditor::Private::OnLandscapeLODChanged(Viewport->GetLevelViewportClient(), LODValue);
+						}
 					}
-				});
+				);
 				LandscapeLODAction.GetActionCheckState = FGetActionCheckState::CreateLambda(
-					[WeakViewport, LODValue]() -> ECheckBoxState {
+					[WeakViewport, LODValue]() -> ECheckBoxState
+					{
 						bool bChecked = false;
 						if (const TSharedPtr<::SLevelViewport> Viewport = WeakViewport.Pin())
 						{
 							bChecked = UE::LevelEditor::Private::IsLandscapeLODSettingChecked(
-								Viewport->GetLevelViewportClient(), LODValue);
+								Viewport->GetLevelViewportClient(), LODValue
+							);
 						}
 						return bChecked ? ECheckBoxState::Checked : ECheckBoxState::Unchecked;
-					});
+					}
+				);
 
 				return LandscapeLODAction;
 			};
 
-			SubMenuSection.AddMenuEntry("LandscapeLODAuto", LOCTEXT("LandscapeLODAuto", "Auto"), FText(), FSlateIcon(),
-				CreateLandscapeLODAction(-1), EUserInterfaceActionType::RadioButton);
+			SubMenuSection.AddMenuEntry(
+				"LandscapeLODAuto",
+				LOCTEXT("LandscapeLODAuto", "Auto"),
+				FText(),
+				FSlateIcon(),
+				CreateLandscapeLODAction(-1),
+				EUserInterfaceActionType::RadioButton
+			);
 
 			SubMenuSection.AddSeparator("LandscapeLODSeparator");
 
 			static const FText FormatString = LOCTEXT("LandscapeLODFixed", "Fixed at {0}");
 			for (int32 i = 0; i < 8; ++i)
 			{
-				SubMenuSection.AddMenuEntry(NAME_None, FText::Format(FormatString, FText::AsNumber(i)), FText(),
-					FSlateIcon(), CreateLandscapeLODAction(i), EUserInterfaceActionType::RadioButton);
+				SubMenuSection.AddMenuEntry(
+					NAME_None,
+					FText::Format(FormatString, FText::AsNumber(i)),
+					FText(),
+					FSlateIcon(),
+					CreateLandscapeLODAction(i),
+					EUserInterfaceActionType::RadioButton
+				);
 			}
 		};
 
-		Section.AddSubMenu("LandscapeLOD", LOCTEXT("LandscapeLODDisplayName", "LOD"),
+		Section.AddSubMenu(
+			"LandscapeLOD",
+			LOCTEXT("LandscapeLODDisplayName", "LOD"),
 			LOCTEXT("LandscapeLODMenu_ToolTip", "Override Landscape LOD in this viewport"),
 			FNewToolMenuDelegate::CreateLambda(BuildLandscapeLODMenu),
-			/*bInOpenSubMenuOnClick=*/false, FSlateIcon(FAppStyle::GetAppStyleSetName(), "EditorViewport.LOD"));
+			/*bInOpenSubMenuOnClick=*/false,
+			FSlateIcon(FAppStyle::GetAppStyleSetName(), "EditorViewport.LOD")
+		);
 	}
 }
 
@@ -381,200 +499,271 @@ void AddViewportToolbarViewModesSubmenu(FToolMenuSection& InSection)
 {
 	// This has to be a dynamic entry for the ViewModes submenu's label to be able to access the context.
 	InSection.AddDynamicEntry(
-		"DynamicViewModes", FNewToolMenuSectionDelegate::CreateLambda([](FToolMenuSection& InDynamicSection) -> void {
-			TAttribute<FText> LabelAttribute = UE::UnrealEd::GetViewModesSubmenuLabel(nullptr);
-			if (ULevelViewportContext* const LevelViewportContext = InDynamicSection.FindContext<ULevelViewportContext>())
+		"DynamicViewModes",
+		FNewToolMenuSectionDelegate::CreateLambda(
+			[](FToolMenuSection& InDynamicSection) -> void
 			{
-				TWeakPtr<SEditorViewport> EditorViewport = LevelViewportContext->LevelViewport;
-				LabelAttribute = TAttribute<FText>::CreateLambda(
-					[EditorViewport]() { return UE::UnrealEd::GetViewModesSubmenuLabel(EditorViewport); });
-			}
+				TAttribute<FText> LabelAttribute = UE::UnrealEd::GetViewModesSubmenuLabel(nullptr);
+				if (ULevelViewportContext* const LevelViewportContext =
+						InDynamicSection.FindContext<ULevelViewportContext>())
+				{
+					TWeakPtr<SEditorViewport> EditorViewport = LevelViewportContext->LevelViewport;
+					LabelAttribute = TAttribute<FText>::CreateLambda(
+						[EditorViewport]()
+						{
+							return UE::UnrealEd::GetViewModesSubmenuLabel(EditorViewport);
+						}
+					);
+				}
 
-			InDynamicSection.AddSubMenu("ViewModes", LabelAttribute,
-				LOCTEXT("ViewModesSubmenuTooltip", "View mode settings for the current viewport."),
-				FNewToolMenuDelegate::CreateLambda([](UToolMenu* Submenu) -> void {
-					ULevelViewportContext* const LevelViewportContext = Submenu->FindContext<ULevelViewportContext>();
-					if (!LevelViewportContext)
-					{
-						return;
-					}
+				InDynamicSection.AddSubMenu(
+					"ViewModes",
+					LabelAttribute,
+					LOCTEXT("ViewModesSubmenuTooltip", "View mode settings for the current viewport."),
+					FNewToolMenuDelegate::CreateLambda(
+						[](UToolMenu* Submenu) -> void
+						{
+							ULevelViewportContext* const LevelViewportContext =
+								Submenu->FindContext<ULevelViewportContext>();
+							if (!LevelViewportContext)
+							{
+								return;
+							}
 
-					if (const TSharedPtr<::SLevelViewport> LevelViewport = LevelViewportContext->LevelViewport.Pin())
-					{
-						UE::UnrealEd::PopulateViewModesMenu(Submenu, LevelViewport.ToSharedRef());
-						PopulateViewModesMenu(Submenu, LevelViewport.ToSharedRef());
-					}
-				}));
-		}));
+							if (const TSharedPtr<::SLevelViewport> LevelViewport = LevelViewportContext->LevelViewport.Pin())
+							{
+								UE::UnrealEd::PopulateViewModesMenu(Submenu, LevelViewport.ToSharedRef());
+								PopulateViewModesMenu(Submenu, LevelViewport.ToSharedRef());
+							}
+						}
+					)
+				);
+			}
+		)
+	);
 }
 
 void AddFeatureLevelPreviewSubmenu(FToolMenuSection& Section)
 {
-	Section.AddSubMenu("FeatureLevelPreview",
+	Section.AddSubMenu(
+		"FeatureLevelPreview",
 		NSLOCTEXT("LevelToolBarViewMenu", "PreviewPlatformSubMenu", "Preview Platform"),
-		NSLOCTEXT("LevelToolBarViewMenu", "PreviewPlatformSubMenu_ToolTip",
-			"Sets the preview platform used by the main editor"),
-		FNewToolMenuDelegate::CreateLambda([](UToolMenu* InMenu) -> void {
-			FToolMenuSection& Section = InMenu->AddSection(
-				"EditorPreviewMode", LOCTEXT("EditorPreviewModeDevices", "Preview Devices"));
-			// Preview platforms discovered from ITargetPlatforms.
-			for (auto& Item : FLevelEditorCommands::Get().PreviewPlatformOverrides)
+		NSLOCTEXT("LevelToolBarViewMenu", "PreviewPlatformSubMenu_ToolTip", "Sets the preview platform used by the main editor"),
+		FNewToolMenuDelegate::CreateLambda(
+			[](UToolMenu* InMenu) -> void
 			{
-				Section.AddMenuEntry(Item);
+				FToolMenuSection& Section =
+					InMenu->AddSection("EditorPreviewMode", LOCTEXT("EditorPreviewModeDevices", "Preview Devices"));
+				// Preview platforms discovered from ITargetPlatforms.
+				for (auto& Item : FLevelEditorCommands::Get().PreviewPlatformOverrides)
+				{
+					Section.AddMenuEntry(Item);
+				}
 			}
-		}));
+		)
+	);
 }
 
 void AddMaterialQualityLevelSubmenu(FToolMenuSection& Section)
 {
-	Section.AddSubMenu("MaterialQualityLevel",
+	Section.AddSubMenu(
+		"MaterialQualityLevel",
 		NSLOCTEXT("LevelToolBarViewMenu", "MaterialQualityLevelSubMenu", "Material Quality Level"),
-		NSLOCTEXT("LevelToolBarViewMenu",
+		NSLOCTEXT(
+			"LevelToolBarViewMenu",
 			"MaterialQualityLevelSubMenu_ToolTip",
 			"Sets the value of the CVar \"r.MaterialQualityLevel\" (low=0, high=1, medium=2, Epic=3). This affects "
-			"materials via the QualitySwitch material expression."),
-		FNewToolMenuDelegate::CreateLambda([](UToolMenu* InMenu) -> void {
-			FToolMenuSection& Section = InMenu->AddSection("LevelEditorMaterialQualityLevel",
-				NSLOCTEXT("LevelToolBarViewMenu", "MaterialQualityLevelHeading", "Material Quality Level"));
-			Section.AddMenuEntry(FLevelEditorCommands::Get().MaterialQualityLevel_Low);
-			Section.AddMenuEntry(FLevelEditorCommands::Get().MaterialQualityLevel_Medium);
-			Section.AddMenuEntry(FLevelEditorCommands::Get().MaterialQualityLevel_High);
-			Section.AddMenuEntry(FLevelEditorCommands::Get().MaterialQualityLevel_Epic);
-		}));
+			"materials via the QualitySwitch material expression."
+		),
+		FNewToolMenuDelegate::CreateLambda(
+			[](UToolMenu* InMenu) -> void
+			{
+				FToolMenuSection& Section = InMenu->AddSection(
+					"LevelEditorMaterialQualityLevel",
+					NSLOCTEXT("LevelToolBarViewMenu", "MaterialQualityLevelHeading", "Material Quality Level")
+				);
+				Section.AddMenuEntry(FLevelEditorCommands::Get().MaterialQualityLevel_Low);
+				Section.AddMenuEntry(FLevelEditorCommands::Get().MaterialQualityLevel_Medium);
+				Section.AddMenuEntry(FLevelEditorCommands::Get().MaterialQualityLevel_High);
+				Section.AddMenuEntry(FLevelEditorCommands::Get().MaterialQualityLevel_Epic);
+			}
+		)
+	);
 }
 
 void AddViewportToolbarPerformanceAndScalabilitySubmenu(FToolMenuSection& InSection)
 {
-	InSection.AddSubMenu("PerformanceAndScalability",
+	InSection.AddSubMenu(
+		"PerformanceAndScalability",
 		LOCTEXT("PerformanceAndScalabilityLabel", "Performance & Scalability"),
 		LOCTEXT("PerformanceAndScalabilityTooltip", "Performance and scalability tools tied to this viewport."),
-		FNewToolMenuDelegate::CreateLambda([](UToolMenu* Submenu) -> void {
+		FNewToolMenuDelegate::CreateLambda(
+			[](UToolMenu* Submenu) -> void
 			{
-				FToolMenuSection& UnnamedSection = Submenu->FindOrAddSection(NAME_None);
-
-				// Add realtime rendering toggle.
-				UnnamedSection.AddDynamicEntry("ToggleRealtimeDynamicSection",
-					FNewToolMenuSectionDelegate::CreateLambda([](FToolMenuSection& InnerSection) -> void {
-						FToolUIAction RealtimeToggleAction;
-						RealtimeToggleAction.ExecuteAction = FToolMenuExecuteAction::CreateLambda(
-							[](const FToolMenuContext& Context) -> void {
-								ULevelViewportContext* const LevelViewportContext =
-									Context.FindContext<ULevelViewportContext>();
-								if (!LevelViewportContext)
-								{
-									return;
-								}
+				{
+					FToolMenuSection& UnnamedSection = Submenu->FindOrAddSection(NAME_None);
 
-								if (const TSharedPtr<::SLevelViewport> LevelViewport =
-										LevelViewportContext->LevelViewport.Pin())
-								{
-									LevelViewport->OnToggleRealtime();
-									UToolMenus::Get()->RefreshAllWidgets();
-								}
-							});
+					// Add realtime rendering toggle.
+					UnnamedSection.AddDynamicEntry(
+						"ToggleRealtimeDynamicSection",
+						FNewToolMenuSectionDelegate::CreateLambda(
+							[](FToolMenuSection& InnerSection) -> void
+							{
+								FToolUIAction RealtimeToggleAction;
+								RealtimeToggleAction.ExecuteAction = FToolMenuExecuteAction::CreateLambda(
+									[](const FToolMenuContext& Context) -> void
+									{
+										ULevelViewportContext* const LevelViewportContext =
+											Context.FindContext<ULevelViewportContext>();
+										if (!LevelViewportContext)
+										{
+											return;
+										}
+
+										if (const TSharedPtr<::SLevelViewport> LevelViewport =
+												LevelViewportContext->LevelViewport.Pin())
+										{
+											LevelViewport->OnToggleRealtime();
+											UToolMenus::Get()->RefreshAllWidgets();
+										}
+									}
+								);
 
-						RealtimeToggleAction.GetActionCheckState = FToolMenuGetActionCheckState::CreateLambda(
-							[](const FToolMenuContext& Context) -> ECheckBoxState {
-								ULevelViewportContext* const LevelViewportContext =
-									Context.FindContext<ULevelViewportContext>();
-								if (!LevelViewportContext)
-								{
-									return ECheckBoxState::Undetermined;
-								}
+								RealtimeToggleAction.GetActionCheckState = FToolMenuGetActionCheckState::CreateLambda(
+									[](const FToolMenuContext& Context) -> ECheckBoxState
+									{
+										ULevelViewportContext* const LevelViewportContext =
+											Context.FindContext<ULevelViewportContext>();
+										if (!LevelViewportContext)
+										{
+											return ECheckBoxState::Undetermined;
+										}
+
+										// Check if the realtime warn state is outdated and if so refresh widgets to
+										// update our top-level status.
+										if (const TOptional<bool> IsDirty =
+												Private::IsDirtyRealtimeWarningFromContext(Context);
+											IsDirty.IsSet() && IsDirty.GetValue())
+										{
+											UToolMenus::Get()->RefreshAllWidgets();
+										}
+
+										if (const TSharedPtr<::SLevelViewport> LevelViewport =
+												LevelViewportContext->LevelViewport.Pin())
+										{
+											return LevelViewport->IsRealtime() ? ECheckBoxState::Checked
+																			   : ECheckBoxState::Unchecked;
+										}
+
+										return ECheckBoxState::Undetermined;
+									}
+								);
 
-								// Check if the realtime warn state is outdated and if so refresh widgets to update our
-								// top-level status.
-								if (const TOptional<bool> IsDirty = Private::IsDirtyRealtimeWarningFromContext(Context);
-									IsDirty.IsSet() && IsDirty.GetValue())
+								bool bDisplayTopLevel = false;
+								if (const TOptional<bool> ShouldWarn =
+										Private::UpdateAndGetRealtimeWarningFromContext(InnerSection.Context);
+									ShouldWarn.IsSet())
 								{
-									UToolMenus::Get()->RefreshAllWidgets();
+									bDisplayTopLevel = ShouldWarn.GetValue();
 								}
-
-								if (const TSharedPtr<::SLevelViewport> LevelViewport =
-										LevelViewportContext->LevelViewport.Pin())
+								else
 								{
-									return LevelViewport->IsRealtime() ? ECheckBoxState::Checked
-																	   : ECheckBoxState::Unchecked;
+									// If we couldn't get the warn state, pretend we don't have to warn.
+									bDisplayTopLevel = false;
 								}
 
-								return ECheckBoxState::Undetermined;
-							});
+								const FText Tooltip =
+									bDisplayTopLevel ? LOCTEXT(
+										"ToggleRealtimeTooltip_WarnRealtimeOff",
+										"This viewport is not updating in realtime.  Click to turn on realtime mode."
+									)
+													 : LOCTEXT(
+														 "ToggleRealtimeTooltip", "Toggle realtime rendering of the viewport"
+													 );
+
+								FToolMenuEntry ToggleRealtime = FToolMenuEntry::InitMenuEntry(
+									"ToggleRealtime",
+									LOCTEXT("ToggleRealtimeLabel", "Realtime Viewport"),
+									Tooltip,
+									FSlateIcon(FAppStyle::GetAppStyleSetName(), "EditorViewport.ToggleRealTime"),
+									RealtimeToggleAction,
+									EUserInterfaceActionType::ToggleButton
+								);
+								ToggleRealtime.SetShowInToolbarTopLevel(bDisplayTopLevel);
+								InnerSection.AddEntry(ToggleRealtime);
+							}
+						)
+					);
+				}
 
-						bool bDisplayTopLevel = false;
-						if (const TOptional<bool> ShouldWarn = Private::UpdateAndGetRealtimeWarningFromContext(
-								InnerSection.Context);
-							ShouldWarn.IsSet())
-						{
-							bDisplayTopLevel = ShouldWarn.GetValue();
-						}
-						else
-						{
-							// If we couldn't get the warn state, pretend we don't have to warn.
-							bDisplayTopLevel = false;
-						}
+				{
+					FToolMenuSection& PerformanceAndScalabilitySection = Submenu->FindOrAddSection(
+						"PerformanceAndScalability",
+						LOCTEXT("PerformanceAndScalabilitySectionLabel", "Performance & Scalability")
+					);
 
-						const FText Tooltip =
-							bDisplayTopLevel
-								? LOCTEXT("ToggleRealtimeTooltip_WarnRealtimeOff",
-									"This viewport is not updating in realtime.  Click to turn on realtime mode.")
-								: LOCTEXT("ToggleRealtimeTooltip", "Toggle realtime rendering of the viewport");
-
-						FToolMenuEntry ToggleRealtime = FToolMenuEntry::InitMenuEntry("ToggleRealtime",
-							LOCTEXT("ToggleRealtimeLabel", "Realtime Viewport"), Tooltip,
-							FSlateIcon(FAppStyle::GetAppStyleSetName(), "EditorViewport.ToggleRealTime"),
-							RealtimeToggleAction, EUserInterfaceActionType::ToggleButton);
-						ToggleRealtime.SetShowInToolbarTopLevel(bDisplayTopLevel);
-						InnerSection.AddEntry(ToggleRealtime);
-					}));
-			}
+					AddFeatureLevelPreviewSubmenu(PerformanceAndScalabilitySection);
 
-			{
-				FToolMenuSection& PerformanceAndScalabilitySection = Submenu->FindOrAddSection(
-					"PerformanceAndScalability",
-					LOCTEXT("PerformanceAndScalabilitySectionLabel", "Performance & Scalability"));
-
-				AddFeatureLevelPreviewSubmenu(PerformanceAndScalabilitySection);
-
-				PerformanceAndScalabilitySection.AddSeparator("PerformanceAndScalabilitySettings");
-
-				PerformanceAndScalabilitySection.AddSubMenu("Scalability",
-					LOCTEXT("ScalabilitySubMenu", "Engine Scalability"),
-					LOCTEXT("ScalabilitySubMenu_ToolTip", "Open the engine scalability settings"),
-					FNewToolMenuDelegate::CreateLambda([](UToolMenu* InMenu) -> void {
-						FToolMenuSection& Section = InMenu->FindOrAddSection(NAME_None);
-						Section.AddEntry(
-							FToolMenuEntry::InitWidget("ScalabilitySettings", SNew(SScalabilitySettings), FText(), true));
-					}));
-
-				AddMaterialQualityLevelSubmenu(PerformanceAndScalabilitySection);
-
-				// FEditorViewportClient& ViewportClient = Viewport.Pin()->GetLevelViewportClient();
-				PerformanceAndScalabilitySection.AddSubMenu("ScreenPercentageSubMenu",
-					LOCTEXT("ScreenPercentageSubMenu", "Screen Percentage"),
-					LOCTEXT("ScreenPercentageSubMenu_ToolTip", "Customize the viewport's screen percentage"),
-					FNewToolMenuDelegate::CreateLambda([](UToolMenu* ScreenPercentageSubMenu) {
-						FToolMenuSection& UnnamedSection = ScreenPercentageSubMenu->FindOrAddSection(NAME_None);
-
-						ScreenPercentageSubMenu->AddDynamicSection(NAME_None,
-							FNewToolMenuDelegateLegacy::CreateLambda([](FMenuBuilder& MenuBuilder, UToolMenu* InMenu) -> void {
-								ULevelViewportContext* const LevelViewportContext =
-									InMenu->FindContext<ULevelViewportContext>();
-								if (!LevelViewportContext)
-								{
-									return;
-								}
+					PerformanceAndScalabilitySection.AddSeparator("PerformanceAndScalabilitySettings");
 
-								if (const TSharedPtr<::SLevelViewport> LevelViewport =
-										LevelViewportContext->LevelViewport.Pin())
-								{
-									TSharedPtr<FEditorViewportClient> Client = LevelViewport->GetViewportClient();
-									SCommonEditorViewportToolbarBase::ConstructScreenPercentageMenu(
-										MenuBuilder, Client.Get());
-								}
-							}));
-					}));
+					PerformanceAndScalabilitySection.AddSubMenu(
+						"Scalability",
+						LOCTEXT("ScalabilitySubMenu", "Engine Scalability"),
+						LOCTEXT("ScalabilitySubMenu_ToolTip", "Open the engine scalability settings"),
+						FNewToolMenuDelegate::CreateLambda(
+							[](UToolMenu* InMenu) -> void
+							{
+								FToolMenuSection& Section = InMenu->FindOrAddSection(NAME_None);
+								Section.AddEntry(FToolMenuEntry::InitWidget(
+									"ScalabilitySettings", SNew(SScalabilitySettings), FText(), true
+								));
+							}
+						)
+					);
+
+					AddMaterialQualityLevelSubmenu(PerformanceAndScalabilitySection);
+
+					// FEditorViewportClient& ViewportClient = Viewport.Pin()->GetLevelViewportClient();
+					PerformanceAndScalabilitySection.AddSubMenu(
+						"ScreenPercentageSubMenu",
+						LOCTEXT("ScreenPercentageSubMenu", "Screen Percentage"),
+						LOCTEXT("ScreenPercentageSubMenu_ToolTip", "Customize the viewport's screen percentage"),
+						FNewToolMenuDelegate::CreateLambda(
+							[](UToolMenu* ScreenPercentageSubMenu)
+							{
+								FToolMenuSection& UnnamedSection = ScreenPercentageSubMenu->FindOrAddSection(NAME_None);
+
+								ScreenPercentageSubMenu->AddDynamicSection(
+									NAME_None,
+									FNewToolMenuDelegateLegacy::CreateLambda(
+										[](FMenuBuilder& MenuBuilder, UToolMenu* InMenu) -> void
+										{
+											ULevelViewportContext* const LevelViewportContext =
+												InMenu->FindContext<ULevelViewportContext>();
+											if (!LevelViewportContext)
+											{
+												return;
+											}
+
+											if (const TSharedPtr<::SLevelViewport> LevelViewport =
+													LevelViewportContext->LevelViewport.Pin())
+											{
+												TSharedPtr<FEditorViewportClient> Client =
+													LevelViewport->GetViewportClient();
+												SCommonEditorViewportToolbarBase::ConstructScreenPercentageMenu(
+													MenuBuilder, Client.Get()
+												);
+											}
+										}
+									)
+								);
+							}
+						)
+					);
+				}
 			}
-		}));
+		)
+	);
 }
 
 void GenerateViewportLayoutsMenu(UToolMenu* InMenu, TSharedPtr<::SLevelViewport> InViewport)
@@ -586,8 +775,8 @@ void GenerateViewportLayoutsMenu(UToolMenu* InMenu, TSharedPtr<::SLevelViewport>
 	InMenu->bSearchable = false;
 
 	{
-		FToolMenuSection& Section = InMenu->AddSection(
-			"LevelViewportOnePaneConfigs", LOCTEXT("OnePaneConfigHeader", "One Pane"));
+		FToolMenuSection& Section =
+			InMenu->AddSection("LevelViewportOnePaneConfigs", LOCTEXT("OnePaneConfigHeader", "One Pane"));
 
 		FSlimHorizontalToolBarBuilder OnePaneButton(CommandList, FMultiBoxCustomization::None);
 		OnePaneButton.SetLabelVisibility(EVisibility::Collapsed);
@@ -595,7 +784,8 @@ void GenerateViewportLayoutsMenu(UToolMenu* InMenu, TSharedPtr<::SLevelViewport>
 
 		OnePaneButton.AddToolBarButton(FLevelViewportCommands::Get().ViewportConfig_OnePane);
 
-		Section.AddEntry(FToolMenuEntry::InitWidget("LevelViewportOnePaneConfigs",
+		Section.AddEntry(FToolMenuEntry::InitWidget(
+			"LevelViewportOnePaneConfigs",
 			// clang-format off
 			SNew(SHorizontalBox)
 			+ SHorizontalBox::Slot()
@@ -609,12 +799,14 @@ void GenerateViewportLayoutsMenu(UToolMenu* InMenu, TSharedPtr<::SLevelViewport>
 				SNullWidget::NullWidget
 			],
 			// clang-format on
-			FText::GetEmpty(), true));
+			FText::GetEmpty(),
+			true
+		));
 	}
 
 	{
-		FToolMenuSection& Section = InMenu->AddSection(
-			"LevelViewportTwoPaneConfigs", LOCTEXT("TwoPaneConfigHeader", "Two Panes"));
+		FToolMenuSection& Section =
+			InMenu->AddSection("LevelViewportTwoPaneConfigs", LOCTEXT("TwoPaneConfigHeader", "Two Panes"));
 		FSlimHorizontalToolBarBuilder TwoPaneButtons(CommandList, FMultiBoxCustomization::None);
 		TwoPaneButtons.SetLabelVisibility(EVisibility::Collapsed);
 		TwoPaneButtons.SetStyle(&FAppStyle::Get(), "ViewportLayoutToolbar");
@@ -622,7 +814,8 @@ void GenerateViewportLayoutsMenu(UToolMenu* InMenu, TSharedPtr<::SLevelViewport>
 		TwoPaneButtons.AddToolBarButton(FLevelViewportCommands::Get().ViewportConfig_TwoPanesH, NAME_None, FText());
 		TwoPaneButtons.AddToolBarButton(FLevelViewportCommands::Get().ViewportConfig_TwoPanesV, NAME_None, FText());
 
-		Section.AddEntry(FToolMenuEntry::InitWidget("LevelViewportTwoPaneConfigs",
+		Section.AddEntry(FToolMenuEntry::InitWidget(
+			"LevelViewportTwoPaneConfigs",
 			// clang-format off
 			SNew(SHorizontalBox)
 			+SHorizontalBox::Slot()
@@ -636,12 +829,14 @@ void GenerateViewportLayoutsMenu(UToolMenu* InMenu, TSharedPtr<::SLevelViewport>
 				SNullWidget::NullWidget
 			],
 			// clang-format on
-			FText::GetEmpty(), true));
+			FText::GetEmpty(),
+			true
+		));
 	}
 
 	{
-		FToolMenuSection& Section = InMenu->AddSection(
-			"LevelViewportThreePaneConfigs", LOCTEXT("ThreePaneConfigHeader", "Three Panes"));
+		FToolMenuSection& Section =
+			InMenu->AddSection("LevelViewportThreePaneConfigs", LOCTEXT("ThreePaneConfigHeader", "Three Panes"));
 		FSlimHorizontalToolBarBuilder ThreePaneButtons(CommandList, FMultiBoxCustomization::None);
 		ThreePaneButtons.SetLabelVisibility(EVisibility::Collapsed);
 		ThreePaneButtons.SetStyle(&FAppStyle::Get(), "ViewportLayoutToolbar");
@@ -649,10 +844,10 @@ void GenerateViewportLayoutsMenu(UToolMenu* InMenu, TSharedPtr<::SLevelViewport>
 		ThreePaneButtons.AddToolBarButton(FLevelViewportCommands::Get().ViewportConfig_ThreePanesLeft, NAME_None, FText());
 		ThreePaneButtons.AddToolBarButton(FLevelViewportCommands::Get().ViewportConfig_ThreePanesRight, NAME_None, FText());
 		ThreePaneButtons.AddToolBarButton(FLevelViewportCommands::Get().ViewportConfig_ThreePanesTop, NAME_None, FText());
-		ThreePaneButtons.AddToolBarButton(
-			FLevelViewportCommands::Get().ViewportConfig_ThreePanesBottom, NAME_None, FText());
+		ThreePaneButtons.AddToolBarButton(FLevelViewportCommands::Get().ViewportConfig_ThreePanesBottom, NAME_None, FText());
 
-		Section.AddEntry(FToolMenuEntry::InitWidget("LevelViewportThreePaneConfigs",
+		Section.AddEntry(FToolMenuEntry::InitWidget(
+			"LevelViewportThreePaneConfigs",
 			// clang-format off
 			SNew(SHorizontalBox)
 			+SHorizontalBox::Slot()
@@ -666,12 +861,14 @@ void GenerateViewportLayoutsMenu(UToolMenu* InMenu, TSharedPtr<::SLevelViewport>
 				SNullWidget::NullWidget
 			],
 			// clang-format on
-			FText::GetEmpty(), true));
+			FText::GetEmpty(),
+			true
+		));
 	}
 
 	{
-		FToolMenuSection& Section = InMenu->AddSection(
-			"LevelViewportFourPaneConfigs", LOCTEXT("FourPaneConfigHeader", "Four Panes"));
+		FToolMenuSection& Section =
+			InMenu->AddSection("LevelViewportFourPaneConfigs", LOCTEXT("FourPaneConfigHeader", "Four Panes"));
 		FSlimHorizontalToolBarBuilder FourPaneButtons(CommandList, FMultiBoxCustomization::None);
 		FourPaneButtons.SetLabelVisibility(EVisibility::Collapsed);
 		FourPaneButtons.SetStyle(&FAppStyle::Get(), "ViewportLayoutToolbar");
@@ -682,7 +879,8 @@ void GenerateViewportLayoutsMenu(UToolMenu* InMenu, TSharedPtr<::SLevelViewport>
 		FourPaneButtons.AddToolBarButton(FLevelViewportCommands::Get().ViewportConfig_FourPanesTop, NAME_None, FText());
 		FourPaneButtons.AddToolBarButton(FLevelViewportCommands::Get().ViewportConfig_FourPanesBottom, NAME_None, FText());
 
-		Section.AddEntry(FToolMenuEntry::InitWidget("LevelViewportFourPaneConfigs",
+		Section.AddEntry(FToolMenuEntry::InitWidget(
+			"LevelViewportFourPaneConfigs",
 			// clang-format off
 			SNew(SHorizontalBox)
 			+SHorizontalBox::Slot()
@@ -696,112 +894,146 @@ void GenerateViewportLayoutsMenu(UToolMenu* InMenu, TSharedPtr<::SLevelViewport>
 				SNullWidget::NullWidget
 			],
 			// clang-format on
-			FText::GetEmpty(), true));
+			FText::GetEmpty(),
+			true
+		));
 	}
 }
 
 TSharedRef<SWidget> BuildVolumeControlCustomWidget()
 {
+	// clang-format off
 	return SNew(SHorizontalBox)
 		+ SHorizontalBox::Slot()
 		.FillWidth(0.9f)
 		.Padding(FMargin(2.0f, 0.0f, 0.0f, 0.0f))
 		[
-		SNew(SVolumeControl)
-			.ToolTipText_Static(&FLevelEditorActionCallbacks::GetAudioVolumeToolTip)
-			.Volume_Static(&FLevelEditorActionCallbacks::GetAudioVolume)
-			.OnVolumeChanged_Static(&FLevelEditorActionCallbacks::OnAudioVolumeChanged)
-			.Muted_Static(&FLevelEditorActionCallbacks::GetAudioMuted)
-			.OnMuteChanged_Static(&FLevelEditorActionCallbacks::OnAudioMutedChanged)
+			SNew(SVolumeControl)
+				.ToolTipText_Static(&FLevelEditorActionCallbacks::GetAudioVolumeToolTip)
+				.Volume_Static(&FLevelEditorActionCallbacks::GetAudioVolume)
+				.OnVolumeChanged_Static(&FLevelEditorActionCallbacks::OnAudioVolumeChanged)
+				.Muted_Static(&FLevelEditorActionCallbacks::GetAudioMuted)
+				.OnMuteChanged_Static(&FLevelEditorActionCallbacks::OnAudioMutedChanged)
 		]
 		+ SHorizontalBox::Slot()
 		.FillWidth(0.1f);
+	// clang-format on
 }
 
 void AddLevelEditorViewportToolbarSettingsSubmenu(FToolMenuSection& InSection)
 {
-	InSection.AddSubMenu("Settings", LOCTEXT("SettingsSubmenuLabel", "Settings"),
+	InSection.AddSubMenu(
+		"Settings",
+		LOCTEXT("SettingsSubmenuLabel", "Settings"),
 		LOCTEXT("SettingsSubmenuTooltip", "Viewport-related settings"),
-		FNewToolMenuDelegate::CreateLambda([](UToolMenu* Submenu) -> void {
+		FNewToolMenuDelegate::CreateLambda(
+			[](UToolMenu* Submenu) -> void
 			{
-				FToolMenuSection& ViewportControlsSection = Submenu->FindOrAddSection(
-					"ViewportControls", LOCTEXT("ViewportControlsSectionLabel", "Viewport Controls"));
-
-				ViewportControlsSection.AddSubMenu("ViewportLayouts", LOCTEXT("ViewportLayoutsLabel", "Layouts"),
-					LOCTEXT("ViewportLayoutsTooltip", "Configure the layouts of the viewport windows"),
-					FNewToolMenuDelegate::CreateLambda([](UToolMenu* InMenu) {
-						ULevelViewportContext* const LevelViewportContext = InMenu->FindContext<ULevelViewportContext>();
-						if (!LevelViewportContext)
-						{
-							return;
-						}
-
-						if (const TSharedPtr<::SLevelViewport> LevelViewport = LevelViewportContext->LevelViewport.Pin())
-						{
-							GenerateViewportLayoutsMenu(InMenu, LevelViewport);
-						}
-					}),
-					false, FSlateIcon(FAppStyle::Get().GetStyleSetName(), "Icons.Layout"));
-			}
+				{
+					FToolMenuSection& ViewportControlsSection = Submenu->FindOrAddSection(
+						"ViewportControls", LOCTEXT("ViewportControlsSectionLabel", "Viewport Controls")
+					);
+
+					ViewportControlsSection.AddSubMenu(
+						"ViewportLayouts",
+						LOCTEXT("ViewportLayoutsLabel", "Layouts"),
+						LOCTEXT("ViewportLayoutsTooltip", "Configure the layouts of the viewport windows"),
+						FNewToolMenuDelegate::CreateLambda(
+							[](UToolMenu* InMenu)
+							{
+								ULevelViewportContext* const LevelViewportContext =
+									InMenu->FindContext<ULevelViewportContext>();
+								if (!LevelViewportContext)
+								{
+									return;
+								}
 
-			{
-				FToolMenuSection& SettingsSection = Submenu->FindOrAddSection(
-					"Settings", LOCTEXT("SettingsSectionLabel", "Settings"));
+								if (const TSharedPtr<::SLevelViewport> LevelViewport =
+										LevelViewportContext->LevelViewport.Pin())
+								{
+									GenerateViewportLayoutsMenu(InMenu, LevelViewport);
+								}
+							}
+						),
+						false,
+						FSlateIcon(FAppStyle::Get().GetStyleSetName(), "Icons.Layout")
+					);
+				}
 
-				SettingsSection.AddEntry(FToolMenuEntry::InitWidget(
-					"Volume", BuildVolumeControlCustomWidget(), LOCTEXT("VolumeControlLabel", "Volume")));
+				{
+					FToolMenuSection& SettingsSection =
+						Submenu->FindOrAddSection("Settings", LOCTEXT("SettingsSectionLabel", "Settings"));
 
-				SettingsSection.AddSeparator("ViewportSizeSeparator");
+					SettingsSection.AddEntry(FToolMenuEntry::InitWidget(
+						"Volume", BuildVolumeControlCustomWidget(), LOCTEXT("VolumeControlLabel", "Volume")
+					));
 
-				SettingsSection.AddMenuEntry(FLevelViewportCommands::Get().ToggleImmersive);
+					SettingsSection.AddSeparator("ViewportSizeSeparator");
 
-				{
-					FToolUIAction MaximizeRestoreAction;
-					MaximizeRestoreAction.ExecuteAction = FToolMenuExecuteAction::CreateLambda([](const FToolMenuContext& Context) {
-						ULevelViewportContext* const LevelViewportContext = Context.FindContext<ULevelViewportContext>();
-						if (!LevelViewportContext)
-						{
-							return;
-						}
+					SettingsSection.AddMenuEntry(FLevelViewportCommands::Get().ToggleImmersive);
 
-						if (const TSharedPtr<::SLevelViewport> LevelViewport = LevelViewportContext->LevelViewport.Pin())
-						{
-							LevelViewport->OnToggleMaximize();
-						}
-					});
-					MaximizeRestoreAction.GetActionCheckState = FToolMenuGetActionCheckState::CreateLambda(
-						[](const FToolMenuContext& Context) -> ECheckBoxState {
-							ULevelViewportContext* const LevelViewportContext =
-								Context.FindContext<ULevelViewportContext>();
-							if (!LevelViewportContext)
+					{
+						FToolUIAction MaximizeRestoreAction;
+						MaximizeRestoreAction.ExecuteAction = FToolMenuExecuteAction::CreateLambda(
+							[](const FToolMenuContext& Context)
 							{
-								return ECheckBoxState::Undetermined;
-							}
+								ULevelViewportContext* const LevelViewportContext =
+									Context.FindContext<ULevelViewportContext>();
+								if (!LevelViewportContext)
+								{
+									return;
+								}
 
-							if (const TSharedPtr<::SLevelViewport> LevelViewport = LevelViewportContext->LevelViewport.Pin())
-							{
-								return LevelViewport->IsMaximized() ? ECheckBoxState::Checked : ECheckBoxState::Unchecked;
+								if (const TSharedPtr<::SLevelViewport> LevelViewport =
+										LevelViewportContext->LevelViewport.Pin())
+								{
+									LevelViewport->OnToggleMaximize();
+								}
 							}
+						);
+						MaximizeRestoreAction.GetActionCheckState = FToolMenuGetActionCheckState::CreateLambda(
+							[](const FToolMenuContext& Context) -> ECheckBoxState
+							{
+								ULevelViewportContext* const LevelViewportContext =
+									Context.FindContext<ULevelViewportContext>();
+								if (!LevelViewportContext)
+								{
+									return ECheckBoxState::Undetermined;
+								}
 
-							return ECheckBoxState::Undetermined;
-						});
+								if (const TSharedPtr<::SLevelViewport> LevelViewport =
+										LevelViewportContext->LevelViewport.Pin())
+								{
+									return LevelViewport->IsMaximized() ? ECheckBoxState::Checked
+																		: ECheckBoxState::Unchecked;
+								}
 
-					SettingsSection
-						.AddMenuEntry("MaximizeRestore", LOCTEXT("MaximizeRestoreLabel", "Maximize Viewport"),
-							LOCTEXT("MaximizeRestoreTooltip", "Maximizes or restores this viewport"),
-							FSlateIcon(FAppStyle::GetAppStyleSetName(), "EditorViewportToolBar.Maximize.Normal"),
-							MaximizeRestoreAction, EUserInterfaceActionType::ToggleButton)
-						.SetShowInToolbarTopLevel(true);
-				}
+								return ECheckBoxState::Undetermined;
+							}
+						);
+
+						SettingsSection
+							.AddMenuEntry(
+								"MaximizeRestore",
+								LOCTEXT("MaximizeRestoreLabel", "Maximize Viewport"),
+								LOCTEXT("MaximizeRestoreTooltip", "Maximizes or restores this viewport"),
+								FSlateIcon(FAppStyle::GetAppStyleSetName(), "EditorViewportToolBar.Maximize.Normal"),
+								MaximizeRestoreAction,
+								EUserInterfaceActionType::ToggleButton
+							)
+							.SetShowInToolbarTopLevel(true);
+					}
 
-				SettingsSection.AddSeparator("AdvancedSeparator");
+					SettingsSection.AddSeparator("AdvancedSeparator");
 
-				{
-					const FLevelViewportCommands& LevelViewportActions = FLevelViewportCommands::Get();
-					SettingsSection.AddMenuEntry(LevelViewportActions.AdvancedSettings);
+					{
+						const FLevelViewportCommands& LevelViewportActions = FLevelViewportCommands::Get();
+						SettingsSection.AddMenuEntry(LevelViewportActions.AdvancedSettings);
+					}
 				}
 			}
-		}));
+		)
+	);
 }
 
 } // namespace UE::LevelEditor
