Use the new ToolMenu alignment support to align some parts of the new viewport toolbar to the right

diff --git a/Engine/Source/Editor/LevelEditor/Private/LevelEditorToolBar.cpp b/Engine/Source/Editor/LevelEditor/Private/LevelEditorToolBar.cpp
--- a/Engine/Source/Editor/LevelEditor/Private/LevelEditorToolBar.cpp
+++ b/Engine/Source/Editor/LevelEditor/Private/LevelEditorToolBar.cpp
@@ -1749,8 +1749,8 @@ void FLevelEditorToolBar::RegisterQuickSettingsMenu()
 			LOCTEXT( "ScalabilitySubMenu_ToolTip", "Open the engine scalability settings" ),
 			FNewToolMenuDelegate::CreateStatic( &MakeScalabilityMenu ) );
 
-		UE::LevelEditor::AddMaterialQualityLevelSubmenu(Section);
-		UE::LevelEditor::AddFeatureLevelPreviewSubmenu(Section);
+		Section.AddEntry(UE::LevelEditor::CreateMaterialQualityLevelSubmenu());
+		Section.AddEntry(UE::LevelEditor::CreateFeatureLevelPreviewSubmenu());
 	}
 
 	{
diff --git a/Engine/Source/Editor/LevelEditor/Private/SLevelViewport.cpp b/Engine/Source/Editor/LevelEditor/Private/SLevelViewport.cpp
--- a/Engine/Source/Editor/LevelEditor/Private/SLevelViewport.cpp
+++ b/Engine/Source/Editor/LevelEditor/Private/SLevelViewport.cpp
@@ -102,6 +102,7 @@
 #include "LevelViewportLayout.h"
 #include "EditorViewportTabContent.h"
 #include "EditorViewportCommands.h"
+#include "SkeletalRenderPublic.h"
 #include "ViewportToolbar/LevelEditorViewportToolbarSections.h"
 #include "ViewportToolbar/LevelViewportContext.h"
 
@@ -1928,22 +1929,39 @@ TSharedPtr<SWidget> SLevelViewport::MakeViewportToolbar()
 
 			{
 				FToolMenuSection& LeftSection = ViewportToolbarMenu->FindOrAddSection("Left");
-				UE::LevelEditor::AddViewportToolbarTransformsSection(LeftSection);
+
+				FToolMenuEntry TransformsSubmenu = UE::LevelEditor::CreateViewportToolbarTransformsSection();
+				TransformsSubmenu.InsertPosition.Position = EToolMenuInsertType::First;
+				LeftSection.AddEntry(TransformsSubmenu);
 			}
 
 			{
+				// Add the submenus of this section as EToolMenuInsertType::Last to sort them after any
+				// default-positioned submenus external code might add.
 				FToolMenuSection& RightSection = ViewportToolbarMenu->FindOrAddSection("Right");
 				RightSection.Alignment = EToolMenuSectionAlign::Last;
 
 				{
 					// Stay backward-compatible with the old viewport toolbar.
 					UToolMenus::Get()->RegisterMenu(
-						"LevelEditor.ViewportToolbar.ViewModes", "LevelEditor.LevelViewportToolbar.View");
-					UE::LevelEditor::AddViewportToolbarViewModesSubmenu(RightSection);
+						"LevelEditor.ViewportToolbar.ViewModes", "LevelEditor.LevelViewportToolbar.View"
+					);
+					FToolMenuEntry LitSubmenu = UE::LevelEditor::CreateViewportToolbarViewModesSubmenu();
+					LitSubmenu.InsertPosition.Position = EToolMenuInsertType::Last;
+					RightSection.AddEntry(LitSubmenu);
+				}
+
+				{
+					FToolMenuEntry PerfSubmenu = UE::LevelEditor::CreateViewportToolbarPerformanceAndScalabilitySubmenu();
+					PerfSubmenu.InsertPosition.Position = EToolMenuInsertType::Last;
+					RightSection.AddEntry(PerfSubmenu);
 				}
 
-				UE::LevelEditor::AddViewportToolbarPerformanceAndScalabilitySubmenu(RightSection);
-				UE::LevelEditor::AddLevelEditorViewportToolbarSettingsSubmenu(RightSection);
+				{
+					FToolMenuEntry SettingsSubmenu = UE::LevelEditor::CreateLevelEditorViewportToolbarSettingsSubmenu();
+					SettingsSubmenu.InsertPosition.Position = EToolMenuInsertType::Last;
+					RightSection.AddEntry(SettingsSubmenu);
+				}
 			}
 		}
 	}
diff --git a/Engine/Source/Editor/LevelEditor/Private/ViewportToolbar/LevelEditorViewportToolbarSections.cpp b/Engine/Source/Editor/LevelEditor/Private/ViewportToolbar/LevelEditorViewportToolbarSections.cpp
--- a/Engine/Source/Editor/LevelEditor/Private/ViewportToolbar/LevelEditorViewportToolbarSections.cpp
+++ b/Engine/Source/Editor/LevelEditor/Private/ViewportToolbar/LevelEditorViewportToolbarSections.cpp
@@ -77,9 +77,9 @@ bool ShowViewportRealtimeWarning(FLevelEditorViewportClient& ViewportClient)
 }
 
 // TODO: Move this outside the level editor and make it publicly available to anyone building a viewport toolbar.
-void AddViewportToolbarTransformsSection(FToolMenuSection& InSection)
+FToolMenuEntry CreateViewportToolbarTransformsSection()
 {
-	InSection.AddSubMenu(
+	return FToolMenuEntry::InitSubMenu(
 		"Transforms",
 		LOCTEXT("TransformsSubmenuLabel", "Transforms"),
 		LOCTEXT("TransformsSubmenuTooltip", "Viewport-related transforms tools"),
@@ -495,10 +495,10 @@ void PopulateViewModesMenu(UToolMenu* InMenu, TSharedRef<::SLevelViewport> InVie
 	}
 }
 
-void AddViewportToolbarViewModesSubmenu(FToolMenuSection& InSection)
+FToolMenuEntry CreateViewportToolbarViewModesSubmenu()
 {
 	// This has to be a dynamic entry for the ViewModes submenu's label to be able to access the context.
-	InSection.AddDynamicEntry(
+	return FToolMenuEntry::InitDynamicEntry(
 		"DynamicViewModes",
 		FNewToolMenuSectionDelegate::CreateLambda(
 			[](FToolMenuSection& InDynamicSection) -> void
@@ -543,9 +543,9 @@ void AddViewportToolbarViewModesSubmenu(FToolMenuSection& InSection)
 	);
 }
 
-void AddFeatureLevelPreviewSubmenu(FToolMenuSection& Section)
+FToolMenuEntry CreateFeatureLevelPreviewSubmenu()
 {
-	Section.AddSubMenu(
+	return FToolMenuEntry::InitSubMenu(
 		"FeatureLevelPreview",
 		NSLOCTEXT("LevelToolBarViewMenu", "PreviewPlatformSubMenu", "Preview Platform"),
 		NSLOCTEXT("LevelToolBarViewMenu", "PreviewPlatformSubMenu_ToolTip", "Sets the preview platform used by the main editor"),
@@ -564,9 +564,9 @@ void AddFeatureLevelPreviewSubmenu(FToolMenuSection& Section)
 	);
 }
 
-void AddMaterialQualityLevelSubmenu(FToolMenuSection& Section)
+FToolMenuEntry CreateMaterialQualityLevelSubmenu()
 {
-	Section.AddSubMenu(
+	return FToolMenuEntry::InitSubMenu(
 		"MaterialQualityLevel",
 		NSLOCTEXT("LevelToolBarViewMenu", "MaterialQualityLevelSubMenu", "Material Quality Level"),
 		NSLOCTEXT(
@@ -591,9 +591,9 @@ void AddMaterialQualityLevelSubmenu(FToolMenuSection& Section)
 	);
 }
 
-void AddViewportToolbarPerformanceAndScalabilitySubmenu(FToolMenuSection& InSection)
+FToolMenuEntry CreateViewportToolbarPerformanceAndScalabilitySubmenu()
 {
-	InSection.AddSubMenu(
+	return FToolMenuEntry::InitSubMenu(
 		"PerformanceAndScalability",
 		LOCTEXT("PerformanceAndScalabilityLabel", "Performance & Scalability"),
 		LOCTEXT("PerformanceAndScalabilityTooltip", "Performance and scalability tools tied to this viewport."),
@@ -702,7 +702,7 @@ void AddViewportToolbarPerformanceAndScalabilitySubmenu(FToolMenuSection& InSect
 						LOCTEXT("PerformanceAndScalabilitySectionLabel", "Performance & Scalability")
 					);
 
-					AddFeatureLevelPreviewSubmenu(PerformanceAndScalabilitySection);
+					PerformanceAndScalabilitySection.AddEntry(CreateFeatureLevelPreviewSubmenu());
 
 					PerformanceAndScalabilitySection.AddSeparator("PerformanceAndScalabilitySettings");
 
@@ -721,7 +721,7 @@ void AddViewportToolbarPerformanceAndScalabilitySubmenu(FToolMenuSection& InSect
 						)
 					);
 
-					AddMaterialQualityLevelSubmenu(PerformanceAndScalabilitySection);
+					PerformanceAndScalabilitySection.AddEntry(CreateMaterialQualityLevelSubmenu());
 
 					// FEditorViewportClient& ViewportClient = Viewport.Pin()->GetLevelViewportClient();
 					PerformanceAndScalabilitySection.AddSubMenu(
@@ -920,9 +920,9 @@ TSharedRef<SWidget> BuildVolumeControlCustomWidget()
 	// clang-format on
 }
 
-void AddLevelEditorViewportToolbarSettingsSubmenu(FToolMenuSection& InSection)
+FToolMenuEntry CreateLevelEditorViewportToolbarSettingsSubmenu()
 {
-	InSection.AddSubMenu(
+	return FToolMenuEntry::InitSubMenu(
 		"Settings",
 		LOCTEXT("SettingsSubmenuLabel", "Settings"),
 		LOCTEXT("SettingsSubmenuTooltip", "Viewport-related settings"),
diff --git a/Engine/Source/Editor/LevelEditor/Private/ViewportToolbar/LevelEditorViewportToolbarSections.h b/Engine/Source/Editor/LevelEditor/Private/ViewportToolbar/LevelEditorViewportToolbarSections.h
--- a/Engine/Source/Editor/LevelEditor/Private/ViewportToolbar/LevelEditorViewportToolbarSections.h
+++ b/Engine/Source/Editor/LevelEditor/Private/ViewportToolbar/LevelEditorViewportToolbarSections.h
@@ -3,6 +3,7 @@
 #pragma once
 
 #include "Templates/SharedPointerFwd.h"
+#include "ToolMenuEntry.h"
 
 class FExtender;
 class FLevelEditorViewportClient;
@@ -16,18 +17,18 @@ namespace UE::LevelEditor
 
 bool ShowViewportRealtimeWarning(FLevelEditorViewportClient& ViewportClient);
 
-void AddViewportToolbarTransformsSection(FToolMenuSection& InSection);
+FToolMenuEntry CreateViewportToolbarTransformsSection();
 
 TSharedPtr<FExtender> GetViewModesLegacyExtenders();
 void PopulateViewModesMenu(UToolMenu* InMenu, TSharedRef<SLevelViewport> InViewport);
-void AddViewportToolbarViewModesSubmenu(FToolMenuSection& InSection);
+FToolMenuEntry CreateViewportToolbarViewModesSubmenu();
 
-void AddFeatureLevelPreviewSubmenu(FToolMenuSection& Section);
-void AddMaterialQualityLevelSubmenu(FToolMenuSection& Section);
-void AddViewportToolbarPerformanceAndScalabilitySubmenu(FToolMenuSection& InSection);
+FToolMenuEntry CreateFeatureLevelPreviewSubmenu();
+FToolMenuEntry CreateMaterialQualityLevelSubmenu();
+FToolMenuEntry CreateViewportToolbarPerformanceAndScalabilitySubmenu();
 
 void GenerateViewportLayoutsMenu(UToolMenu* InMenu, TSharedPtr<::SLevelViewport> InViewport);
 TSharedRef<SWidget> BuildVolumeControlCustomWidget();
-void AddLevelEditorViewportToolbarSettingsSubmenu(FToolMenuSection& InSection);
+FToolMenuEntry CreateLevelEditorViewportToolbarSettingsSubmenu();
 
 } // namespace UE::LevelEditor
