Update the Transforms section of the new level editor viewport toolbar

diff --git a/Engine/Source/Editor/LevelEditor/Private/ViewportToolbar/LevelEditorViewportToolbarSections.cpp b/Engine/Source/Editor/LevelEditor/Private/ViewportToolbar/LevelEditorViewportToolbarSections.cpp
--- a/Engine/Source/Editor/LevelEditor/Private/ViewportToolbar/LevelEditorViewportToolbarSections.cpp
+++ b/Engine/Source/Editor/LevelEditor/Private/ViewportToolbar/LevelEditorViewportToolbarSections.cpp
@@ -16,6 +16,7 @@
 #include "ToolMenu.h"
 #include "ToolMenus.h"
 #include "ViewportToolbar/UnrealEdViewportToolbar.h"
+#include "Widgets/Input/SSpinBox.h"
 #include "Widgets/Input/SVolumeControl.h"
 #include "Widgets/SBoxPanel.h"
 
@@ -86,27 +87,96 @@ FToolMenuEntry CreateViewportToolbarTransformsSection()
 		FNewToolMenuDelegate::CreateLambda(
 			[](UToolMenu* Submenu) -> void
 			{
-				FToolMenuSection& Section = Submenu->FindOrAddSection(NAME_None);
-
-				FToolMenuEntry SelectMode = FToolMenuEntry::InitMenuEntry(FEditorViewportCommands::Get().SelectMode);
-				SelectMode.UserInterfaceActionType = EUserInterfaceActionType::RadioButton;
-				SelectMode.SetShowInToolbarTopLevel(true);
-				Section.AddEntry(SelectMode);
-
-				FToolMenuEntry TranslateMode = FToolMenuEntry::InitMenuEntry(FEditorViewportCommands::Get().TranslateMode);
-				TranslateMode.UserInterfaceActionType = EUserInterfaceActionType::RadioButton;
-				TranslateMode.SetShowInToolbarTopLevel(true);
-				Section.AddEntry(TranslateMode);
-
-				FToolMenuEntry RotateMode = FToolMenuEntry::InitMenuEntry(FEditorViewportCommands::Get().RotateMode);
-				RotateMode.UserInterfaceActionType = EUserInterfaceActionType::RadioButton;
-				RotateMode.SetShowInToolbarTopLevel(true);
-				Section.AddEntry(RotateMode);
-
-				FToolMenuEntry ScaleMode = FToolMenuEntry::InitMenuEntry(FEditorViewportCommands::Get().ScaleMode);
-				ScaleMode.UserInterfaceActionType = EUserInterfaceActionType::RadioButton;
-				ScaleMode.SetShowInToolbarTopLevel(true);
-				Section.AddEntry(ScaleMode);
+				{
+					FToolMenuSection& TransformToolsSection =
+						Submenu->FindOrAddSection("TransformTools", LOCTEXT("TransformToolsLabel", "Transform Tools"));
+
+					FToolMenuEntry SelectMode = FToolMenuEntry::InitMenuEntry(FEditorViewportCommands::Get().SelectMode);
+					SelectMode.UserInterfaceActionType = EUserInterfaceActionType::RadioButton;
+					SelectMode.SetShowInToolbarTopLevel(true);
+					TransformToolsSection.AddEntry(SelectMode);
+
+					FToolMenuEntry TranslateMode =
+						FToolMenuEntry::InitMenuEntry(FEditorViewportCommands::Get().TranslateMode);
+					TranslateMode.UserInterfaceActionType = EUserInterfaceActionType::RadioButton;
+					TranslateMode.SetShowInToolbarTopLevel(true);
+					TransformToolsSection.AddEntry(TranslateMode);
+
+					FToolMenuEntry RotateMode = FToolMenuEntry::InitMenuEntry(FEditorViewportCommands::Get().RotateMode);
+					RotateMode.UserInterfaceActionType = EUserInterfaceActionType::RadioButton;
+					RotateMode.SetShowInToolbarTopLevel(true);
+					TransformToolsSection.AddEntry(RotateMode);
+
+					FToolMenuEntry ScaleMode = FToolMenuEntry::InitMenuEntry(FEditorViewportCommands::Get().ScaleMode);
+					ScaleMode.UserInterfaceActionType = EUserInterfaceActionType::RadioButton;
+					ScaleMode.SetShowInToolbarTopLevel(true);
+					TransformToolsSection.AddEntry(ScaleMode);
+				}
+
+				{
+					FToolMenuSection& SpacesSection = Submenu->FindOrAddSection("Spaces", LOCTEXT("SpacesLabel", "Spaces"));
+
+					FToolMenuEntry WorldSpace =
+						FToolMenuEntry::InitMenuEntry(FEditorViewportCommands::Get().RelativeCoordinateSystem_World);
+					WorldSpace.UserInterfaceActionType = EUserInterfaceActionType::RadioButton;
+					WorldSpace.SetShowInToolbarTopLevel(true);
+					SpacesSection.AddEntry(WorldSpace);
+
+					FToolMenuEntry LocalSpace =
+						FToolMenuEntry::InitMenuEntry(FEditorViewportCommands::Get().RelativeCoordinateSystem_Local);
+					LocalSpace.UserInterfaceActionType = EUserInterfaceActionType::RadioButton;
+					LocalSpace.SetShowInToolbarTopLevel(true);
+					SpacesSection.AddEntry(LocalSpace);
+				}
+
+				{
+					FToolMenuSection& GizmoSection = Submenu->FindOrAddSection("Gizmo", LOCTEXT("GizmoLabel", "Gizmo"));
+
+					GizmoSection.AddMenuEntry(
+						FLevelEditorCommands::Get().ShowTransformWidget,
+						LOCTEXT("ShowTransformGizmoLabel", "Show Transform Gizmo")
+					);
+					
+					TSharedRef<SWidget> GizmoScaleWidget =
+						// clang-format off
+						SNew(SHorizontalBox)
+						+ SHorizontalBox::Slot()
+						.FillWidth(0.9f)
+						[
+							SNew(SSpinBox<int32>)
+								.MinValue(-10)
+								.MaxValue(150)
+								.ToolTipText_Lambda(
+									[]() -> FText
+									{
+										return FText::AsNumber(
+											GetDefault<ULevelEditorViewportSettings>()->TransformWidgetSizeAdjustment
+										);
+									}
+								)
+								.Value_Lambda(
+									[]() -> float
+									{
+										return GetDefault<ULevelEditorViewportSettings>()->TransformWidgetSizeAdjustment;
+									}
+								)
+								.OnValueChanged_Lambda(
+									[](float InValue)
+									{
+										ULevelEditorViewportSettings* ViewportSettings =
+											GetMutableDefault<ULevelEditorViewportSettings>();
+										ViewportSettings->TransformWidgetSizeAdjustment = InValue;
+										ViewportSettings->PostEditChange();
+									}
+								)
+						]
+						+ SHorizontalBox::Slot()
+						.FillWidth(0.1f);
+					// clang-format on
+					GizmoSection.AddEntry(FToolMenuEntry::InitWidget(
+						"GizmoScale", GizmoScaleWidget, LOCTEXT("GizmoScaleLabel", "Gizmo Scale")
+					));
+				}
 			}
 		)
 	);
