Scene Outliner: Grab the selection set from the menu context instead of caching it in the lambda on level editor creation.

The latter would cause things to break if you re-created the level editor (e.g layout load) and the cached selection set would still use the one belonging to the old instance.

diff --git a/Engine/Source/Editor/LevelEditor/Private/SLevelEditor.cpp b/Engine/Source/Editor/LevelEditor/Private/SLevelEditor.cpp
--- a/Engine/Source/Editor/LevelEditor/Private/SLevelEditor.cpp
+++ b/Engine/Source/Editor/LevelEditor/Private/SLevelEditor.cpp
@@ -824,9 +824,15 @@ TSharedRef<ISceneOutliner> SLevelEditor::CreateSceneOutliner(FName TabIdentifier
 		if (!ToolMenus->IsMenuRegistered(MenuName))
 		{
 			UToolMenu* Menu = ToolMenus->RegisterMenu(MenuName, "SceneOutliner.DefaultContextMenuBase");
-			FToolMenuSection& Section = Menu->AddDynamicSection("LevelEditorContextMenu", FNewToolMenuDelegate::CreateLambda([SelectionSet = TWeakObjectPtr<const UTypedElementSelectionSet>(GetElementSelectionSet())](UToolMenu* InMenu)
+			FToolMenuSection& Section = Menu->AddDynamicSection("LevelEditorContextMenu", FNewToolMenuDelegate::CreateLambda([](UToolMenu* InMenu)
 			{
-				FName LevelContextMenuName = FLevelEditorContextMenu::GetContextMenuName(ELevelEditorMenuContext::SceneOutliner, SelectionSet.Get());
+				ULevelEditorContextMenuContext* LevelEditorContext = InMenu->FindContext<ULevelEditorContextMenuContext>();
+				if (!LevelEditorContext)
+				{
+					return;
+				}
+
+				FName LevelContextMenuName = FLevelEditorContextMenu::GetContextMenuName(ELevelEditorMenuContext::SceneOutliner, LevelEditorContext->CurrentSelection);
 				if (LevelContextMenuName != NAME_None)
 				{
 					// Extend the menu even if no actors selected, as Edit menu should always exist for scene outliner
