Delegate-based top-level flagging of ToolMenu toolbar submenu entries

diff --git a/Engine/Source/Developer/ToolMenus/Private/ToolMenuEntry.cpp b/Engine/Source/Developer/ToolMenus/Private/ToolMenuEntry.cpp
--- a/Engine/Source/Developer/ToolMenus/Private/ToolMenuEntry.cpp
+++ b/Engine/Source/Developer/ToolMenus/Private/ToolMenuEntry.cpp
@@ -22,7 +22,7 @@ FToolMenuEntry::FToolMenuEntry()
 	, StyleNameOverride(NAME_None)
 	, bAddedDuringRegister(false)
 	, bCommandIsKeybindOnly(false)
-	, bShowInToolbarTopLevel(false)
+	, ShowInToolbarTopLevel(false)
 {
 }
 
@@ -36,7 +36,7 @@ FToolMenuEntry::FToolMenuEntry(const FToolMenuOwner InOwner, const FName InName,
 	, StyleNameOverride(NAME_None)
 	, bAddedDuringRegister(false)
 	, bCommandIsKeybindOnly(false)
-	, bShowInToolbarTopLevel(false)
+	, ShowInToolbarTopLevel(false)
 {
 }
 
@@ -306,9 +306,9 @@ bool FToolMenuEntry::TryExecuteToolUIAction(const FToolMenuContext& InContext)
 	return bCanExecute;
 }
 
-void FToolMenuEntry::SetShowInToolbarTopLevel(bool InTopLevel)
+void FToolMenuEntry::SetShowInToolbarTopLevel(TAttribute<bool> InTopLevel)
 {
-	bShowInToolbarTopLevel = InTopLevel;
+	ShowInToolbarTopLevel = InTopLevel;
 }
 
 bool FToolMenuEntry::IsScriptObjectDynamicConstruct() const
diff --git a/Engine/Source/Developer/ToolMenus/Private/ToolMenus.cpp b/Engine/Source/Developer/ToolMenus/Private/ToolMenus.cpp
--- a/Engine/Source/Developer/ToolMenus/Private/ToolMenus.cpp
+++ b/Engine/Source/Developer/ToolMenus/Private/ToolMenus.cpp
@@ -1469,16 +1469,22 @@ void UToolMenus::PopulateToolBarBuilderWithTopLevelChildren(
 			break;
 		}
 
-		if (Block->bShowInToolbarTopLevel)
+		const bool bIsBound = Block->ShowInToolbarTopLevel.IsBound();
+		const bool bIsSetToValue = !bIsBound && Block->ShowInToolbarTopLevel.IsSet();
+		const bool bIsSetToTrueValue = bIsSetToValue && Block->ShowInToolbarTopLevel.Get();
+		if (bIsBound || bIsSetToTrueValue)
 		{
-			PopulateToolBarBuilderWithEntry(ToolBarBuilder, SubMenu, *Block);
+			const bool bRaiseToTopLevel = true;
+			PopulateToolBarBuilderWithEntry(ToolBarBuilder, SubMenu, *Block, bRaiseToTopLevel);
 		}
 
 		ExtractChildBlocksFromSubMenu(SubMenu, *Block, SubMenuBlockPairs);
 	}
 }
 
-void UToolMenus::PopulateToolBarBuilderWithEntry(FToolBarBuilder& ToolBarBuilder, UToolMenu* MenuData, FToolMenuEntry& Block)
+void UToolMenus::PopulateToolBarBuilderWithEntry(
+	FToolBarBuilder& ToolBarBuilder, UToolMenu* MenuData, FToolMenuEntry& Block, bool bIsRaisingToTopLevel
+)
 {
 	if (Block.ToolBarData.ConstructLegacy.IsBound())
 	{
@@ -1488,6 +1494,28 @@ void UToolMenus::PopulateToolBarBuilderWithEntry(FToolBarBuilder& ToolBarBuilder
 
 	ToolBarBuilder.BeginStyleOverride(Block.StyleNameOverride);
 
+	FUIAction UIAction = UToolMenus::ConvertUIAction(Block, MenuData->Context);
+	if (bIsRaisingToTopLevel && Block.ShowInToolbarTopLevel.IsBound())
+	{
+		// Patch the IsVisible delegate with the top-level status of the menu entry.
+		UIAction.IsActionVisibleDelegate = FIsActionButtonVisible::CreateLambda(
+			[OriginalIsVisible = UIAction.IsActionVisibleDelegate, ShowInToolbarTopLevel = Block.ShowInToolbarTopLevel]() -> bool
+			{
+				if (ShowInToolbarTopLevel.IsSet() && !ShowInToolbarTopLevel.Get())
+				{
+					return false;
+				}
+
+				if (OriginalIsVisible.IsBound())
+				{
+					return OriginalIsVisible.Execute();
+				}
+
+				return true;
+			}
+		);
+	}
+
 	if (Block.Type == EMultiBlockType::ToolBarButton || (Block.Type == EMultiBlockType::MenuEntry && !Block.IsSubMenu()))
 	{
 		if (Block.Command.IsValid() && !Block.IsCommandKeybindOnly())
@@ -1518,7 +1546,6 @@ void UToolMenus::PopulateToolBarBuilderWithEntry(FToolBarBuilder& ToolBarBuilder
 		{
 			UToolMenuEntryScript* ScriptObject = Block.ScriptObject;
 			TAttribute<FSlateIcon> Icon = ScriptObject->CreateIconAttribute(MenuData->Context);
-			FUIAction UIAction = ConvertUIAction(Block, MenuData->Context);
 			ToolBarBuilder.AddToolBarButton(UIAction, ScriptObject->Data.Name,
 				ScriptObject->CreateLabelAttribute(MenuData->Context),
 				ScriptObject->CreateToolTipAttribute(MenuData->Context), Icon, Block.UserInterfaceActionType,
@@ -1526,7 +1553,6 @@ void UToolMenus::PopulateToolBarBuilderWithEntry(FToolBarBuilder& ToolBarBuilder
 		}
 		else
 		{
-			FUIAction UIAction = ConvertUIAction(Block, MenuData->Context);
 			ToolBarBuilder.AddToolBarButton(UIAction, Block.Name, Block.Label, Block.ToolTip, Block.Icon,
 				Block.UserInterfaceActionType, Block.TutorialHighlightName);
 		}
@@ -1546,7 +1572,6 @@ void UToolMenus::PopulateToolBarBuilderWithEntry(FToolBarBuilder& ToolBarBuilder
 			Block.ToolBarData.ComboButtonContextMenuGenerator, MenuData->Context);
 		if (OnGetContent.IsBound())
 		{
-			FUIAction UIAction = ConvertUIAction(Block, MenuData->Context);
 			ToolBarBuilder.AddComboButton(UIAction, OnGetContent, Block.Label, Block.ToolTip, Block.Icon,
 				Block.ToolBarData.bSimpleComboBox, Block.TutorialHighlightName);
 		}
@@ -1554,7 +1579,6 @@ void UToolMenus::PopulateToolBarBuilderWithEntry(FToolBarBuilder& ToolBarBuilder
 		{
 			FOnGetContent Delegate = FOnGetContent::CreateUObject(
 				this, &UToolMenus::GenerateToolbarComboButtonMenu, TWeakObjectPtr<UToolMenu>(MenuData), Block.Name);
-			FUIAction UIAction = ConvertUIAction(Block, MenuData->Context);
 			ToolBarBuilder.AddComboButton(UIAction, Delegate, Block.Label, Block.ToolTip, Block.Icon,
 				Block.ToolBarData.bSimpleComboBox, Block.TutorialHighlightName);
 
diff --git a/Engine/Source/Developer/ToolMenus/Public/ToolMenuEntry.h b/Engine/Source/Developer/ToolMenus/Public/ToolMenuEntry.h
--- a/Engine/Source/Developer/ToolMenus/Public/ToolMenuEntry.h
+++ b/Engine/Source/Developer/ToolMenus/Public/ToolMenuEntry.h
@@ -159,18 +159,15 @@ struct TOOLMENUS_API FToolMenuEntry
 	/**
 	 * Show this menu entry in the top-level toolbar section of a toolbar.
 	 *
-	 * A toolbar that has a section menu only shows entries flagged as top-level in the toolbar section itself while the
-	 * others can only be accessed through the section submenu. Use this method to flag a menu entry to show up in that
-	 * top-level toolbar section.
+	 * Entries of a toolbar submenu can be raised to the top-level of the toolbar. Such top-level entires appear in the
+	 * toolbar to the right of the submenu they belong to.
 	 *
-	 * This flag only affects entries of ToolMenus rendered as toolbars that also have a section menu. For ToolMenus
-	 * rendered as menus or toolbar ToolMenus that lack a section menu, the entry is always placed in the top-level
-	 * section itself.
+	 * This flag only effects entries within submenus of toolbar-type ToolMenus.
 	 *
-	 * @param InTopLevel Pass true to set this menu entry to appear in the top-level toolbar section for toolbar
-	 * sections with a section menu. Pass false to only show the entry in the section submenu.
+	 * @param InTopLevel True shows the entry in the top-level next to its submenu, false (default) only displays it in
+	 * the submenu itself. Pass a delegate to drive the top-level state from code.
 	 */
-	void SetShowInToolbarTopLevel(bool InTopLevel);
+	void SetShowInToolbarTopLevel(TAttribute<bool> InTopLevel);
 
 private:
 
@@ -249,5 +246,5 @@ private:
 	UPROPERTY()
 	bool bCommandIsKeybindOnly;
 
-	bool bShowInToolbarTopLevel;
+	TAttribute<bool> ShowInToolbarTopLevel;
 };
diff --git a/Engine/Source/Developer/ToolMenus/Public/ToolMenus.h b/Engine/Source/Developer/ToolMenus/Public/ToolMenus.h
--- a/Engine/Source/Developer/ToolMenus/Public/ToolMenus.h
+++ b/Engine/Source/Developer/ToolMenus/Public/ToolMenus.h
@@ -390,7 +390,9 @@ private:
 		UToolMenu* ParentMenu, FToolMenuEntry& InBlock, TArray<TPair<UToolMenu*, FToolMenuEntry*>>& SubMenuBlockPairs);
 	void PopulateToolBarBuilderWithTopLevelChildren(
 		FToolBarBuilder& ToolBarBuilder, UToolMenu* ParentMenu, FToolMenuEntry& InBlock);
-	void PopulateToolBarBuilderWithEntry(FToolBarBuilder& ToolBarBuilder, UToolMenu* MenuData, FToolMenuEntry& Entry);
+	void PopulateToolBarBuilderWithEntry(
+		FToolBarBuilder& ToolBarBuilder, UToolMenu* MenuData, FToolMenuEntry& Entry, bool bIsRaisingToTopLevel = false
+	);
 	void PopulateToolBarBuilder(FToolBarBuilder& ToolBarBuilder, UToolMenu* MenuData);
 
 	TSharedRef<SWidget> GenerateToolbarComboButtonMenu(TWeakObjectPtr<UToolMenu> InParent, const FName InBlockName);
diff --git a/Engine/Source/Editor/LevelEditor/Private/ViewportToolbar/LevelEditorViewportToolbarSections.cpp b/Engine/Source/Editor/LevelEditor/Private/ViewportToolbar/LevelEditorViewportToolbarSections.cpp
--- a/Engine/Source/Editor/LevelEditor/Private/ViewportToolbar/LevelEditorViewportToolbarSections.cpp
+++ b/Engine/Source/Editor/LevelEditor/Private/ViewportToolbar/LevelEditorViewportToolbarSections.cpp
@@ -36,37 +36,6 @@ void OnLandscapeLODChanged(FLevelEditorViewportClient& ViewportClient, int32 New
 	ViewportClient.Invalidate();
 }
 
-TOptional<bool> UpdateAndGetRealtimeWarningFromContext(const FToolMenuContext& Context)
-{
-	if (ULevelViewportContext* const LevelViewportContext = Context.FindContext<ULevelViewportContext>())
-	{
-		if (const TSharedPtr<::SLevelViewport> LevelViewport = LevelViewportContext->LevelViewport.Pin())
-		{
-			const bool bShowWarning = ShowViewportRealtimeWarning(LevelViewport->GetLevelViewportClient());
-
-			LevelViewportContext->CachedShouldShowRealtimeOffWarning = bShowWarning;
-
-			return LevelViewportContext->CachedShouldShowRealtimeOffWarning;
-		}
-	}
-
-	return TOptional<bool>();
-}
-
-TOptional<bool> IsDirtyRealtimeWarningFromContext(const FToolMenuContext& Context)
-{
-	if (ULevelViewportContext* const LevelViewportContext = Context.FindContext<ULevelViewportContext>())
-	{
-		if (const TSharedPtr<::SLevelViewport> LevelViewport = LevelViewportContext->LevelViewport.Pin())
-		{
-			const bool bShowWarning = ShowViewportRealtimeWarning(LevelViewport->GetLevelViewportClient());
-			return bShowWarning != LevelViewportContext->CachedShouldShowRealtimeOffWarning;
-		}
-	}
-
-	return TOptional<bool>();
-}
-
 } // namespace UE::LevelEditor::Private
 
 namespace UE::LevelEditor
@@ -694,7 +663,6 @@ FToolMenuEntry CreateViewportToolbarPerformanceAndScalabilitySubmenu()
 												LevelViewportContext->LevelViewport.Pin())
 										{
 											LevelViewport->OnToggleRealtime();
-											UToolMenus::Get()->RefreshAllWidgets();
 										}
 									}
 								);
@@ -709,15 +677,6 @@ FToolMenuEntry CreateViewportToolbarPerformanceAndScalabilitySubmenu()
 											return ECheckBoxState::Undetermined;
 										}
 
-										// Check if the realtime warn state is outdated and if so refresh widgets to
-										// update our top-level status.
-										if (const TOptional<bool> IsDirty =
-												Private::IsDirtyRealtimeWarningFromContext(Context);
-											IsDirty.IsSet() && IsDirty.GetValue())
-										{
-											UToolMenus::Get()->RefreshAllWidgets();
-										}
-
 										if (const TSharedPtr<::SLevelViewport> LevelViewport =
 												LevelViewportContext->LevelViewport.Pin())
 										{
@@ -729,27 +688,40 @@ FToolMenuEntry CreateViewportToolbarPerformanceAndScalabilitySubmenu()
 									}
 								);
 
-								bool bDisplayTopLevel = false;
-								if (const TOptional<bool> ShouldWarn =
-										Private::UpdateAndGetRealtimeWarningFromContext(InnerSection.Context);
-									ShouldWarn.IsSet())
-								{
-									bDisplayTopLevel = ShouldWarn.GetValue();
-								}
-								else
+								TAttribute<FText> Tooltip;
 								{
-									// If we couldn't get the warn state, pretend we don't have to warn.
-									bDisplayTopLevel = false;
-								}
-
-								const FText Tooltip =
-									bDisplayTopLevel ? LOCTEXT(
+									const FText NonRealtimeTooltip = LOCTEXT(
 										"ToggleRealtimeTooltip_WarnRealtimeOff",
 										"This viewport is not updating in realtime.  Click to turn on realtime mode."
-									)
-													 : LOCTEXT(
-														 "ToggleRealtimeTooltip", "Toggle realtime rendering of the viewport"
-													 );
+									);
+									const FText RealtimeTooltip =
+										LOCTEXT("ToggleRealtimeTooltip", "Toggle realtime rendering of the viewport");
+
+									// If we can find a context with a viewport, use that to adjust the tooltip
+									// based on the viewport's realtime status.
+									if (ULevelViewportContext* const LevelViewportContext =
+											InnerSection.FindContext<ULevelViewportContext>())
+									{
+										Tooltip = TAttribute<FText>::CreateLambda(
+											[WeakViewport = LevelViewportContext->LevelViewport,
+											 NonRealtimeTooltip,
+											 RealtimeTooltip]() -> FText
+											{
+												bool bDisplayTopLevel = false;
+												if (const TSharedPtr<::SLevelViewport> LevelViewport = WeakViewport.Pin())
+												{
+													bDisplayTopLevel = !LevelViewport->IsRealtime();
+												}
+
+												return bDisplayTopLevel ? NonRealtimeTooltip : RealtimeTooltip;
+											}
+										);
+									}
+									else
+									{
+										Tooltip = RealtimeTooltip;
+									}
+								}
 
 								FToolMenuEntry ToggleRealtime = FToolMenuEntry::InitMenuEntry(
 									"ToggleRealtime",
@@ -759,7 +731,25 @@ FToolMenuEntry CreateViewportToolbarPerformanceAndScalabilitySubmenu()
 									RealtimeToggleAction,
 									EUserInterfaceActionType::ToggleButton
 								);
-								ToggleRealtime.SetShowInToolbarTopLevel(bDisplayTopLevel);
+
+								// If we can find a context with a viewport, bind the top-level status of the
+								// realtime button to the viewport's realtime state where we show the realtime
+								// toggle in the top-level if the viewport is NOT realtime.
+								if (ULevelViewportContext* const LevelViewportContext =
+										InnerSection.FindContext<ULevelViewportContext>())
+								{
+									ToggleRealtime.SetShowInToolbarTopLevel(TAttribute<bool>::CreateLambda(
+										[WeakViewport = LevelViewportContext->LevelViewport]() -> bool
+										{
+											if (const TSharedPtr<::SLevelViewport> LevelViewport = WeakViewport.Pin())
+											{
+												return !LevelViewport->IsRealtime();
+											}
+											return false;
+										}
+									));
+								}
+
 								InnerSection.AddEntry(ToggleRealtime);
 							}
 						)
diff --git a/Engine/Source/Editor/LevelEditor/Private/ViewportToolbar/LevelViewportContext.h b/Engine/Source/Editor/LevelEditor/Private/ViewportToolbar/LevelViewportContext.h
--- a/Engine/Source/Editor/LevelEditor/Private/ViewportToolbar/LevelViewportContext.h
+++ b/Engine/Source/Editor/LevelEditor/Private/ViewportToolbar/LevelViewportContext.h
@@ -15,7 +15,4 @@ class LEVELEDITOR_API ULevelViewportContext : public UObject
 
 public:
 	TWeakPtr<SLevelViewport> LevelViewport;
-	// This allows the viewport toolbar to know if it needs to refresh itself to update the top-level state of the
-	// realtime button.
-	TOptional<bool> CachedShouldShowRealtimeOffWarning;
 };
