Initial work on "Select" submenu for the new level editor viewport toolbar

This is just a skeleton to have something to extend. Further work is needed to flesh out the menu.

diff --git a/Engine/Source/Editor/LevelEditor/Private/SLevelViewport.cpp b/Engine/Source/Editor/LevelEditor/Private/SLevelViewport.cpp
--- a/Engine/Source/Editor/LevelEditor/Private/SLevelViewport.cpp
+++ b/Engine/Source/Editor/LevelEditor/Private/SLevelViewport.cpp
@@ -1933,6 +1933,10 @@ TSharedPtr<SWidget> SLevelViewport::MakeViewportToolbar()
 				FToolMenuEntry TransformsSubmenu = UE::LevelEditor::CreateViewportToolbarTransformsSection();
 				TransformsSubmenu.InsertPosition.Position = EToolMenuInsertType::First;
 				LeftSection.AddEntry(TransformsSubmenu);
+
+				FToolMenuEntry SelectionSubmenu = UE::LevelEditor::CreateViewportToolbarSelectionSection();
+				SelectionSubmenu.InsertPosition.Position = EToolMenuInsertType::First;
+				LeftSection.AddEntry(SelectionSubmenu);
 			}
 
 			{
diff --git a/Engine/Source/Editor/LevelEditor/Private/ViewportToolbar/LevelEditorViewportToolbarSections.cpp b/Engine/Source/Editor/LevelEditor/Private/ViewportToolbar/LevelEditorViewportToolbarSections.cpp
--- a/Engine/Source/Editor/LevelEditor/Private/ViewportToolbar/LevelEditorViewportToolbarSections.cpp
+++ b/Engine/Source/Editor/LevelEditor/Private/ViewportToolbar/LevelEditorViewportToolbarSections.cpp
@@ -15,6 +15,7 @@
 #include "Templates/SharedPointer.h"
 #include "ToolMenu.h"
 #include "ToolMenus.h"
+#include "Framework/Commands/GenericCommands.h"
 #include "ViewportToolbar/UnrealEdViewportToolbar.h"
 #include "Widgets/Input/SSpinBox.h"
 #include "Widgets/Input/SVolumeControl.h"
@@ -151,6 +152,61 @@ FToolMenuEntry CreateViewportToolbarTransformsSection()
 	);
 }
 
+FToolMenuEntry CreateViewportToolbarSelectionSection()
+{
+	return FToolMenuEntry::InitSubMenu(
+		"Select",
+		LOCTEXT("SelectonSubmenuLabel", "Select"),
+		LOCTEXT("SelectionSubmenuTooltip", "Viewport-related selection tools"),
+		FNewToolMenuDelegate::CreateLambda(
+			[](UToolMenu* Submenu) -> void
+			{
+				{
+					FToolMenuSection& UnnamedSection = Submenu->FindOrAddSection(NAME_None);
+
+					UnnamedSection.AddMenuEntry(FGenericCommands::Get().SelectAll);
+					UnnamedSection.AddMenuEntry(FLevelEditorCommands::Get().SelectNone);
+					UnnamedSection.AddMenuEntry(FLevelEditorCommands::Get().InvertSelection);
+
+					UnnamedSection.AddSeparator("Advanced");
+
+					UnnamedSection.AddMenuEntry(FLevelEditorCommands::Get().SelectAllActorsOfSameClass);
+				}
+
+				{
+					FToolMenuSection& ByTypeSection =
+						Submenu->FindOrAddSection("ByTypeSection", LOCTEXT("ByTypeSectionLabel", "By Type"));
+
+					ByTypeSection.AddSubMenu(
+						"BSP",
+						LOCTEXT("BspLabel", "BSP"),
+						LOCTEXT("BspTooltip", "BSP-related tools"),
+						FNewToolMenuDelegate::CreateLambda(
+							[](UToolMenu* BspMenu)
+							{
+								FToolMenuSection& SelectAllSection = BspMenu->FindOrAddSection(
+									"SelectAllBSP", LOCTEXT("SelectAllBSPLabel", "Select All BSP")
+								);
+
+								SelectAllSection.AddMenuEntry(FLevelEditorCommands::Get().SelectAllAddditiveBrushes);
+								SelectAllSection.AddMenuEntry(FLevelEditorCommands::Get().SelectAllSubtractiveBrushes);
+								SelectAllSection.AddMenuEntry(FLevelEditorCommands::Get().SelectAllSurfaces);
+							}
+						)
+					);
+				}
+
+				{
+					FToolMenuSection& OptionsSection =
+						Submenu->FindOrAddSection("Options", LOCTEXT("OptionsLabel", "Options"));
+
+					OptionsSection.AddMenuEntry(FLevelEditorCommands::Get().AllowTranslucentSelection);
+				}
+			}
+		)
+	);
+}
+
 TSharedPtr<FExtender> GetViewModesLegacyExtenders()
 {
 	FLevelEditorModule& LevelEditorModule = FModuleManager::GetModuleChecked<FLevelEditorModule>(TEXT("LevelEditor"));
diff --git a/Engine/Source/Editor/LevelEditor/Private/ViewportToolbar/LevelEditorViewportToolbarSections.h b/Engine/Source/Editor/LevelEditor/Private/ViewportToolbar/LevelEditorViewportToolbarSections.h
--- a/Engine/Source/Editor/LevelEditor/Private/ViewportToolbar/LevelEditorViewportToolbarSections.h
+++ b/Engine/Source/Editor/LevelEditor/Private/ViewportToolbar/LevelEditorViewportToolbarSections.h
@@ -19,6 +19,8 @@ bool ShowViewportRealtimeWarning(FLevelEditorViewportClient& ViewportClient);
 
 FToolMenuEntry CreateViewportToolbarTransformsSection();
 
+FToolMenuEntry CreateViewportToolbarSelectionSection();
+
 TSharedPtr<FExtender> GetViewModesLegacyExtenders();
 void PopulateViewModesMenu(UToolMenu* InMenu, TSharedRef<SLevelViewport> InViewport);
 FToolMenuEntry CreateViewportToolbarViewModesSubmenu();
