Prevent SLevelViewportToolBar::GenerateShowMenu() from crashing if its context hasn't been set

diff --git a/Engine/Source/Editor/LevelEditor/Private/SLevelViewportToolBar.cpp b/Engine/Source/Editor/LevelEditor/Private/SLevelViewportToolBar.cpp
--- a/Engine/Source/Editor/LevelEditor/Private/SLevelViewportToolBar.cpp
+++ b/Engine/Source/Editor/LevelEditor/Private/SLevelViewportToolBar.cpp
@@ -1111,11 +1111,18 @@ TSharedRef<SWidget> SLevelViewportToolBar::GenerateShowMenu() const
 	if (!UToolMenus::Get()->IsMenuRegistered(MenuName))
 	{
 		UToolMenu* Menu = UToolMenus::Get()->RegisterMenu(MenuName);
-		Menu->AddDynamicSection("LevelDynamicSection", FNewToolMenuDelegate::CreateLambda([](UToolMenu* InMenu)
-		{
-			ULevelViewportToolBarContext* Context = InMenu->FindContext<ULevelViewportToolBarContext>();
-			Context->LevelViewportToolBarWidget.Pin()->FillShowMenu(InMenu);
-		}));
+		Menu->AddDynamicSection(
+			"LevelDynamicSection",
+			FNewToolMenuDelegate::CreateLambda(
+				[](UToolMenu* InMenu)
+				{
+					if (ULevelViewportToolBarContext* Context = InMenu->FindContext<ULevelViewportToolBarContext>())
+					{
+						Context->LevelViewportToolBarWidget.Pin()->FillShowMenu(InMenu);
+					}
+				}
+			)
+		);
 	}
 
 	Viewport.Pin()->OnFloatingButtonClicked();
