Keep track of whether the old viewport toolbar has created its "Show" menu via a bool instead of the ToolMenu system

By doing this, we can create the menu from the new viewport toolbar to support backward compatibility without breaking the old viewport toolbar. Without this change, creating the old menu for the sake of backward compatiblity will make the current code not build the menu, because it's already registered, so the old viewport toolbar's "Show" menu will be empty.

diff --git a/Engine/Source/Editor/LevelEditor/Private/SLevelViewportToolBar.cpp b/Engine/Source/Editor/LevelEditor/Private/SLevelViewportToolBar.cpp
--- a/Engine/Source/Editor/LevelEditor/Private/SLevelViewportToolBar.cpp
+++ b/Engine/Source/Editor/LevelEditor/Private/SLevelViewportToolBar.cpp
@@ -1108,8 +1108,10 @@ void SLevelViewportToolBar::FillCameraMenu(UToolMenu* Menu) const
 TSharedRef<SWidget> SLevelViewportToolBar::GenerateShowMenu() const
 {
 	static const FName MenuName("LevelEditor.LevelViewportToolbar.Show");
-	if (!UToolMenus::Get()->IsMenuRegistered(MenuName))
+	static bool bDidRegisterMenu = false;
+	if (!bDidRegisterMenu)
 	{
+		bDidRegisterMenu = true;
 		UToolMenu* Menu = UToolMenus::Get()->RegisterMenu(MenuName);
 		Menu->AddDynamicSection(
 			"LevelDynamicSection",
