Add "Show" menu to the new level editor viewport toolbar

diff --git a/Engine/Source/Editor/LevelEditor/Private/SLevelViewport.cpp b/Engine/Source/Editor/LevelEditor/Private/SLevelViewport.cpp
--- a/Engine/Source/Editor/LevelEditor/Private/SLevelViewport.cpp
+++ b/Engine/Source/Editor/LevelEditor/Private/SLevelViewport.cpp
@@ -102,6 +102,7 @@
 #include "LevelViewportLayout.h"
 #include "EditorViewportTabContent.h"
 #include "EditorViewportCommands.h"
+#include "FunctionalUIScreenshotTest.h"
 #include "SkeletalRenderPublic.h"
 #include "ViewportToolbar/LevelEditorViewportToolbarSections.h"
 #include "ViewportToolbar/LevelViewportContext.h"
@@ -1927,40 +1928,68 @@ TSharedPtr<SWidget> SLevelViewport::MakeViewportToolbar()
 			UToolMenu* const ViewportToolbarMenu = UToolMenus::Get()->RegisterMenu(
 				LevelEditorViewportToolbarName, NAME_None /* parent */, EMultiBoxType::ToolBar);
 
+			// Add the left-aligned part of the viewport toolbar.
 			{
 				FToolMenuSection& LeftSection = ViewportToolbarMenu->FindOrAddSection("Left");
 
-				FToolMenuEntry TransformsSubmenu = UE::LevelEditor::CreateViewportToolbarTransformsSection();
-				TransformsSubmenu.InsertPosition.Position = EToolMenuInsertType::First;
-				LeftSection.AddEntry(TransformsSubmenu);
+				// Add the "Transforms" sub menu.
+				{
+					FToolMenuEntry TransformsSubmenu = UE::LevelEditor::CreateViewportToolbarTransformsSection();
+					TransformsSubmenu.InsertPosition.Position = EToolMenuInsertType::First;
+					LeftSection.AddEntry(TransformsSubmenu);
+				}
 
-				FToolMenuEntry SelectionSubmenu = UE::LevelEditor::CreateViewportToolbarSelectionSection();
-				SelectionSubmenu.InsertPosition.Position = EToolMenuInsertType::First;
-				LeftSection.AddEntry(SelectionSubmenu);
+				// Add the "Selection" sub menu.
+				{
+					FToolMenuEntry SelectionSubmenu = UE::LevelEditor::CreateViewportToolbarSelectionSection();
+					SelectionSubmenu.InsertPosition.Position = EToolMenuInsertType::First;
+					LeftSection.AddEntry(SelectionSubmenu);
+				}
 			}
 
+			// Add the right-aligned part of the viewport toolbar.
 			{
 				// Add the submenus of this section as EToolMenuInsertType::Last to sort them after any
 				// default-positioned submenus external code might add.
 				FToolMenuSection& RightSection = ViewportToolbarMenu->FindOrAddSection("Right");
 				RightSection.Alignment = EToolMenuSectionAlign::Last;
 
+				// Add the "View Modes" sub menu.
 				{
 					// Stay backward-compatible with the old viewport toolbar.
 					UToolMenus::Get()->RegisterMenu(
 						"LevelEditor.ViewportToolbar.ViewModes", "LevelEditor.LevelViewportToolbar.View"
 					);
-					FToolMenuEntry LitSubmenu = UE::LevelEditor::CreateViewportToolbarViewModesSubmenu();
-					LitSubmenu.InsertPosition.Position = EToolMenuInsertType::Last;
-					RightSection.AddEntry(LitSubmenu);
+					FToolMenuEntry ViewModesSubmenu = UE::LevelEditor::CreateViewportToolbarViewModesSubmenu();
+					ViewModesSubmenu.InsertPosition.Position = EToolMenuInsertType::Last;
+					RightSection.AddEntry(ViewModesSubmenu);
 				}
 
+				// Add the "Show" submenu.
+				{
+					// Stay backward-compatible with the old viewport toolbar.
+					{
+						if (!UToolMenus::Get()->IsMenuRegistered("LevelEditor.LevelViewportToolbar.Show"))
+						{
+							UToolMenus::Get()->RegisterMenu("LevelEditor.LevelViewportToolbar.Show");
+						}
+						UToolMenus::Get()->RegisterMenu(
+							"LevelEditor.ViewportToolbar.Show", "LevelEditor.LevelViewportToolbar.Show"
+						);
+					}
+					FToolMenuEntry ShowSubmenu = UE::LevelEditor::CreateViewportToolbarShowSubmenu();
+					ShowSubmenu.InsertPosition.Position = EToolMenuInsertType::Last;
+					RightSection.AddEntry(ShowSubmenu);
+				}
+
+				// Add the "Performance & Scalability" submenu.
 				{
 					FToolMenuEntry PerfSubmenu = UE::LevelEditor::CreateViewportToolbarPerformanceAndScalabilitySubmenu();
 					PerfSubmenu.InsertPosition.Position = EToolMenuInsertType::Last;
 					RightSection.AddEntry(PerfSubmenu);
 				}
 
+				// Add the "Settings" submenu.
 				{
 					FToolMenuEntry SettingsSubmenu = UE::LevelEditor::CreateLevelEditorViewportToolbarSettingsSubmenu();
 					SettingsSubmenu.InsertPosition.Position = EToolMenuInsertType::Last;
@@ -1973,9 +2002,22 @@ TSharedPtr<SWidget> SLevelViewport::MakeViewportToolbar()
 	FToolMenuContext ViewportToolbarContext;
 	{
 		ViewportToolbarContext.AppendCommandList(GetCommandList());
-		// Stay backward-compatible with legacy view menu extenders. Note that these extenders can now leak between
-		// submenus of the viewport toolbar.
-		ViewportToolbarContext.AddExtender(UE::LevelEditor::GetViewModesLegacyExtenders());
+		// Note that these extenders can now leak between submenus of the viewport toolbar.
+		{
+			// Stay backward-compatible with legacy view menu extenders.
+			ViewportToolbarContext.AddExtender(UE::LevelEditor::GetViewModesLegacyExtenders());
+
+			// Stay backward-compatible with legacy show menu extenders.
+			{
+				FLevelEditorModule& LevelEditorModule =
+					FModuleManager::GetModuleChecked<FLevelEditorModule>(TEXT("LevelEditor"));
+				TSharedRef<FUICommandList> CommandListRef = GetCommandList().ToSharedRef();
+				TSharedPtr<FExtender> Extenders = LevelEditorModule.AssembleExtenders(
+					CommandListRef, LevelEditorModule.GetAllLevelViewportShowMenuExtenders()
+				);
+				ViewportToolbarContext.AddExtender(Extenders);
+			}
+		}
 
 		ULevelViewportContext* const ContextObject = NewObject<ULevelViewportContext>();
 		ContextObject->LevelViewport = SharedThis(this);
diff --git a/Engine/Source/Editor/LevelEditor/Private/SLevelViewportToolBar.cpp b/Engine/Source/Editor/LevelEditor/Private/SLevelViewportToolBar.cpp
--- a/Engine/Source/Editor/LevelEditor/Private/SLevelViewportToolBar.cpp
+++ b/Engine/Source/Editor/LevelEditor/Private/SLevelViewportToolBar.cpp
@@ -58,6 +58,7 @@
 #include "SActionableMessageViewportWidget.h"
 #include "Engine/SceneCapture.h"
 #include "ViewportToolbar/LevelEditorViewportToolbarSections.h"
+#include "ViewportToolbar/LevelViewportContext.h"
 
 #define LOCTEXT_NAMESPACE "LevelViewportToolBar"
 
@@ -99,94 +100,6 @@ public:
 	}
 };
 
-static void FillShowMenuStatic(UToolMenu* Menu, TArray< FLevelViewportCommands::FShowMenuCommand > MenuCommands, int32 EntryOffset)
-{
-	FToolMenuSection& Section = Menu->AddSection("Section");
-
-	// Generate entries for the standard show flags
-	// Assumption: the first 'n' entries types like 'Show All' and 'Hide All' buttons, so insert a separator after them
-	for (int32 EntryIndex = 0; EntryIndex < MenuCommands.Num(); ++EntryIndex)
-	{
-		FName EntryName = NAME_None;
-
-		if (MenuCommands[EntryIndex].ShowMenuItem)
-		{
-			EntryName = MenuCommands[EntryIndex].ShowMenuItem->GetCommandName();
-			ensure(Section.FindEntry(EntryName) == nullptr);
-		}
-
-		Section.AddMenuEntry(
-			EntryName,
-			MenuCommands[EntryIndex].ShowMenuItem,
-			MenuCommands[EntryIndex].LabelOverride
-		);
-
-		if (EntryIndex == EntryOffset - 1)
-		{
-			Section.AddSeparator(NAME_None);
-		}
-	}
-}
-
-#if STATS
-static void FillShowStatsSubMenus(UToolMenu* Menu, TArray< FLevelViewportCommands::FShowMenuCommand > MenuCommands, TMap< FString, TArray< FLevelViewportCommands::FShowMenuCommand > > StatCatCommands)
-{
-	FillShowMenuStatic(Menu, MenuCommands, 1);
-
-	FToolMenuSection& Section = Menu->FindOrAddSection("Section");
-
-	// Separate out stats into two list, those with and without submenus
-	TArray< FLevelViewportCommands::FShowMenuCommand > SingleStatCommands;
-	TMap< FString, TArray< FLevelViewportCommands::FShowMenuCommand > > SubbedStatCommands;
-	for (auto StatCatIt = StatCatCommands.CreateConstIterator(); StatCatIt; ++StatCatIt)
-	{
-		const TArray< FLevelViewportCommands::FShowMenuCommand >& ShowStatCommands = StatCatIt.Value();
-		const FString& CategoryName = StatCatIt.Key();
-
-		// If no category is specified, or there's only one category, don't use submenus
-		FString NoCategory = FStatConstants::NAME_NoCategory.ToString();
-		NoCategory.RemoveFromStart(TEXT("STATCAT_"));
-		if (CategoryName == NoCategory || StatCatCommands.Num() == 1)
-		{
-			for (int32 StatIndex = 0; StatIndex < ShowStatCommands.Num(); ++StatIndex)
-			{
-				const FLevelViewportCommands::FShowMenuCommand& StatCommand = ShowStatCommands[StatIndex];
-				SingleStatCommands.Add(StatCommand);
-			}
-		}
-		else
-		{
-			SubbedStatCommands.Add(CategoryName, ShowStatCommands);
-		}
-	}
-
-	// First add all the stats that don't have a sub menu
-	for (auto StatCatIt = SingleStatCommands.CreateConstIterator(); StatCatIt; ++StatCatIt)
-	{
-		const FLevelViewportCommands::FShowMenuCommand& StatCommand = *StatCatIt;
-		Section.AddMenuEntry(
-			NAME_None,
-			StatCommand.ShowMenuItem,
-			StatCommand.LabelOverride
-		);
-	}
-
-	// Now add all the stats that have sub menus
-	for (auto StatCatIt = SubbedStatCommands.CreateConstIterator(); StatCatIt; ++StatCatIt)
-	{
-		const TArray< FLevelViewportCommands::FShowMenuCommand >& StatCommands = StatCatIt.Value();
-		const FText CategoryName = FText::FromString(StatCatIt.Key());
-
-		FFormatNamedArguments Args;
-		Args.Add(TEXT("StatCat"), CategoryName);
-		const FText CategoryDescription = FText::Format(NSLOCTEXT("UICommands", "StatShowCatName", "Show {StatCat} stats"), Args);
-
-		Section.AddSubMenu(NAME_None, CategoryName, CategoryDescription,
-			FNewToolMenuDelegate::CreateStatic(&FillShowMenuStatic, StatCommands, 0));
-	}
-}
-#endif
-
 void SLevelViewportToolBar::Construct( const FArguments& InArgs )
 {
 	Viewport = InArgs._Viewport;
@@ -745,24 +658,11 @@ void SLevelViewportToolBar::FillOptionsMenu(UToolMenu* Menu)
 			Section.AddMenuEntry(FEditorViewportCommands::Get().ToggleFPS);
 
 #if STATS
-			TArray< FLevelViewportCommands::FShowMenuCommand > HideStatsMenu;
-
-			Section.AddMenuEntry(FEditorViewportCommands::Get().ToggleStats);
-			Section.AddSubMenu(
-				"ShowStatsMenu",
-				LOCTEXT("ShowStatsMenu", "Stat"),
-				LOCTEXT("ShowStatsMenu_ToolTip", "Show Stat commands"),
-				FNewToolMenuDelegate::CreateStatic(&FillShowStatsSubMenus, HideStatsMenu, LevelViewportActions.ShowStatCatCommands),
-				false,
-				FSlateIcon(FAppStyle::Get().GetStyleSetName(), "EditorViewport.SubMenu.Stats"));
-
-
-
+			{
+				Section.AddMenuEntry(FEditorViewportCommands::Get().ToggleStats);
 
-			FText HideAllLabel = LOCTEXT("HideAllLabel", "Hide All");
-		
-			// 'Hide All' button
-			HideStatsMenu.Add(FLevelViewportCommands::FShowMenuCommand(LevelViewportActions.HideAllStats, HideAllLabel));
+				Section.AddEntry(UE::LevelEditor::CreateShowStatsSubmenu());
+			}
 #endif
 
 			Section.AddMenuEntry(LevelViewportActions.ToggleViewportToolbar);
@@ -1135,10 +1035,24 @@ TSharedRef<SWidget> SLevelViewportToolBar::GenerateShowMenu() const
 	FLevelEditorModule& LevelEditorModule = FModuleManager::GetModuleChecked<FLevelEditorModule>(TEXT("LevelEditor"));
 	TSharedPtr<FExtender> MenuExtender = LevelEditorModule.AssembleExtenders(CommandList, LevelEditorModule.GetAllLevelViewportShowMenuExtenders());
 
-	ULevelViewportToolBarContext* ContextObject = NewObject<ULevelViewportToolBarContext>();
-	ContextObject->LevelViewportToolBarWidget = ConstCastSharedRef<SLevelViewportToolBar>(SharedThis(this));
+	FToolMenuContext MenuContext;
+	{
+		MenuContext.AppendCommandList(CommandList);
+		MenuContext.AddExtender(MenuExtender);
+
+		{
+			ULevelViewportToolBarContext* ContextObject = NewObject<ULevelViewportToolBarContext>();
+			ContextObject->LevelViewportToolBarWidget = ConstCastSharedRef<SLevelViewportToolBar>(SharedThis(this));
+			MenuContext.AddObject(ContextObject);
+		}
+
+		{
+			ULevelViewportContext* const ContextObject = NewObject<ULevelViewportContext>();
+			ContextObject->LevelViewport = Viewport;
+			MenuContext.AddObject(ContextObject);
+		}
+	}
 
-	FToolMenuContext MenuContext(CommandList, MenuExtender, ContextObject);
 	return UToolMenus::Get()->GenerateWidget(MenuName, MenuContext);
 }
 
@@ -1161,54 +1075,16 @@ void SLevelViewportToolBar::FillShowMenu(UToolMenu* Menu) const
 
 		{
 			FToolMenuSection& Section = Menu->AddSection("LevelViewportEditorShow", LOCTEXT("EditorShowHeader", "Editor"));
-			// Show Volumes sub-menu
-			{
-				TArray< FLevelViewportCommands::FShowMenuCommand > ShowVolumesMenu;
-
-				// 'Show All' and 'Hide All' buttons
-				ShowVolumesMenu.Add(FLevelViewportCommands::FShowMenuCommand(Actions.ShowAllVolumes, ShowAllLabel));
-				ShowVolumesMenu.Add(FLevelViewportCommands::FShowMenuCommand(Actions.HideAllVolumes, HideAllLabel));
-
-				// Get each show flag command and put them in their corresponding groups
-				ShowVolumesMenu += Actions.ShowVolumeCommands;
-
-				Section.AddSubMenu("ShowVolumes", LOCTEXT("ShowVolumesMenu", "Volumes"), LOCTEXT("ShowVolumesMenu_ToolTip", "Show volumes flags"),
-					FNewToolMenuDelegate::CreateStatic(&FillShowMenuStatic, ShowVolumesMenu, 2), false, FSlateIcon(FAppStyle::Get().GetStyleSetName(), "ShowFlagsMenu.SubMenu.Volumes"));
-			}
-
-			// Show Layers sub-menu is dynamically generated when the user enters 'show' menu
-			if (!UWorld::IsPartitionedWorld(World))
-			{
-				Section.AddSubMenu("ShowLayers", LOCTEXT("ShowLayersMenu", "Layers"), LOCTEXT("ShowLayersMenu_ToolTip", "Show layers flags"),
-					FNewToolMenuDelegate::CreateStatic(&SLevelViewportToolBar::FillShowLayersMenu, Viewport), false, FSlateIcon(FAppStyle::Get().GetStyleSetName(), "ShowFlagsMenu.SubMenu.Layers"));
-			}
-
-			// Show Sprites sub-menu
-			{
-				TArray< FLevelViewportCommands::FShowMenuCommand > ShowSpritesMenu;
 
-				// 'Show All' and 'Hide All' buttons
-				ShowSpritesMenu.Add(FLevelViewportCommands::FShowMenuCommand(Actions.ShowAllSprites, ShowAllLabel));
-				ShowSpritesMenu.Add(FLevelViewportCommands::FShowMenuCommand(Actions.HideAllSprites, HideAllLabel));
-
-				// Get each show flag command and put them in their corresponding groups
-				ShowSpritesMenu += Actions.ShowSpriteCommands;
-
-				Section.AddSubMenu("ShowSprites", LOCTEXT("ShowSpritesMenu", "Sprites"), LOCTEXT("ShowSpritesMenu_ToolTip", "Show sprites flags"),
-					FNewToolMenuDelegate::CreateStatic(&FillShowMenuStatic, ShowSpritesMenu, 2), false, FSlateIcon(FAppStyle::Get().GetStyleSetName(), "ShowFlagsMenu.SubMenu.Sprites"));
-			}
-
-			// Show 'Foliage types' sub-menu is dynamically generated when the user enters 'show' menu
-			{
-				Section.AddSubMenu("ShowFoliageTypes", LOCTEXT("ShowFoliageTypesMenu", "Foliage Types"), LOCTEXT("ShowFoliageTypesMenu_ToolTip", "Show/hide specific foliage types"),
-					FNewToolMenuDelegate::CreateStatic(&SLevelViewportToolBar::FillShowFoliageTypesMenu, Viewport), false, FSlateIcon(FAppStyle::Get().GetStyleSetName(), "ShowFlagsMenu.SubMenu.FoliageTypes"));
-			}
+			Section.AddEntry(UE::LevelEditor::CreateShowVolumesSubmenu());
+			Section.AddEntry(UE::LevelEditor::CreateShowLayersSubmenu());
+			Section.AddEntry(UE::LevelEditor::CreateShowSpritesSubmenu());
+			Section.AddEntry(UE::LevelEditor::CreateShowFoliageSubmenu());
 
 			// Show 'HLODs' sub-menu is dynamically generated when the user enters 'show' menu
 			if (World->IsPartitionedWorld())
 			{
-				Section.AddSubMenu("ShowHLODsMenu", LOCTEXT("ShowHLODsMenu", "HLODs"), LOCTEXT("ShowHLODsMenu_ToolTip", "Settings for HLODs in editor"),
-					FNewToolMenuDelegate::CreateSP(this, &SLevelViewportToolBar::FillShowHLODsMenu), false, FSlateIcon(FAppStyle::Get().GetStyleSetName(), "ShowFlagsMenu.SubMenu.HLODs"));
+				Section.AddEntry(UE::LevelEditor::CreateShowHLODsSubmenu());
 			}
 		}
 	}
@@ -1337,119 +1213,6 @@ void SLevelViewportToolBar::OnFarViewPlaneValueChanged( float NewValue )
 	Viewport.Pin()->GetLevelViewportClient().OverrideFarClipPlane(NewValue);
 }
 
-void SLevelViewportToolBar::FillShowLayersMenu(UToolMenu* Menu, TWeakPtr<class SLevelViewport> Viewport)
-{	
-	{
-		FToolMenuSection& Section = Menu->AddSection("LevelViewportLayers");
-		Section.AddMenuEntry(FLevelViewportCommands::Get().ShowAllLayers, LOCTEXT("ShowAllLabel", "Show All"));
-		Section.AddMenuEntry(FLevelViewportCommands::Get().HideAllLayers, LOCTEXT("HideAllLabel", "Hide All"));
-	}
-
-	if( Viewport.IsValid() )
-	{
-		TSharedRef<SLevelViewport> ViewportRef = Viewport.Pin().ToSharedRef();
-		{
-			FToolMenuSection& Section = Menu->AddSection("LevelViewportLayers2");
-			// Get all the layers and create an entry for each of them
-			TArray<FName> AllLayerNames;
-			ULayersSubsystem* Layers = GEditor->GetEditorSubsystem<ULayersSubsystem>();
-			Layers->AddAllLayerNamesTo(AllLayerNames);
-
-			for( int32 LayerIndex = 0; LayerIndex < AllLayerNames.Num(); ++LayerIndex )
-			{
-				const FName LayerName = AllLayerNames[LayerIndex];
-				//const FString LayerNameString = LayerName;
-
-				FUIAction Action( FExecuteAction::CreateSP( ViewportRef, &SLevelViewport::ToggleShowLayer, LayerName ),
-						FCanExecuteAction(),
-						FIsActionChecked::CreateSP( ViewportRef, &SLevelViewport::IsLayerVisible, LayerName ) );
-
-				Section.AddMenuEntry(NAME_None, FText::FromName( LayerName ), FText::GetEmpty(), FSlateIcon(), Action, EUserInterfaceActionType::ToggleButton);
-			}
-		}
-	}
-}
-
-static TMap<FName, TArray<UFoliageType*>> GroupFoliageByOuter(const TArray<UFoliageType*> FoliageList)
-{
-	TMap<FName, TArray<UFoliageType*>> Result;
-
-	for (UFoliageType* FoliageType : FoliageList)
-	{
-		if (FoliageType->IsAsset())
-		{
-			Result.FindOrAdd(NAME_None).Add(FoliageType);
-		}
-		else
-		{
-			FName LevelName = FoliageType->GetOutermost()->GetFName();
-			Result.FindOrAdd(LevelName).Add(FoliageType);
-		}
-	}
-
-	Result.KeySort([](const FName& A, const FName& B) { return (A.LexicalLess(B) && B != NAME_None); });
-	return Result;
-}
-
-void SLevelViewportToolBar::FillShowFoliageTypesMenu(UToolMenu* Menu, TWeakPtr<class SLevelViewport> Viewport)
-{
-	auto ViewportPtr = Viewport.Pin();
-	if (!ViewportPtr.IsValid())
-	{
-		return;
-	}
-		
-	{
-		FToolMenuSection& Section = Menu->AddSection("LevelViewportFoliageMeshes");
-		// Map 'Show All' and 'Hide All' commands
-		FUIAction ShowAllFoliage(FExecuteAction::CreateSP(ViewportPtr.ToSharedRef(), &SLevelViewport::ToggleAllFoliageTypes, true));
-		FUIAction HideAllFoliage(FExecuteAction::CreateSP(ViewportPtr.ToSharedRef(), &SLevelViewport::ToggleAllFoliageTypes, false));
-		
-		Section.AddMenuEntry("ShowAll", LOCTEXT("ShowAllLabel", "Show All"), FText::GetEmpty(), FSlateIcon(), ShowAllFoliage);
-		Section.AddMenuEntry("HideAll", LOCTEXT("HideAllLabel", "Hide All"), FText::GetEmpty(), FSlateIcon(), HideAllFoliage);
-	}
-	
-	// Gather all foliage types used in this world and group them by sub-levels
-	auto AllFoliageMap = GroupFoliageByOuter(GEditor->GetFoliageTypesInWorld(ViewportPtr->GetWorld()));
-
-	for (auto& FoliagePair : AllFoliageMap)
-	{
-		// Name foliage group by an outer sub-level name, or empty if foliage type is an asset
-		FText EntryName = (FoliagePair.Key == NAME_None ? FText::GetEmpty() : FText::FromName(FPackageName::GetShortFName(FoliagePair.Key)));
-		FToolMenuSection& Section = Menu->AddSection(NAME_None, EntryName);
-
-		TArray<UFoliageType*>& FoliageList = FoliagePair.Value;
-		for (UFoliageType* FoliageType : FoliageList)
-		{
-			FName MeshName = FoliageType->GetDisplayFName();
-			TWeakObjectPtr<UFoliageType> FoliageTypePtr = FoliageType;
-
-			FUIAction Action(
-				FExecuteAction::CreateSP(ViewportPtr.ToSharedRef(), &SLevelViewport::ToggleShowFoliageType, FoliageTypePtr),
-				FCanExecuteAction(),
-				FIsActionChecked::CreateSP(ViewportPtr.ToSharedRef(), &SLevelViewport::IsFoliageTypeVisible, FoliageTypePtr));
-
-			Section.AddMenuEntry(NAME_None, FText::FromName(MeshName), FText::GetEmpty(), FSlateIcon(), Action, EUserInterfaceActionType::ToggleButton);
-		}
-	}
-}
-
-double SLevelViewportToolBar::OnGetHLODInEditorMinDrawDistanceValue() const
-{
-	IWorldPartitionEditorModule* WorldPartitionEditorModule = FModuleManager::GetModulePtr<IWorldPartitionEditorModule>("WorldPartitionEditor");
-	return WorldPartitionEditorModule ? WorldPartitionEditorModule->GetHLODInEditorMinDrawDistance() : 0;
-}
-
-void SLevelViewportToolBar::OnHLODInEditorMinDrawDistanceValueChanged(double NewValue) const
-{
-	IWorldPartitionEditorModule* WorldPartitionEditorModule = FModuleManager::GetModulePtr<IWorldPartitionEditorModule>("WorldPartitionEditor");
-	if (WorldPartitionEditorModule)
-	{
-		WorldPartitionEditorModule->SetHLODInEditorMinDrawDistance(NewValue);
-		GEditor->RedrawLevelEditingViewports(true);
-	}
-}
-
 double SLevelViewportToolBar::OnGetHLODInEditorMaxDrawDistanceValue() const
 {
 	IWorldPartitionEditorModule* WorldPartitionEditorModule = FModuleManager::GetModulePtr<IWorldPartitionEditorModule>("WorldPartitionEditor");
@@ -1466,115 +1229,6 @@ void SLevelViewportToolBar::OnHLODInEditorMaxDrawDistanceValueChanged(double New
 	}
 }
 
-void SLevelViewportToolBar::FillShowHLODsMenu(UToolMenu* Menu) const
-{
-	auto ViewportPtr = Viewport.Pin();
-	if (!ViewportPtr.IsValid())
-	{
-		return;
-	}
-
-	UWorld* World = ViewportPtr->GetWorld();
-	UWorldPartition* WorldPartition = World ? World->GetWorldPartition() : nullptr;
-	IWorldPartitionEditorModule* WorldPartitionEditorModule = FModuleManager::GetModulePtr<IWorldPartitionEditorModule>("WorldPartitionEditor");
-	
-	if (WorldPartition == nullptr || WorldPartitionEditorModule == nullptr)
-	{
-		return;
-	}
-
-	FText HLODInEditorDisallowedReason;
-	const bool bHLODInEditorAllowed = WorldPartitionEditorModule->IsHLODInEditorAllowed(World, &HLODInEditorDisallowedReason);
-
-	// Show HLODs
-	{
-		FToolUIAction UIAction;
-		UIAction.ExecuteAction = FToolMenuExecuteAction::CreateLambda([WorldPartitionEditorModule](const FToolMenuContext& InContext) { WorldPartitionEditorModule->SetShowHLODsInEditor(!WorldPartitionEditorModule->GetShowHLODsInEditor()); });
-		UIAction.CanExecuteAction = FToolMenuCanExecuteAction::CreateLambda([bHLODInEditorAllowed](const FToolMenuContext& InContext) { return bHLODInEditorAllowed; });
-		UIAction.GetActionCheckState = FToolMenuGetActionCheckState::CreateLambda([WorldPartitionEditorModule](const FToolMenuContext& InContext) { return WorldPartitionEditorModule->GetShowHLODsInEditor() ? ECheckBoxState::Checked : ECheckBoxState::Unchecked; });
-		FToolMenuEntry MenuEntry = FToolMenuEntry::InitMenuEntry("ShowHLODs", LOCTEXT("ShowHLODs", "Show HLODs"), bHLODInEditorAllowed ? LOCTEXT("ShowHLODsToolTip", "Show/Hide HLODs") : HLODInEditorDisallowedReason, FSlateIcon(), UIAction, EUserInterfaceActionType::ToggleButton);
-		Menu->AddMenuEntry(NAME_None, MenuEntry);
-	}
-
-	// Show HLODs Over Loaded Regions
-	{
-		FToolUIAction UIAction;
-		UIAction.ExecuteAction = FToolMenuExecuteAction::CreateLambda([WorldPartitionEditorModule](const FToolMenuContext& InContext) { WorldPartitionEditorModule->SetShowHLODsOverLoadedRegions(!WorldPartitionEditorModule->GetShowHLODsOverLoadedRegions()); });
-		UIAction.CanExecuteAction = FToolMenuCanExecuteAction::CreateLambda([bHLODInEditorAllowed](const FToolMenuContext& InContext) { return bHLODInEditorAllowed; });
-		UIAction.GetActionCheckState = FToolMenuGetActionCheckState::CreateLambda([WorldPartitionEditorModule](const FToolMenuContext& InContext) { return WorldPartitionEditorModule->GetShowHLODsOverLoadedRegions() ? ECheckBoxState::Checked : ECheckBoxState::Unchecked; });
-		FToolMenuEntry ShowHLODsEntry = FToolMenuEntry::InitMenuEntry("ShowHLODsOverLoadedRegions", LOCTEXT("ShowHLODsOverLoadedRegions", "Show HLODs Over Loaded Regions"), bHLODInEditorAllowed ? LOCTEXT("ShowHLODsOverLoadedRegions_ToolTip", "Show/Hide HLODs over loaded actors or regions") : HLODInEditorDisallowedReason, FSlateIcon(), UIAction, EUserInterfaceActionType::ToggleButton);
-		Menu->AddMenuEntry(NAME_None, ShowHLODsEntry);
-	}
-
-	// Min/Max Draw Distance
-	{
-		const double MinDrawDistanceMinValue = 0;
-		const double MinDrawDistanceMaxValue = 102400;
-
-		const double MaxDrawDistanceMinValue = 0;
-		const double MaxDrawDistanceMaxValue = 1638400;
-
-		TSharedRef<SSpinBox<double>> MinDrawDistanceSpinBox = SNew(SSpinBox<double>)
-			.MinValue(MinDrawDistanceMinValue)
-			.MaxValue(MinDrawDistanceMaxValue)
-			.IsEnabled(bHLODInEditorAllowed)
-			.Value(this, &SLevelViewportToolBar::OnGetHLODInEditorMinDrawDistanceValue)
-			.OnValueChanged(this, &SLevelViewportToolBar::OnHLODInEditorMinDrawDistanceValueChanged)
-			.ToolTipText(bHLODInEditorAllowed ? LOCTEXT("HLODsInEditor_MinDrawDistance_Tooltip", "Sets the minimum distance at which HLOD will be rendered") : HLODInEditorDisallowedReason)
-			.OnBeginSliderMovement_Lambda([this]()
-			{
-				// Disable Slate throttling during slider drag to ensure immediate updates while moving the slider.
-				FSlateThrottleManager::Get().DisableThrottle(true);
-			})
-			.OnEndSliderMovement_Lambda([this](float)
-			{
-				FSlateThrottleManager::Get().DisableThrottle(false);
-			});
-
-		TSharedRef<SSpinBox<double>> MaxDrawDistanceSpinBox = SNew(SSpinBox<double>)
-			.MinValue(MaxDrawDistanceMinValue)
-			.MaxValue(MaxDrawDistanceMaxValue)
-			.IsEnabled(bHLODInEditorAllowed)
-			.Value(this, &SLevelViewportToolBar::OnGetHLODInEditorMaxDrawDistanceValue)
-			.OnValueChanged(this, &SLevelViewportToolBar::OnHLODInEditorMaxDrawDistanceValueChanged)
-			.ToolTipText(bHLODInEditorAllowed ? LOCTEXT("HLODsInEditor_MaxDrawDistance_Tooltip", "Sets the maximum distance at which HLODs will be rendered") : HLODInEditorDisallowedReason)
-			.OnBeginSliderMovement_Lambda([this]()
-			{
-				// Disable Slate throttling during slider drag to ensure immediate updates while moving the slider.
-				FSlateThrottleManager::Get().DisableThrottle(true);
-			})
-			.OnEndSliderMovement_Lambda([this](float)
-			{
-				FSlateThrottleManager::Get().DisableThrottle(false);
-			});
-
-		auto CreateDrawDistanceWidget = [](TSharedRef<SSpinBox<double>> InSpinBoxWidget)
-		{
-			return SNew(SBox)
-				.HAlign(HAlign_Right)
-				[
-					SNew(SBox)
-						.Padding(FMargin(0.0f, 0.0f, 0.0f, 0.0f))
-						.WidthOverride(100.0f)
-						[
-							SNew(SBorder)
-								.BorderImage(FAppStyle::Get().GetBrush("Menu.WidgetBorder"))
-								.Padding(FMargin(1.0f))
-								[
-									InSpinBoxWidget
-								]
-						]
-				];
-		};
-			
-		FToolMenuEntry MinDrawDistanceMenuEntry = FToolMenuEntry::InitWidget("Min Draw Distance", CreateDrawDistanceWidget(MinDrawDistanceSpinBox), LOCTEXT("MinDrawDistance", "Min Draw Distance"));
-		Menu->AddMenuEntry(NAME_None, MinDrawDistanceMenuEntry);
-
-		FToolMenuEntry MaxDrawDistanceMenuEntry = FToolMenuEntry::InitWidget("Max Draw Distance", CreateDrawDistanceWidget(MaxDrawDistanceSpinBox), LOCTEXT("MaxDrawDistance", "Max Draw Distance"));
-		Menu->AddMenuEntry(NAME_None, MaxDrawDistanceMenuEntry);
-	}
-}
-
 TWeakObjectPtr<UWorld> SLevelViewportToolBar::GetWorld() const
 {
 	if (Viewport.IsValid())
diff --git a/Engine/Source/Editor/LevelEditor/Private/ViewportToolbar/LevelEditorViewportToolbarSections.cpp b/Engine/Source/Editor/LevelEditor/Private/ViewportToolbar/LevelEditorViewportToolbarSections.cpp
--- a/Engine/Source/Editor/LevelEditor/Private/ViewportToolbar/LevelEditorViewportToolbarSections.cpp
+++ b/Engine/Source/Editor/LevelEditor/Private/ViewportToolbar/LevelEditorViewportToolbarSections.cpp
@@ -3,23 +3,28 @@
 #include "ViewportToolbar/LevelEditorViewportToolbarSections.h"
 
 #include "EditorViewportCommands.h"
+#include "FoliageType.h"
+#include "Framework/Commands/GenericCommands.h"
 #include "GameFramework/ActorPrimitiveColorHandler.h"
 #include "GroomVisualizationData.h"
+#include "Layers/LayersSubsystem.h"
 #include "LevelEditor.h"
 #include "LevelEditorActions.h"
 #include "LevelViewportActions.h"
-#include "LevelViewportContext.h"
 #include "SCommonEditorViewportToolbarBase.h"
 #include "SLevelViewport.h"
 #include "SScalabilitySettings.h"
+#include "ShowFlagMenuCommands.h"
+#include "Stats/StatsData.h"
 #include "Templates/SharedPointer.h"
 #include "ToolMenu.h"
 #include "ToolMenus.h"
-#include "Framework/Commands/GenericCommands.h"
+#include "ViewportToolbar/LevelViewportContext.h"
 #include "ViewportToolbar/UnrealEdViewportToolbar.h"
 #include "Widgets/Input/SSpinBox.h"
 #include "Widgets/Input/SVolumeControl.h"
 #include "Widgets/SBoxPanel.h"
+#include "WorldPartition/IWorldPartitionEditorModule.h"
 
 #define LOCTEXT_NAMESPACE "LevelEditorViewportToolbar"
 
@@ -37,6 +42,91 @@ void OnLandscapeLODChanged(FLevelEditorViewportClient& ViewportClient, int32 New
 	ViewportClient.Invalidate();
 }
 
+TMap<FName, TArray<UFoliageType*>> GroupFoliageByOuter(const TArray<UFoliageType*> FoliageList)
+{
+	TMap<FName, TArray<UFoliageType*>> Result;
+
+	for (UFoliageType* FoliageType : FoliageList)
+	{
+		if (FoliageType->IsAsset())
+		{
+			Result.FindOrAdd(NAME_None).Add(FoliageType);
+		}
+		else
+		{
+			FName LevelName = FoliageType->GetOutermost()->GetFName();
+			Result.FindOrAdd(LevelName).Add(FoliageType);
+		}
+	}
+
+	Result.KeySort(
+		[](const FName& A, const FName& B)
+		{
+			return (A.LexicalLess(B) && B != NAME_None);
+		}
+	);
+	return Result;
+}
+
+void PopulateMenuWithCommands(UToolMenu* Menu, TArray<FLevelViewportCommands::FShowMenuCommand> MenuCommands, int32 EntryOffset)
+{
+	FToolMenuSection& Section = Menu->AddSection("Section");
+
+	// Generate entries for the standard show flags
+	// Assumption: the first 'n' entries types like 'Show All' and 'Hide All' buttons, so insert a separator after them
+	for (int32 EntryIndex = 0; EntryIndex < MenuCommands.Num(); ++EntryIndex)
+	{
+		FName EntryName = NAME_None;
+
+		if (MenuCommands[EntryIndex].ShowMenuItem)
+		{
+			EntryName = MenuCommands[EntryIndex].ShowMenuItem->GetCommandName();
+			ensure(Section.FindEntry(EntryName) == nullptr);
+		}
+
+		Section.AddMenuEntry(EntryName, MenuCommands[EntryIndex].ShowMenuItem, MenuCommands[EntryIndex].LabelOverride);
+
+		if (EntryIndex == EntryOffset - 1)
+		{
+			Section.AddSeparator(NAME_None);
+		}
+	}
+}
+
+void PopulateShowLayersSubmenu(UToolMenu* InMenu, TWeakPtr<::SLevelViewport> InViewport)
+{
+	{
+		FToolMenuSection& Section = InMenu->AddSection("LevelViewportLayers");
+		Section.AddMenuEntry(FLevelViewportCommands::Get().ShowAllLayers, LOCTEXT("ShowAllLabel", "Show All"));
+		Section.AddMenuEntry(FLevelViewportCommands::Get().HideAllLayers, LOCTEXT("HideAllLabel", "Hide All"));
+	}
+
+	if (TSharedPtr<::SLevelViewport> ViewportPinned = InViewport.Pin())
+	{
+		FToolMenuSection& Section = InMenu->AddSection("LevelViewportLayers2");
+		// Get all the layers and create an entry for each of them
+		TArray<FName> AllLayerNames;
+		ULayersSubsystem* Layers = GEditor->GetEditorSubsystem<ULayersSubsystem>();
+		Layers->AddAllLayerNamesTo(AllLayerNames);
+
+		for (int32 LayerIndex = 0; LayerIndex < AllLayerNames.Num(); ++LayerIndex)
+		{
+			const FName LayerName = AllLayerNames[LayerIndex];
+			// const FString LayerNameString = LayerName;
+
+			FUIAction Action(
+				FExecuteAction::CreateSP(ViewportPinned.ToSharedRef(), &::SLevelViewport::ToggleShowLayer, LayerName),
+				FCanExecuteAction(),
+				FIsActionChecked::CreateSP(ViewportPinned.ToSharedRef(), &::SLevelViewport::IsLayerVisible, LayerName)
+			);
+
+			Section.AddMenuEntry(
+				NAME_None, FText::FromName(LayerName), FText::GetEmpty(), FSlateIcon(), Action, EUserInterfaceActionType::ToggleButton
+			);
+		}
+	}
+}
+
 } // namespace UE::LevelEditor::Private
 
 namespace UE::LevelEditor
@@ -638,6 +728,606 @@ FToolMenuEntry CreateViewportToolbarViewModesSubmenu()
 	);
 }
 
+FToolMenuEntry CreateShowFoliageSubmenu()
+{
+	return FToolMenuEntry::InitSubMenu(
+		"ShowFoliage",
+		LOCTEXT("ShowFoliageTypesMenu", "Foliage Types"),
+		LOCTEXT("ShowFoliageTypesMenu_ToolTip", "Show/hide specific foliage types"),
+		FNewToolMenuDelegate::CreateLambda(
+			[](UToolMenu* Submenu)
+			{
+				ULevelViewportContext* const LevelViewportContext = Submenu->FindContext<ULevelViewportContext>();
+				if (!LevelViewportContext)
+				{
+					return;
+				}
+
+				TSharedPtr<::SLevelViewport> Viewport = LevelViewportContext->LevelViewport.Pin();
+				if (!Viewport)
+				{
+					return;
+				}
+
+				{
+					FToolMenuSection& Section = Submenu->AddSection("LevelViewportFoliageMeshes");
+					// Map 'Show All' and 'Hide All' commands
+					FUIAction ShowAllFoliage(
+						FExecuteAction::CreateSP(Viewport.ToSharedRef(), &::SLevelViewport::ToggleAllFoliageTypes, true)
+					);
+					FUIAction HideAllFoliage(
+						FExecuteAction::CreateSP(Viewport.ToSharedRef(), &::SLevelViewport::ToggleAllFoliageTypes, false)
+					);
+
+					Section.AddMenuEntry(
+						"ShowAll", LOCTEXT("ShowAllLabel", "Show All"), FText::GetEmpty(), FSlateIcon(), ShowAllFoliage
+					);
+					Section.AddMenuEntry(
+						"HideAll", LOCTEXT("HideAllLabel", "Hide All"), FText::GetEmpty(), FSlateIcon(), HideAllFoliage
+					);
+				}
+
+				// Gather all foliage types used in this world and group them by sub-levels
+				auto AllFoliageMap =
+					UE::LevelEditor::Private::GroupFoliageByOuter(GEditor->GetFoliageTypesInWorld(Viewport->GetWorld()));
+
+				for (auto& FoliagePair : AllFoliageMap)
+				{
+					// Name foliage group by an outer sub-level name, or empty if foliage type is an asset
+					FText EntryName =
+						(FoliagePair.Key == NAME_None ? FText::GetEmpty()
+													  : FText::FromName(FPackageName::GetShortFName(FoliagePair.Key)));
+					FToolMenuSection& Section = Submenu->AddSection(NAME_None, EntryName);
+
+					TArray<UFoliageType*>& FoliageList = FoliagePair.Value;
+					for (UFoliageType* FoliageType : FoliageList)
+					{
+						FName MeshName = FoliageType->GetDisplayFName();
+						TWeakObjectPtr<UFoliageType> FoliageTypePtr = FoliageType;
+
+						FUIAction Action(
+							FExecuteAction::CreateSP(
+								Viewport.ToSharedRef(), &::SLevelViewport::ToggleShowFoliageType, FoliageTypePtr
+							),
+							FCanExecuteAction(),
+							FIsActionChecked::CreateSP(Viewport.ToSharedRef(), &::SLevelViewport::IsFoliageTypeVisible, FoliageTypePtr)
+						);
+
+						Section.AddMenuEntry(
+							NAME_None,
+							FText::FromName(MeshName),
+							FText::GetEmpty(),
+							FSlateIcon(),
+							Action,
+							EUserInterfaceActionType::ToggleButton
+						);
+					}
+				}
+			}
+		),
+		false,
+		FSlateIcon(FAppStyle::Get().GetStyleSetName(), "ShowFlagsMenu.SubMenu.FoliageTypes")
+	);
+}
+
+FToolMenuEntry CreateShowHLODsSubmenu()
+{
+	// This is a dynamic entry so we can skip adding the submenu if the context
+	// indicates that the viewport's world isn't partitioned.
+	return FToolMenuEntry::InitDynamicEntry(
+		"ShowHLODsDynamic",
+		FNewToolMenuSectionDelegate::CreateLambda(
+			[](FToolMenuSection& InDynamicSection) -> void
+			{
+				ULevelViewportContext* const LevelViewportContext = InDynamicSection.FindContext<ULevelViewportContext>();
+				if (!LevelViewportContext)
+				{
+					return;
+				}
+
+				TSharedPtr<::SLevelViewport> Viewport = LevelViewportContext->LevelViewport.Pin();
+				if (!Viewport)
+				{
+					return;
+				}
+
+				UWorld* World = Viewport->GetWorld();
+				if (!World)
+				{
+					return;
+				}
+
+				// Only add this submenu for partitioned worlds.
+				if (!World->IsPartitionedWorld())
+				{
+					return;
+				}
+
+				InDynamicSection.AddSubMenu(
+					"ShowHLODsMenu",
+					LOCTEXT("ShowHLODsMenu", "HLODs"),
+					LOCTEXT("ShowHLODsMenu_ToolTip", "Settings for HLODs in editor"),
+					FNewToolMenuDelegate::CreateLambda(
+						[](UToolMenu* Submenu)
+						{
+							ULevelViewportContext* const LevelViewportContext =
+								Submenu->FindContext<ULevelViewportContext>();
+							if (!LevelViewportContext)
+							{
+								return;
+							}
+
+							TSharedPtr<::SLevelViewport> Viewport = LevelViewportContext->LevelViewport.Pin();
+							if (!Viewport)
+							{
+								return;
+							}
+
+							UWorld* World = Viewport->GetWorld();
+							UWorldPartition* WorldPartition = World ? World->GetWorldPartition() : nullptr;
+							if (!WorldPartition)
+							{
+								return;
+							}
+
+							IWorldPartitionEditorModule* WorldPartitionEditorModule =
+								FModuleManager::GetModulePtr<IWorldPartitionEditorModule>("WorldPartitionEditor");
+							if (!WorldPartitionEditorModule)
+							{
+								return;
+							}
+
+							FText HLODInEditorDisallowedReason;
+							const bool bHLODInEditorAllowed =
+								WorldPartitionEditorModule->IsHLODInEditorAllowed(World, &HLODInEditorDisallowedReason);
+
+							// Show HLODs
+							{
+								FToolUIAction UIAction;
+								UIAction.ExecuteAction = FToolMenuExecuteAction::CreateLambda(
+									[WorldPartitionEditorModule](const FToolMenuContext& InContext)
+									{
+										WorldPartitionEditorModule->SetShowHLODsInEditor(
+											!WorldPartitionEditorModule->GetShowHLODsInEditor()
+										);
+									}
+								);
+								UIAction.CanExecuteAction = FToolMenuCanExecuteAction::CreateLambda(
+									[bHLODInEditorAllowed](const FToolMenuContext& InContext)
+									{
+										return bHLODInEditorAllowed;
+									}
+								);
+								UIAction.GetActionCheckState = FToolMenuGetActionCheckState::CreateLambda(
+									[WorldPartitionEditorModule](const FToolMenuContext& InContext)
+									{
+										return WorldPartitionEditorModule->GetShowHLODsInEditor()
+												 ? ECheckBoxState::Checked
+												 : ECheckBoxState::Unchecked;
+									}
+								);
+								FToolMenuEntry MenuEntry = FToolMenuEntry::InitMenuEntry(
+									"ShowHLODs",
+									LOCTEXT("ShowHLODs", "Show HLODs"),
+									bHLODInEditorAllowed ? LOCTEXT("ShowHLODsToolTip", "Show/Hide HLODs")
+														 : HLODInEditorDisallowedReason,
+									FSlateIcon(),
+									UIAction,
+									EUserInterfaceActionType::ToggleButton
+								);
+								Submenu->AddMenuEntry(NAME_None, MenuEntry);
+							}
+
+							// Show HLODs Over Loaded Regions
+							{
+								FToolUIAction UIAction;
+								UIAction.ExecuteAction = FToolMenuExecuteAction::CreateLambda(
+									[WorldPartitionEditorModule](const FToolMenuContext& InContext)
+									{
+										WorldPartitionEditorModule->SetShowHLODsOverLoadedRegions(
+											!WorldPartitionEditorModule->GetShowHLODsOverLoadedRegions()
+										);
+									}
+								);
+								UIAction.CanExecuteAction = FToolMenuCanExecuteAction::CreateLambda(
+									[bHLODInEditorAllowed](const FToolMenuContext& InContext)
+									{
+										return bHLODInEditorAllowed;
+									}
+								);
+								UIAction.GetActionCheckState = FToolMenuGetActionCheckState::CreateLambda(
+									[WorldPartitionEditorModule](const FToolMenuContext& InContext)
+									{
+										return WorldPartitionEditorModule->GetShowHLODsOverLoadedRegions()
+												 ? ECheckBoxState::Checked
+												 : ECheckBoxState::Unchecked;
+									}
+								);
+								FToolMenuEntry ShowHLODsEntry = FToolMenuEntry::InitMenuEntry(
+									"ShowHLODsOverLoadedRegions",
+									LOCTEXT("ShowHLODsOverLoadedRegions", "Show HLODs Over Loaded Regions"),
+									bHLODInEditorAllowed
+										? LOCTEXT("ShowHLODsOverLoadedRegions_ToolTip", "Show/Hide HLODs over loaded actors or regions")
+										: HLODInEditorDisallowedReason,
+									FSlateIcon(),
+									UIAction,
+									EUserInterfaceActionType::ToggleButton
+								);
+								Submenu->AddMenuEntry(NAME_None, ShowHLODsEntry);
+							}
+
+							// Min/Max Draw Distance
+							{
+								const double MinDrawDistanceMinValue = 0;
+								const double MinDrawDistanceMaxValue = 102400;
+
+								const double MaxDrawDistanceMinValue = 0;
+								const double MaxDrawDistanceMaxValue = 1638400;
+
+								// double SLevelViewportToolBar::OnGetHLODInEditorMinDrawDistanceValue() const
+								auto OnGetHLODInEditorMinDrawDistanceValue = []() -> double
+								{
+									IWorldPartitionEditorModule* WorldPartitionEditorModule =
+										FModuleManager::GetModulePtr<IWorldPartitionEditorModule>("WorldPartitionEditor");
+									return WorldPartitionEditorModule
+											 ? WorldPartitionEditorModule->GetHLODInEditorMinDrawDistance()
+											 : 0;
+								};
+
+								// void SLevelViewportToolBar::OnHLODInEditorMinDrawDistanceValueChanged(double NewValue) const
+								auto OnHLODInEditorMinDrawDistanceValueChanged = [](double NewValue) -> void
+								{
+									IWorldPartitionEditorModule* WorldPartitionEditorModule =
+										FModuleManager::GetModulePtr<IWorldPartitionEditorModule>("WorldPartitionEditor");
+									if (WorldPartitionEditorModule)
+									{
+										WorldPartitionEditorModule->SetHLODInEditorMinDrawDistance(NewValue);
+										GEditor->RedrawLevelEditingViewports(true);
+									}
+								};
+
+								TSharedRef<SSpinBox<double>> MinDrawDistanceSpinBox =
+									SNew(SSpinBox<double>)
+										.MinValue(MinDrawDistanceMinValue)
+										.MaxValue(MinDrawDistanceMaxValue)
+										.IsEnabled(bHLODInEditorAllowed)
+										.Value_Lambda(OnGetHLODInEditorMinDrawDistanceValue)
+										.OnValueChanged_Lambda(OnHLODInEditorMinDrawDistanceValueChanged)
+										.ToolTipText(
+											bHLODInEditorAllowed
+												? LOCTEXT(
+													  "HLODsInEditor_MinDrawDistance_Tooltip",
+													  "Sets the minimum distance at which HLOD will be rendered"
+												  )
+												: HLODInEditorDisallowedReason
+										)
+										.OnBeginSliderMovement_Lambda(
+											[]()
+											{
+												// Disable Slate throttling during slider drag to ensure immediate updates while moving the slider.
+												FSlateThrottleManager::Get().DisableThrottle(true);
+											}
+										)
+										.OnEndSliderMovement_Lambda(
+											[](float)
+											{
+												FSlateThrottleManager::Get().DisableThrottle(false);
+											}
+										);
+
+								TSharedRef<SSpinBox<double>> MaxDrawDistanceSpinBox =
+									SNew(SSpinBox<double>)
+										.MinValue(MaxDrawDistanceMinValue)
+										.MaxValue(MaxDrawDistanceMaxValue)
+										.IsEnabled(bHLODInEditorAllowed)
+										.Value_Lambda(OnGetHLODInEditorMinDrawDistanceValue)
+										.OnValueChanged_Lambda(OnHLODInEditorMinDrawDistanceValueChanged)
+										.ToolTipText(
+											bHLODInEditorAllowed
+												? LOCTEXT(
+													  "HLODsInEditor_MaxDrawDistance_Tooltip",
+													  "Sets the maximum distance at which HLODs will be rendered"
+												  )
+												: HLODInEditorDisallowedReason
+										)
+										.OnBeginSliderMovement_Lambda(
+											[]()
+											{
+												// Disable Slate throttling during slider drag to ensure immediate updates while moving the slider.
+												FSlateThrottleManager::Get().DisableThrottle(true);
+											}
+										)
+										.OnEndSliderMovement_Lambda(
+											[](float)
+											{
+												FSlateThrottleManager::Get().DisableThrottle(false);
+											}
+										);
+
+								auto CreateDrawDistanceWidget = [](TSharedRef<SSpinBox<double>> InSpinBoxWidget)
+								{
+									return SNew(SBox).HAlign(HAlign_Right
+									)[SNew(SBox)
+										  .Padding(FMargin(0.0f, 0.0f, 0.0f, 0.0f))
+										  .WidthOverride(100.0f
+										  )[SNew(SBorder)
+												.BorderImage(FAppStyle::Get().GetBrush("Menu.WidgetBorder"))
+												.Padding(FMargin(1.0f))[InSpinBoxWidget]]];
+								};
+
+								FToolMenuEntry MinDrawDistanceMenuEntry = FToolMenuEntry::InitWidget(
+									"Min Draw Distance",
+									CreateDrawDistanceWidget(MinDrawDistanceSpinBox),
+									LOCTEXT("MinDrawDistance", "Min Draw Distance")
+								);
+								Submenu->AddMenuEntry(NAME_None, MinDrawDistanceMenuEntry);
+
+								FToolMenuEntry MaxDrawDistanceMenuEntry = FToolMenuEntry::InitWidget(
+									"Max Draw Distance",
+									CreateDrawDistanceWidget(MaxDrawDistanceSpinBox),
+									LOCTEXT("MaxDrawDistance", "Max Draw Distance")
+								);
+								Submenu->AddMenuEntry(NAME_None, MaxDrawDistanceMenuEntry);
+							}
+						}
+					),
+					false,
+					FSlateIcon(FAppStyle::Get().GetStyleSetName(), "ShowFlagsMenu.SubMenu.HLODs")
+				);
+			}
+		)
+	);
+}
+
+FToolMenuEntry CreateShowLayersSubmenu()
+{
+	// This is a dynamic entry so we can skip adding the submenu if the context
+	// indicates that the viewport's world is partitioned.
+	return FToolMenuEntry::InitDynamicEntry(
+		"ShowHLODsDynamic",
+		FNewToolMenuSectionDelegate::CreateLambda(
+			[](FToolMenuSection& InDynamicSection) -> void
+			{
+				ULevelViewportContext* const LevelViewportContext = InDynamicSection.FindContext<ULevelViewportContext>();
+				if (!LevelViewportContext)
+				{
+					return;
+				}
+
+				TSharedPtr<::SLevelViewport> Viewport = LevelViewportContext->LevelViewport.Pin();
+				if (!Viewport)
+				{
+					return;
+				}
+
+				UWorld* World = Viewport->GetWorld();
+				if (!World)
+				{
+					return;
+				}
+
+				// Only add this submenu for non-partitioned worlds.
+				if (World->IsPartitionedWorld())
+				{
+					return;
+				}
+
+				InDynamicSection.AddSubMenu(
+					"ShowLayers",
+					LOCTEXT("ShowLayersMenu", "Layers"),
+					LOCTEXT("ShowLayersMenu_ToolTip", "Show layers flags"),
+					FNewToolMenuDelegate::CreateStatic(
+						&UE::LevelEditor::Private::PopulateShowLayersSubmenu, Viewport.ToWeakPtr()
+					),
+					false,
+					FSlateIcon(FAppStyle::Get().GetStyleSetName(), "ShowFlagsMenu.SubMenu.Layers")
+				);
+			}
+		)
+	);
+}
+
+FToolMenuEntry CreateShowSpritesSubmenu()
+{
+	const FLevelViewportCommands& Actions = FLevelViewportCommands::Get();
+	TArray<FLevelViewportCommands::FShowMenuCommand> ShowSpritesMenu;
+
+	// 'Show All' and 'Hide All' buttons
+	ShowSpritesMenu.Add(
+		FLevelViewportCommands::FShowMenuCommand(Actions.ShowAllSprites, LOCTEXT("ShowAllLabel", "Show All"))
+	);
+	ShowSpritesMenu.Add(
+		FLevelViewportCommands::FShowMenuCommand(Actions.HideAllSprites, LOCTEXT("HideAllLabel", "Hide All"))
+	);
+
+	// Get each show flag command and put them in their corresponding groups
+	ShowSpritesMenu += Actions.ShowSpriteCommands;
+
+	return FToolMenuEntry::InitSubMenu(
+		"ShowSprites",
+		LOCTEXT("ShowSpritesMenu", "Sprites"),
+		LOCTEXT("ShowSpritesMenu_ToolTip", "Show sprites flags"),
+		FNewToolMenuDelegate::CreateStatic(&UE::LevelEditor::Private::PopulateMenuWithCommands, ShowSpritesMenu, 2),
+		false,
+		FSlateIcon(FAppStyle::Get().GetStyleSetName(), "ShowFlagsMenu.SubMenu.Sprites")
+	);
+}
+
+FToolMenuEntry CreateShowVolumesSubmenu()
+{
+	const FLevelViewportCommands& Actions = FLevelViewportCommands::Get();
+	TArray<FLevelViewportCommands::FShowMenuCommand> ShowVolumesMenu;
+
+	// 'Show All' and 'Hide All' buttons
+	ShowVolumesMenu.Add(
+		FLevelViewportCommands::FShowMenuCommand(Actions.ShowAllVolumes, LOCTEXT("ShowAllLabel", "Show All"))
+	);
+	ShowVolumesMenu.Add(
+		FLevelViewportCommands::FShowMenuCommand(Actions.HideAllVolumes, LOCTEXT("HideAllLabel", "Hide All"))
+	);
+
+	// Get each show flag command and put them in their corresponding groups
+	ShowVolumesMenu += Actions.ShowVolumeCommands;
+
+	return FToolMenuEntry::InitSubMenu(
+		"ShowVolumes",
+		LOCTEXT("ShowVolumesMenu", "Volumes"),
+		LOCTEXT("ShowVolumesMenu_ToolTip", "Show volumes flags"),
+		FNewToolMenuDelegate::CreateStatic(&UE::LevelEditor::Private::PopulateMenuWithCommands, ShowVolumesMenu, 2),
+		false,
+		FSlateIcon(FAppStyle::Get().GetStyleSetName(), "ShowFlagsMenu.SubMenu.Volumes")
+	);
+}
+
+#if STATS
+FToolMenuEntry CreateShowStatsSubmenu()
+{
+	return FToolMenuEntry::InitSubMenu(
+		"ShowStatsMenu",
+		LOCTEXT("ShowStatsMenu", "Stat"),
+		LOCTEXT("ShowStatsMenu_ToolTip", "Show Stat commands"),
+		FNewToolMenuDelegate::CreateLambda(
+			[](UToolMenu* InMenu) -> void
+			{
+				TArray<FLevelViewportCommands::FShowMenuCommand> HideStatsMenu;
+				HideStatsMenu.Add(FLevelViewportCommands::FShowMenuCommand(
+					FLevelViewportCommands::Get().HideAllStats, LOCTEXT("HideAllLabel", "Hide All")
+				));
+
+				UE::LevelEditor::Private::PopulateMenuWithCommands(InMenu, HideStatsMenu, 1);
+
+				FToolMenuSection& Section = InMenu->FindOrAddSection("Section");
+
+				// Separate out stats into two list, those with and without submenus
+				TArray<FLevelViewportCommands::FShowMenuCommand> SingleStatCommands;
+				TMap<FString, TArray<FLevelViewportCommands::FShowMenuCommand>> SubbedStatCommands;
+				for (auto StatCatIt = FLevelViewportCommands::Get().ShowStatCatCommands.CreateConstIterator(); StatCatIt;
+					 ++StatCatIt)
+				{
+					const TArray<FLevelViewportCommands::FShowMenuCommand>& ShowStatCommands = StatCatIt.Value();
+					const FString& CategoryName = StatCatIt.Key();
+
+					// If no category is specified, or there's only one category, don't use submenus
+					FString NoCategory = FStatConstants::NAME_NoCategory.ToString();
+					NoCategory.RemoveFromStart(TEXT("STATCAT_"));
+					if (CategoryName == NoCategory || FLevelViewportCommands::Get().ShowStatCatCommands.Num() == 1)
+					{
+						for (int32 StatIndex = 0; StatIndex < ShowStatCommands.Num(); ++StatIndex)
+						{
+							const FLevelViewportCommands::FShowMenuCommand& StatCommand = ShowStatCommands[StatIndex];
+							SingleStatCommands.Add(StatCommand);
+						}
+					}
+					else
+					{
+						SubbedStatCommands.Add(CategoryName, ShowStatCommands);
+					}
+				}
+
+				// First add all the stats that don't have a sub menu
+				for (auto StatCatIt = SingleStatCommands.CreateConstIterator(); StatCatIt; ++StatCatIt)
+				{
+					const FLevelViewportCommands::FShowMenuCommand& StatCommand = *StatCatIt;
+					Section.AddMenuEntry(NAME_None, StatCommand.ShowMenuItem, StatCommand.LabelOverride);
+				}
+
+				// Now add all the stats that have sub menus
+				for (auto StatCatIt = SubbedStatCommands.CreateConstIterator(); StatCatIt; ++StatCatIt)
+				{
+					const TArray<FLevelViewportCommands::FShowMenuCommand>& StatCommands = StatCatIt.Value();
+					const FText CategoryName = FText::FromString(StatCatIt.Key());
+
+					FFormatNamedArguments Args;
+					Args.Add(TEXT("StatCat"), CategoryName);
+					const FText CategoryDescription =
+						FText::Format(NSLOCTEXT("UICommands", "StatShowCatName", "Show {StatCat} stats"), Args);
+
+					Section.AddSubMenu(
+						NAME_None,
+						CategoryName,
+						CategoryDescription,
+						FNewToolMenuDelegate::CreateStatic(&UE::LevelEditor::Private::PopulateMenuWithCommands, StatCommands, 0)
+					);
+				}
+			}
+		),
+		false,
+		FSlateIcon(FAppStyle::Get().GetStyleSetName(), "EditorViewport.SubMenu.Stats")
+	);
+}
+#endif
+
+FToolMenuEntry CreateViewportToolbarShowSubmenu()
+{
+	return FToolMenuEntry::InitSubMenu(
+		"Show",
+		LOCTEXT("ShowSubmenuLabel", "Show"),
+		LOCTEXT("ShowSubmenuTooltip", "Show flags related to the current viewport"),
+		FNewToolMenuDelegate::CreateLambda(
+			[](UToolMenu* InMenu) -> void
+			{
+				{
+					FToolMenuSection& UnnamedSection = InMenu->FindOrAddSection(NAME_None);
+
+					UnnamedSection.AddMenuEntry(FLevelViewportCommands::Get().UseDefaultShowFlags);
+
+					UnnamedSection.AddSeparator("ViewportStatsSeparator");
+
+					// Override the label of the stats submenu for the new viewport toolbar.
+					{
+						FToolMenuEntry StatsSubmenu = UE::LevelEditor::CreateShowStatsSubmenu();
+						StatsSubmenu.Label = LOCTEXT("ViewportStatsLabel", "Viewport Stats");
+						UnnamedSection.AddEntry(StatsSubmenu);
+					}
+				}
+
+				{
+					FToolMenuSection& CommonShowFlagsSection =
+						InMenu->FindOrAddSection("CommonShowFlags", LOCTEXT("CommonShowFlagsLabel", "Common Show Flags"));
+
+					FShowFlagMenuCommands::Get().PopulateCommonShowFlagsSection(CommonShowFlagsSection);
+				}
+
+				{
+					FToolMenuSection& AllShowFlagsSection =
+						InMenu->FindOrAddSection("AllShowFlags", LOCTEXT("AllShowFlagsLabel", "All Show Flags"));
+
+					{
+						FToolMenuEntry ShowFoliageSubmenu = CreateShowFoliageSubmenu();
+						ShowFoliageSubmenu.Label = LOCTEXT("ShowFoliageLabel", "Foliage");
+						AllShowFlagsSection.AddEntry(ShowFoliageSubmenu);
+					}
+
+					AllShowFlagsSection.AddEntry(CreateShowHLODsSubmenu());
+					AllShowFlagsSection.AddEntry(CreateShowLayersSubmenu());
+					AllShowFlagsSection.AddEntry(CreateShowSpritesSubmenu());
+					AllShowFlagsSection.AddEntry(CreateShowVolumesSubmenu());
+
+					FShowFlagMenuCommands::Get().PopulateAllShowFlagsSection(AllShowFlagsSection);
+				}
+
+				// Create these sections for backward compatibility with the old viewport toolbar.
+				{
+					// If your entries end up in this section, you should move it to the new "CommonShowFlags" section instead.
+					InMenu->FindOrAddSection(
+						"ShowFlagsMenuSectionCommon",
+						LOCTEXT("ShowFlagsMenuSectionCommonLabel", "Common Show Flags (Deprecated section)")
+					);
+
+					// If your entries end up in these sections, you should move them to the above "AllShowFlags" section instead.
+					InMenu->FindOrAddSection(
+						"LevelViewportShowFlags",
+						LOCTEXT("LevelViewportShowFlagsLabel", "All Show Flags (Deprecated section)")
+					);
+					InMenu->FindOrAddSection(
+						"LevelViewportEditorShow", LOCTEXT("LevelViewportEditorShowLabel", "Editor (Deprecated section)")
+					);
+				}
+			}
+		)
+	);
+}
+
 FToolMenuEntry CreateFeatureLevelPreviewSubmenu()
 {
 	return FToolMenuEntry::InitSubMenu(
diff --git a/Engine/Source/Editor/LevelEditor/Private/ViewportToolbar/LevelEditorViewportToolbarSections.h b/Engine/Source/Editor/LevelEditor/Private/ViewportToolbar/LevelEditorViewportToolbarSections.h
--- a/Engine/Source/Editor/LevelEditor/Private/ViewportToolbar/LevelEditorViewportToolbarSections.h
+++ b/Engine/Source/Editor/LevelEditor/Private/ViewportToolbar/LevelEditorViewportToolbarSections.h
@@ -22,9 +22,19 @@ FToolMenuEntry CreateViewportToolbarTransformsSection();
 FToolMenuEntry CreateViewportToolbarSelectionSection();
 
 TSharedPtr<FExtender> GetViewModesLegacyExtenders();
-void PopulateViewModesMenu(UToolMenu* InMenu, TSharedRef<SLevelViewport> InViewport);
+void PopulateViewModesMenu(UToolMenu* InMenu, TSharedRef<::SLevelViewport> InViewport);
 FToolMenuEntry CreateViewportToolbarViewModesSubmenu();
 
+FToolMenuEntry CreateShowFoliageSubmenu();
+FToolMenuEntry CreateShowHLODsSubmenu();
+FToolMenuEntry CreateShowLayersSubmenu();
+FToolMenuEntry CreateShowSpritesSubmenu();
+#if STATS
+FToolMenuEntry CreateShowStatsSubmenu();
+#endif
+FToolMenuEntry CreateShowVolumesSubmenu();
+FToolMenuEntry CreateViewportToolbarShowSubmenu();
+
 FToolMenuEntry CreateFeatureLevelPreviewSubmenu();
 FToolMenuEntry CreateMaterialQualityLevelSubmenu();
 FToolMenuEntry CreateViewportToolbarPerformanceAndScalabilitySubmenu();
