Wrap call to UE::LevelEditor::CreateShowStatsSubmenu() in #if STATS

diff --git a/Engine/Source/Editor/LevelEditor/Private/ViewportToolbar/LevelEditorViewportToolbarSections.cpp b/Engine/Source/Editor/LevelEditor/Private/ViewportToolbar/LevelEditorViewportToolbarSections.cpp
--- a/Engine/Source/Editor/LevelEditor/Private/ViewportToolbar/LevelEditorViewportToolbarSections.cpp
+++ b/Engine/Source/Editor/LevelEditor/Private/ViewportToolbar/LevelEditorViewportToolbarSections.cpp
@@ -1273,12 +1273,14 @@ FToolMenuEntry CreateViewportToolbarShowSubmenu()
 
 					UnnamedSection.AddSeparator("ViewportStatsSeparator");
 
+#if STATS
 					// Override the label of the stats submenu for the new viewport toolbar.
 					{
 						FToolMenuEntry StatsSubmenu = UE::LevelEditor::CreateShowStatsSubmenu();
 						StatsSubmenu.Label = LOCTEXT("ViewportStatsLabel", "Viewport Stats");
 						UnnamedSection.AddEntry(StatsSubmenu);
 					}
+#endif
 				}
 
 				{
