Improve styling of the new viewport toolbar by using a SlimHorizontalToolBar

diff --git a/Engine/Source/Editor/LevelEditor/Private/SLevelViewport.cpp b/Engine/Source/Editor/LevelEditor/Private/SLevelViewport.cpp
--- a/Engine/Source/Editor/LevelEditor/Private/SLevelViewport.cpp
+++ b/Engine/Source/Editor/LevelEditor/Private/SLevelViewport.cpp
@@ -1926,7 +1926,8 @@ TSharedPtr<SWidget> SLevelViewport::MakeViewportToolbar()
 		if (!UToolMenus::Get()->IsMenuRegistered(LevelEditorViewportToolbarName))
 		{
 			UToolMenu* const ViewportToolbarMenu = UToolMenus::Get()->RegisterMenu(
-				LevelEditorViewportToolbarName, NAME_None /* parent */, EMultiBoxType::ToolBar);
+				LevelEditorViewportToolbarName, NAME_None /* parent */, EMultiBoxType::SlimHorizontalToolBar
+			);
 
 			// Add the left-aligned part of the viewport toolbar.
 			{
