[Viewport Toolbar] Create the "Snapping" submenu for the new level editor viewport toolbar

Adds the Snapping submenu, with the following voices:

- Surface (with submenu)
- Grid
- Rotation
- Scale
- Actor (with submenu)
- Socket
- Vertex

Still missing:
- submenus for Grid, Rotation and Scale
- Planar snapping still missing

Actor snapping distance widget has some performance issues which still need to be addressed.

diff --git a/Engine/Source/Editor/LevelEditor/Private/SLevelViewport.cpp b/Engine/Source/Editor/LevelEditor/Private/SLevelViewport.cpp
--- a/Engine/Source/Editor/LevelEditor/Private/SLevelViewport.cpp
+++ b/Engine/Source/Editor/LevelEditor/Private/SLevelViewport.cpp
@@ -1946,6 +1946,13 @@ TSharedPtr<SWidget> SLevelViewport::MakeViewportToolbar()
 					SelectionSubmenu.InsertPosition.Position = EToolMenuInsertType::First;
 					LeftSection.AddEntry(SelectionSubmenu);
 				}
+
+				// Add the "Snapping" sub menu.
+				{
+					FToolMenuEntry SnappingSubmenu = UE::LevelEditor::CreateViewportToolbarSnappingSubmenu();
+					SnappingSubmenu.InsertPosition.Position = EToolMenuInsertType::First;
+					LeftSection.AddEntry(SnappingSubmenu);
+				}
 			}
 
 			// Add the right-aligned part of the viewport toolbar.
diff --git a/Engine/Source/Editor/LevelEditor/Private/ViewportToolbar/LevelEditorViewportToolbarSections.cpp b/Engine/Source/Editor/LevelEditor/Private/ViewportToolbar/LevelEditorViewportToolbarSections.cpp
--- a/Engine/Source/Editor/LevelEditor/Private/ViewportToolbar/LevelEditorViewportToolbarSections.cpp
+++ b/Engine/Source/Editor/LevelEditor/Private/ViewportToolbar/LevelEditorViewportToolbarSections.cpp
@@ -21,9 +21,11 @@
 #include "ToolMenus.h"
 #include "ViewportToolbar/LevelViewportContext.h"
 #include "ViewportToolbar/UnrealEdViewportToolbar.h"
+#include "Widgets/Input/SNumericEntryBox.h"
 #include "Widgets/Input/SSpinBox.h"
 #include "Widgets/Input/SVolumeControl.h"
 #include "Widgets/SBoxPanel.h"
+#include "Widgets/SToolTip.h"
 #include "WorldPartition/IWorldPartitionEditorModule.h"
 
 #define LOCTEXT_NAMESPACE "LevelEditorViewportToolbar"
@@ -127,6 +129,208 @@ void PopulateShowLayersSubmenu(UToolMenu* InMenu, TWeakPtr<::SLevelViewport> InV
 	}
 }
 
+// TODO: Maybe export CreateSurfaceSnapOffsetEntry function, so that it can be used elsewhere, e.g. STransformViewportToolbar.cpp
+FToolMenuEntry CreateSurfaceSnapOffsetEntry()
+{
+	FText Label = LOCTEXT("SurfaceOffsetLabel", "Surface Offset");
+	FText Tooltip = LOCTEXT("SurfaceOffsetTooltip", "The amount of offset to apply when snapping to surfaces");
+
+	FToolMenuEntry SurfaceOffset = FToolMenuEntry::InitMenuEntry(
+		"SurfaceOffset",
+		Label,
+		Tooltip,
+		FSlateIcon(),
+		FUIAction());
+
+	const FMargin WidgetsMargin(8.0f, 0.0f, 0.0f, 0.0f);
+
+	SurfaceOffset.MakeCustomWidget.BindLambda(
+		[Label, Tooltip, WidgetsMargin](const FToolMenuContext& InContext, const FToolMenuCustomWidgetContext& InWidgetContext)
+		{
+			// clang-format on
+			return SNew(SHorizontalBox).IsEnabled_Lambda([]()
+				{
+					return GetDefault<ULevelEditorViewportSettings>()->SnapToSurface.bEnabled;
+				})
+				+ SHorizontalBox::Slot().VAlign(VAlign_Center).Padding(WidgetsMargin).AutoWidth()
+				[
+					SNew(STextBlock).Text(Label).TextStyle(InWidgetContext.StyleSet,
+					                                       ISlateStyle::Join(InWidgetContext.StyleName, ".Label"))
+				]
+				+ SHorizontalBox::Slot().VAlign(VAlign_Center).Padding(WidgetsMargin).AutoWidth()
+				[
+					SNew(SBox).Padding(WidgetsMargin).MinDesiredWidth(100.0f)
+					[
+						// Min/Max/Slider values taken from STransformViewportToolbar.cpp
+						SNew(SNumericEntryBox<float>)
+						.ToolTipText(Tooltip)
+						.MinValue(0.0f)
+						.MaxValue(static_cast<float>(HALF_WORLD_MAX))
+						.MaxSliderValue(1000.0f)
+						.AllowSpin(true)
+						.MaxFractionalDigits(2)
+						.Font(FAppStyle::GetFontStyle(TEXT("MenuItem.Font")))
+						.OnValueChanged_Lambda([](float InNewValue)
+						{
+							auto* Settings = GetMutableDefault<ULevelEditorViewportSettings>();
+							Settings->SnapToSurface.SnapOffsetExtent = InNewValue;
+						})
+						.Value_Lambda([]()
+						{
+							return GetDefault<ULevelEditorViewportSettings>()->SnapToSurface.SnapOffsetExtent;
+						})
+					]
+				];
+			// clang-format off
+		});
+
+	return SurfaceOffset;
+}
+
+FToolMenuEntry CreateSurfaceSnapMenu()
+{
+	FNewToolMenuDelegate MakeMenuDelegate = FNewToolMenuDelegate::CreateLambda([](UToolMenu* Submenu)
+	{
+		FToolMenuSection& SurfaceSnappingSection = Submenu->FindOrAddSection("SurfaceSnapping", LOCTEXT("SurfaceSnappingLabel", "Surface Snapping"));
+
+		// Add "Rotate to surface normal" checkbox.
+		{
+			FToolMenuEntry RotateToSurfaceNormalSnapping = FToolMenuEntry::InitMenuEntry(FEditorViewportCommands::Get().RotateToSurfaceNormal);
+			SurfaceSnappingSection.AddEntry(RotateToSurfaceNormalSnapping);
+		}
+
+		// Add "Surface offset" widget.
+		{
+			SurfaceSnappingSection.AddEntry(CreateSurfaceSnapOffsetEntry());
+		}
+	});
+
+	FToolUIAction CheckboxMenuAction;
+	{
+		CheckboxMenuAction.ExecuteAction = FToolMenuExecuteAction::CreateLambda(
+			[](const FToolMenuContext& InContext)
+			{
+				auto* Settings = GetMutableDefault<ULevelEditorViewportSettings>();
+				Settings->SnapToSurface.bEnabled = !Settings->SnapToSurface.bEnabled;
+			}
+		);
+		CheckboxMenuAction.GetActionCheckState = FToolMenuGetActionCheckState::CreateLambda(
+			[](const FToolMenuContext& InContext)
+			{
+				return GetDefault<ULevelEditorViewportSettings>()->SnapToSurface.bEnabled ? ECheckBoxState::Checked : ECheckBoxState::Unchecked;
+			}
+		);
+	}
+
+	FToolMenuEntry SurfaceSnapping = FToolMenuEntry::InitSubMenu(
+		"SurfaceSnapping",
+		LOCTEXT("SurfaceSnapLabel", "Surface"),
+		FEditorViewportCommands::Get().SurfaceSnapping->MakeTooltip()->GetTextTooltip(),
+		MakeMenuDelegate,
+		CheckboxMenuAction,
+		EUserInterfaceActionType::ToggleButton
+	);
+
+	return SurfaceSnapping;
+}
+
+FToolMenuEntry CreateActorSnapDistanceEntry()
+{
+	const FText Label = LOCTEXT("ActorSnapDistanceLabel", "Snap Distance");
+	const FText Tooltip = LOCTEXT("ActorSnapDistanceTooltip", "The amount of offset to apply when snapping to surfaces");
+
+	FToolMenuEntry SnapDistance = FToolMenuEntry::InitMenuEntry(
+		"ActorSnapDistance",
+		Label,
+		Tooltip,
+		FSlateIcon(),
+		FUIAction());
+
+	const FMargin WidgetsMargin(8.0f, 0.0f, 0.0f, 0.0f);
+
+	SnapDistance.MakeCustomWidget.BindLambda(
+		[Label, Tooltip, WidgetsMargin](const FToolMenuContext& InContext, const FToolMenuCustomWidgetContext& InWidgetContext)
+		{
+			// clang-format on
+			return SNew(SHorizontalBox).IsEnabled_Lambda([]()
+				{
+					return !!GetDefault<ULevelEditorViewportSettings>()->bEnableActorSnap;
+				})
+				+ SHorizontalBox::Slot().VAlign(VAlign_Center).Padding(WidgetsMargin).AutoWidth()
+				[
+					SNew(STextBlock).Text(Label).TextStyle(InWidgetContext.StyleSet,
+					                                       ISlateStyle::Join(InWidgetContext.StyleName, ".Label"))
+				]
+				+ SHorizontalBox::Slot().VAlign(VAlign_Center).Padding(WidgetsMargin).AutoWidth()
+				[
+					SNew(SBox).Padding(WidgetsMargin).MinDesiredWidth(100.0f)
+					[
+						// TODO: Check how to improve performance for this widget OnValueChanged.
+						// Same functionality in LevelEditorToolBar.cpp seems to have better performance
+						SNew(SNumericEntryBox<float>)
+						.ToolTipText(Tooltip)
+						.MinValue(0.0f)
+						.MaxValue(1.0f)
+						.MaxSliderValue(1.0f)
+						.AllowSpin(true)
+						.MaxFractionalDigits(1)
+						.Font(FAppStyle::GetFontStyle(TEXT("MenuItem.Font")))
+						.OnValueChanged_Static(&FLevelEditorActionCallbacks::SetActorSnapSetting)
+						.Value_Lambda([]()
+						{
+							return FLevelEditorActionCallbacks::GetActorSnapSetting();
+						})
+					]
+				];
+			// clang-format off
+		});
+
+	return SnapDistance;
+}
+
+FToolMenuEntry CreateActorSnapMenu()
+{
+	FNewToolMenuDelegate MakeMenuDelegate = FNewToolMenuDelegate::CreateLambda([](UToolMenu* Submenu)
+	{
+		FToolMenuSection& ActorSnappingSection = Submenu->FindOrAddSection("ActorSnapping", LOCTEXT("ActorSnappingLabel", "Actor Snapping"));
+
+		// Add "Actor snapping" widget.
+		{
+			ActorSnappingSection.AddEntry(CreateActorSnapDistanceEntry());
+		}
+	});
+
+	FToolUIAction CheckboxMenuAction;
+	{
+		CheckboxMenuAction.ExecuteAction = FToolMenuExecuteAction::CreateLambda(
+			[](const FToolMenuContext& InContext)
+			{
+				if (ULevelEditorViewportSettings* Settings = GetMutableDefault<ULevelEditorViewportSettings>())
+				{
+					Settings->bEnableActorSnap = !Settings->bEnableActorSnap;
+				}
+			}
+		);
+		CheckboxMenuAction.GetActionCheckState = FToolMenuGetActionCheckState::CreateLambda(
+			[](const FToolMenuContext& InContext)
+			{
+				return GetDefault<ULevelEditorViewportSettings>()->bEnableActorSnap ? ECheckBoxState::Checked : ECheckBoxState::Unchecked;
+			}
+		);
+	}
+
+	FToolMenuEntry ActorSnapping = FToolMenuEntry::InitSubMenu(
+		"ActorSnapping",
+		LOCTEXT("ActorSnapLabel", "Actor"),
+		FLevelEditorCommands::Get().EnableActorSnap->MakeTooltip()->GetTextTooltip(),
+		MakeMenuDelegate,
+		CheckboxMenuAction,
+		EUserInterfaceActionType::ToggleButton
+	);
+
+	return ActorSnapping;
+}
+
 } // namespace UE::LevelEditor::Private
 
 namespace UE::LevelEditor
@@ -1844,6 +2048,54 @@ FToolMenuEntry CreateLevelEditorViewportToolbarSettingsSubmenu()
 	);
 }
 
+FToolMenuEntry CreateViewportToolbarSnappingSubmenu()
+{
+	return FToolMenuEntry::InitSubMenu(
+		"Snapping",
+		LOCTEXT("SnappingSubmenuLabel", "Snapping"),
+		LOCTEXT("SnappingSubmenuTooltip", "Viewport-related snapping settings"),
+		FNewToolMenuDelegate::CreateLambda(
+			[](UToolMenu* Submenu) -> void
+			{
+				FToolMenuSection& SnappingSection =
+					Submenu->FindOrAddSection("Snapping", LOCTEXT("SnappingLabel", "Snapping"));
+
+				SnappingSection.AddEntry(Private::CreateSurfaceSnapMenu());
+
+				FToolMenuEntry GridSnapping = FToolMenuEntry::InitMenuEntry(FEditorViewportCommands::Get().LocationGridSnap);
+				GridSnapping.UserInterfaceActionType = EUserInterfaceActionType::ToggleButton;
+				GridSnapping.Label = LOCTEXT("GridSnapLabel", "Grid");
+				GridSnapping.SetShowInToolbarTopLevel(true);
+				SnappingSection.AddEntry(GridSnapping);
+
+				FToolMenuEntry RotationSnapping = FToolMenuEntry::InitMenuEntry(FEditorViewportCommands::Get().RotationGridSnap);
+				RotationSnapping.UserInterfaceActionType = EUserInterfaceActionType::ToggleButton;
+				RotationSnapping.Label = LOCTEXT("RotationSnapLabel", "Rotation");
+				SnappingSection.AddEntry(RotationSnapping);
+
+				FToolMenuEntry ScaleSnapping = FToolMenuEntry::InitMenuEntry(FEditorViewportCommands::Get().ScaleGridSnap);
+				ScaleSnapping.UserInterfaceActionType = EUserInterfaceActionType::ToggleButton;
+				ScaleSnapping.Label = LOCTEXT("ScaleSnapLabel", "Scale");
+				SnappingSection.AddEntry(ScaleSnapping);
+
+				SnappingSection.AddEntry(Private::CreateActorSnapMenu());
+
+				FToolMenuEntry SocketSnapping = FToolMenuEntry::InitMenuEntry(FLevelEditorCommands::Get().ToggleSocketSnapping);
+				SocketSnapping.UserInterfaceActionType = EUserInterfaceActionType::ToggleButton;
+				SocketSnapping.Label = LOCTEXT("SocketSnapLabel", "Socket");
+				SnappingSection.AddEntry(SocketSnapping);
+
+				FToolMenuEntry VertexSnapping = FToolMenuEntry::InitMenuEntry(FLevelEditorCommands::Get().EnableVertexSnap);
+				VertexSnapping.UserInterfaceActionType = EUserInterfaceActionType::ToggleButton;
+				VertexSnapping.Label = LOCTEXT("VertexSnapLabel", "Vertex");
+				SnappingSection.AddEntry(VertexSnapping);
+
+				// TODO: add Planar Snapping
+			}
+		)
+	);
+}
+
 } // namespace UE::LevelEditor
 
 #undef LOCTEXT_NAMESPACE
diff --git a/Engine/Source/Editor/LevelEditor/Private/ViewportToolbar/LevelEditorViewportToolbarSections.h b/Engine/Source/Editor/LevelEditor/Private/ViewportToolbar/LevelEditorViewportToolbarSections.h
--- a/Engine/Source/Editor/LevelEditor/Private/ViewportToolbar/LevelEditorViewportToolbarSections.h
+++ b/Engine/Source/Editor/LevelEditor/Private/ViewportToolbar/LevelEditorViewportToolbarSections.h
@@ -43,4 +43,6 @@ void GenerateViewportLayoutsMenu(UToolMenu* InMenu, TSharedPtr<::SLevelViewport>
 TSharedRef<SWidget> BuildVolumeControlCustomWidget();
 FToolMenuEntry CreateLevelEditorViewportToolbarSettingsSubmenu();
 
+FToolMenuEntry CreateViewportToolbarSnappingSubmenu();
+
 } // namespace UE::LevelEditor
diff --git a/Engine/Source/Editor/UnrealEd/Private/EditorViewportCommands.cpp b/Engine/Source/Editor/UnrealEd/Private/EditorViewportCommands.cpp
--- a/Engine/Source/Editor/UnrealEd/Private/EditorViewportCommands.cpp
+++ b/Engine/Source/Editor/UnrealEd/Private/EditorViewportCommands.cpp
@@ -174,6 +174,7 @@ void FEditorViewportCommands::RegisterCommands()
 	UI_COMMAND( Layer2DSnap, "Layer2D Snap", "Enables or disables snapping objects to a 2D layer", EUserInterfaceActionType::ToggleButton, FInputChord());
 	UI_COMMAND( ScaleGridSnap, "Scale Snap", "Enables or disables snapping objects to a scale grid", EUserInterfaceActionType::ToggleButton, FInputChord() );
 	UI_COMMAND( SurfaceSnapping, "Surface Snapping", "If enabled, actors will snap to surfaces in the world when dragging", EUserInterfaceActionType::ToggleButton, FInputChord() );
+	UI_COMMAND( RotateToSurfaceNormal, "Rotate to Surface Normal", "If enabled, actors rotation will match that of their snapping surface", EUserInterfaceActionType::ToggleButton, FInputChord() );
 
 	UI_COMMAND( ToggleAutoExposure, "Auto", "If enabled, enables automatic exposure", EUserInterfaceActionType::ToggleButton, FInputChord() );
 	UI_COMMAND( ToggleInGameExposure, "Game Settings", "If enabled, uses game settings", EUserInterfaceActionType::ToggleButton, FInputChord() );
diff --git a/Engine/Source/Editor/UnrealEd/Private/SEditorViewport.cpp b/Engine/Source/Editor/UnrealEd/Private/SEditorViewport.cpp
--- a/Engine/Source/Editor/UnrealEd/Private/SEditorViewport.cpp
+++ b/Engine/Source/Editor/UnrealEd/Private/SEditorViewport.cpp
@@ -370,6 +370,13 @@ void SEditorViewport::BindCommands()
 		FIsActionChecked::CreateStatic( &SEditorViewport::OnIsSurfaceSnapEnabled ) 
 		);
 
+	CommandListRef.MapAction(
+		Commands.RotateToSurfaceNormal,
+		FExecuteAction::CreateStatic( &SEditorViewport::OnToggleRotateToSurfaceNormal ),
+		FCanExecuteAction::CreateStatic( &SEditorViewport::OnIsSurfaceSnapEnabled ),
+		FIsActionChecked::CreateStatic( &SEditorViewport::IsRotateToSurfaceNormalEnabled ) 
+		);
+
 	CommandListRef.MapAction(
 		(Client.IsValid() && Client->IsLevelEditorClient()) ? Commands.ToggleInGameExposure : Commands.ToggleAutoExposure,
 		FExecuteAction::CreateSP( this, &SEditorViewport::ChangeExposureSetting),
@@ -850,6 +857,18 @@ bool SEditorViewport::OnIsSurfaceSnapEnabled()
 	return GetDefault<ULevelEditorViewportSettings>()->SnapToSurface.bEnabled;
 }
 
+void SEditorViewport::OnToggleRotateToSurfaceNormal()
+{
+	auto& Settings = GetMutableDefault<ULevelEditorViewportSettings>()->SnapToSurface;
+	Settings.bSnapRotation = !Settings.bSnapRotation;
+}
+
+bool SEditorViewport::IsRotateToSurfaceNormalEnabled()
+{
+	const auto& Settings = GetDefault<ULevelEditorViewportSettings>()->SnapToSurface;
+	return Settings.bSnapRotation;
+}
+
 bool SEditorViewport::IsPreviewingScreenPercentage() const
 {
 	return Client->IsPreviewingScreenPercentage();
diff --git a/Engine/Source/Editor/UnrealEd/Private/STransformViewportToolbar.cpp b/Engine/Source/Editor/UnrealEd/Private/STransformViewportToolbar.cpp
--- a/Engine/Source/Editor/UnrealEd/Private/STransformViewportToolbar.cpp
+++ b/Engine/Source/Editor/UnrealEd/Private/STransformViewportToolbar.cpp
@@ -80,23 +80,7 @@ TSharedRef<SWidget> STransformViewportToolBar::GenerateSurfaceSnappingMenu()
 
 	MenuBuilder.BeginSection("SurfaceSnappingSettings", LOCTEXT("SnapToSurfaceSettings", "Settings"));
 	{
-		MenuBuilder.AddMenuEntry(
-			LOCTEXT("SnapToSurfaceSettings_Rotation", "Rotate to Surface Normal"),
-			LOCTEXT("SnapToSurfaceSettings_RotationTip", "When checked, snapping an object to a surface will also rotate the object to align to the surface normal"),
-			FSlateIcon(),
-			FUIAction(
-				FExecuteAction::CreateStatic([] {
-					auto& Settings = GetMutableDefault<ULevelEditorViewportSettings>()->SnapToSurface;
-					Settings.bSnapRotation = !Settings.bSnapRotation;
-					}),
-				FCanExecuteAction::CreateStatic(IsSnappingEnabled),
-				FIsActionChecked::CreateStatic([] {
-					const auto& Settings = GetDefault<ULevelEditorViewportSettings>()->SnapToSurface;
-					return Settings.bSnapRotation;
-					})),
-				NAME_None,
-			EUserInterfaceActionType::Check);
-
+		MenuBuilder.AddMenuEntry(FEditorViewportCommands::Get().RotateToSurfaceNormal);
 
 		MenuBuilder.AddWidget(
 
diff --git a/Engine/Source/Editor/UnrealEd/Public/EditorViewportCommands.h b/Engine/Source/Editor/UnrealEd/Public/EditorViewportCommands.h
--- a/Engine/Source/Editor/UnrealEd/Public/EditorViewportCommands.h
+++ b/Engine/Source/Editor/UnrealEd/Public/EditorViewportCommands.h
@@ -251,6 +251,10 @@ public:
 
 	/** Enables or disables snapping to surfaces when dragging in the viewport */
 	TSharedPtr< FUICommandInfo > SurfaceSnapping;
+
+	/** Enables or disables rotating to surface normal when snapping to surfaces */
+	TSharedPtr< FUICommandInfo > RotateToSurfaceNormal;
+
 public:
 	/** Registers our commands with the binding system */
 	virtual void RegisterCommands() override;
diff --git a/Engine/Source/Editor/UnrealEd/Public/SEditorViewport.h b/Engine/Source/Editor/UnrealEd/Public/SEditorViewport.h
--- a/Engine/Source/Editor/UnrealEd/Public/SEditorViewport.h
+++ b/Engine/Source/Editor/UnrealEd/Public/SEditorViewport.h
@@ -260,6 +260,15 @@ protected:
 	 */
 	static UNREALED_API bool OnIsSurfaceSnapEnabled();
 
+	/**
+	 * Called when rotation matching surface normal has been enabled/disabled
+	 */
+	static UNREALED_API void OnToggleRotateToSurfaceNormal();
+
+	/**
+	 * Called to test whether rotation matching surface normal is enabled or not
+	 */
+	static UNREALED_API bool IsRotateToSurfaceNormalEnabled();
 
 protected:
 	TSharedPtr<SOverlay> ViewportOverlay;
