Improve backward compatibility of the new viewport toolbar View Modes submenu

Use bools to track the registered state of its parent and grandparent menus. This enables the new viewport toolbar to create those menus if they don't exist, avoiding warnings at runtime, while at the same time not breaking the old menus where they previously wouldn't run some setup code if the menu was already registered.

diff --git a/Engine/Source/Editor/LevelEditor/Private/SLevelViewport.cpp b/Engine/Source/Editor/LevelEditor/Private/SLevelViewport.cpp
--- a/Engine/Source/Editor/LevelEditor/Private/SLevelViewport.cpp
+++ b/Engine/Source/Editor/LevelEditor/Private/SLevelViewport.cpp
@@ -1965,9 +1965,27 @@ TSharedPtr<SWidget> SLevelViewport::MakeViewportToolbar()
 				// Add the "View Modes" sub menu.
 				{
 					// Stay backward-compatible with the old viewport toolbar.
-					UToolMenus::Get()->RegisterMenu(
-						"LevelEditor.ViewportToolbar.ViewModes", "LevelEditor.LevelViewportToolbar.View"
-					);
+					{
+						// Create our grandparent menu.
+						if (!UToolMenus::Get()->IsMenuRegistered("UnrealEd.ViewportToolbar.View"))
+						{
+							UToolMenus::Get()->RegisterMenu("UnrealEd.ViewportToolbar.View");
+						}
+
+						// Create our parent menu.
+						if (!UToolMenus::Get()->IsMenuRegistered("LevelEditor.LevelViewportToolbar.View"))
+						{
+							UToolMenus::Get()->RegisterMenu(
+								"LevelEditor.LevelViewportToolbar.View", "UnrealEd.ViewportToolbar.View"
+							);
+						}
+
+						// Create our menu.
+						UToolMenus::Get()->RegisterMenu(
+							"LevelEditor.ViewportToolbar.ViewModes", "LevelEditor.LevelViewportToolbar.View"
+						);
+					}
+
 					FToolMenuEntry ViewModesSubmenu = UE::LevelEditor::CreateViewportToolbarViewModesSubmenu();
 					ViewModesSubmenu.InsertPosition.Position = EToolMenuInsertType::Last;
 					RightSection.AddEntry(ViewModesSubmenu);
@@ -1985,6 +2003,7 @@ TSharedPtr<SWidget> SLevelViewport::MakeViewportToolbar()
 							"LevelEditor.ViewportToolbar.Show", "LevelEditor.LevelViewportToolbar.Show"
 						);
 					}
+
 					FToolMenuEntry ShowSubmenu = UE::LevelEditor::CreateViewportToolbarShowSubmenu();
 					ShowSubmenu.InsertPosition.Position = EToolMenuInsertType::Last;
 					RightSection.AddEntry(ShowSubmenu);
diff --git a/Engine/Source/Editor/LevelEditor/Private/SLevelViewportToolBar.cpp b/Engine/Source/Editor/LevelEditor/Private/SLevelViewportToolBar.cpp
--- a/Engine/Source/Editor/LevelEditor/Private/SLevelViewportToolBar.cpp
+++ b/Engine/Source/Editor/LevelEditor/Private/SLevelViewportToolBar.cpp
@@ -77,9 +77,19 @@ public:
 	{
 		SEditorViewportViewMenu::RegisterMenus();
 
-		if (!UToolMenus::Get()->IsMenuRegistered("LevelEditor.LevelViewportToolBar.View"))
+		// Use a static bool to track whether or not this menu is registered. Bool instead of checking the registered
+		// state with ToolMenus because we want the new viewport toolbar to be able to create this menu without breaking
+		// this code. Static because this code can be called multiple times using different instances of this class.
+		static bool bDidRegisterMenu = false;
+		if (!bDidRegisterMenu)
 		{
-			UToolMenu* Menu = UToolMenus::Get()->RegisterMenu("LevelEditor.LevelViewportToolBar.View", "UnrealEd.ViewportToolbar.View");
+			bDidRegisterMenu = true;
+
+			// Don't warn here to avoid warnings if the new viewport toolbar already has created an empty version
+			// of this menu.
+			UToolMenu* Menu = UToolMenus::Get()->RegisterMenu(
+				"LevelEditor.LevelViewportToolBar.View", "UnrealEd.ViewportToolbar.View", EMultiBoxType::Menu, false
+			);
 			Menu->AddDynamicSection("LevelSection", FNewToolMenuDelegate::CreateLambda([](UToolMenu* InMenu) {
 				if (UEditorViewportViewMenuContext* Context = InMenu->FindContext<UEditorViewportViewMenuContext>())
 				{
@@ -1008,11 +1018,18 @@ void SLevelViewportToolBar::FillCameraMenu(UToolMenu* Menu) const
 TSharedRef<SWidget> SLevelViewportToolBar::GenerateShowMenu() const
 {
 	static const FName MenuName("LevelEditor.LevelViewportToolbar.Show");
+
+	// Use a static bool to track whether or not this menu is registered. Bool instead of checking the registered state
+	// with ToolMenus because we want the new viewport toolbar to be able to create this menu without breaking this
+	// code. Static because this code can be called multiple times using different instances of this class.
 	static bool bDidRegisterMenu = false;
 	if (!bDidRegisterMenu)
 	{
 		bDidRegisterMenu = true;
-		UToolMenu* Menu = UToolMenus::Get()->RegisterMenu(MenuName);
+
+		// Don't warn here to avoid warnings if the new viewport toolbar already has created an empty version
+		// of this menu.
+		UToolMenu* Menu = UToolMenus::Get()->RegisterMenu(MenuName, NAME_None, EMultiBoxType::Menu, false);
 		Menu->AddDynamicSection(
 			"LevelDynamicSection",
 			FNewToolMenuDelegate::CreateLambda(
diff --git a/Engine/Source/Editor/UnrealEd/Private/SEditorViewportViewMenu.cpp b/Engine/Source/Editor/UnrealEd/Private/SEditorViewportViewMenu.cpp
--- a/Engine/Source/Editor/UnrealEd/Private/SEditorViewportViewMenu.cpp
+++ b/Engine/Source/Editor/UnrealEd/Private/SEditorViewportViewMenu.cpp
@@ -54,9 +54,17 @@ const FSlateBrush* SEditorViewportViewMenu::GetViewMenuLabelIcon() const
 
 void SEditorViewportViewMenu::RegisterMenus() const
 {
-	if (!UToolMenus::Get()->IsMenuRegistered(BaseMenuName))
+	// Use a static bool to track whether or not this menu is registered. Bool instead of checking the registered state
+	// with ToolMenus because we want the new viewport toolbar to be able to create this menu without breaking this
+	// code. Static because this code can be called multiple times using different instances of this class.
+	static bool bDidRegisterMenu = false;
+	if (!bDidRegisterMenu)
 	{
-		UToolMenu* Menu = UToolMenus::Get()->RegisterMenu(BaseMenuName);
+		bDidRegisterMenu = true;
+
+		// Don't warn here to avoid warnings if the new viewport toolbar already has created an empty version
+		// of this menu.
+		UToolMenu* Menu = UToolMenus::Get()->RegisterMenu(BaseMenuName, NAME_None, EMultiBoxType::Menu, false);
 		Menu->AddDynamicSection("BaseSection", FNewToolMenuDelegate::CreateLambda([](UToolMenu* InMenu)
 		{
 			if (UEditorViewportViewMenuContext* Context = InMenu->FindContext<UEditorViewportViewMenuContext>())
