Move CreateViewportToolbarViewModesSubmenu() to UnrealEd so asset editors can use it

This splits the old UE::LevelEditor::CreateViewportToolbarViewModesSubmenu() into UE::UnrealEd::CreateViewportToolbarViewModesSubmenu() and UE::LevelEditor::ExtendViewModesSubmenu(FName). The generic code was moved to UE::UnrealEd and the level editor specific code stays in UE::LevelEditor.

diff --git a/Engine/Source/Editor/LevelEditor/Private/SLevelViewport.cpp b/Engine/Source/Editor/LevelEditor/Private/SLevelViewport.cpp
--- a/Engine/Source/Editor/LevelEditor/Private/SLevelViewport.cpp
+++ b/Engine/Source/Editor/LevelEditor/Private/SLevelViewport.cpp
@@ -106,6 +106,7 @@
 #include "SkeletalRenderPublic.h"
 #include "ViewportToolbar/LevelEditorViewportToolbarSections.h"
 #include "ViewportToolbar/LevelViewportContext.h"
+#include "ViewportToolbar/UnrealEdViewportToolbar.h"
 
 static const FName LevelEditorName("LevelEditor");
 static FAutoConsoleCommand EnableInViewportMenu(TEXT("Editor.EnableInViewportMenu"), TEXT("Enables the new in-viewport property menu"), FConsoleCommandDelegate::CreateStatic(&SLevelViewport::EnableInViewportMenu));
@@ -1986,7 +1987,12 @@ TSharedPtr<SWidget> SLevelViewport::MakeViewportToolbar()
 						);
 					}
 
-					FToolMenuEntry ViewModesSubmenu = UE::LevelEditor::CreateViewportToolbarViewModesSubmenu();
+					// Add the level editor specific entries.
+					UE::LevelEditor::ExtendViewModesSubmenu("LevelEditor.ViewportToolbar.ViewModes");
+
+					// Create and add the submenu entry to make the menu we just created and extended appear in the
+					// viewport toolbar itself.
+					FToolMenuEntry ViewModesSubmenu = UE::UnrealEd::CreateViewportToolbarViewModesSubmenu();
 					ViewModesSubmenu.InsertPosition.Position = EToolMenuInsertType::Last;
 					RightSection.AddEntry(ViewModesSubmenu);
 				}
@@ -2046,10 +2052,19 @@ TSharedPtr<SWidget> SLevelViewport::MakeViewportToolbar()
 			}
 		}
 
-		ULevelViewportContext* const ContextObject = NewObject<ULevelViewportContext>();
-		ContextObject->LevelViewport = SharedThis(this);
+		// Add the UnrealEd viewport toolbar context.
+		{
+			UUnrealEdViewportToolbarContext* const ContextObject = NewObject<UUnrealEdViewportToolbarContext>();
+			ContextObject->Viewport = SharedThis(this);
+			ViewportToolbarContext.AddObject(ContextObject);
+		}
 
-		ViewportToolbarContext.AddObject(ContextObject);
+		// Add the level editor viewport toolbar context.
+		{
+			ULevelViewportContext* const ContextObject = NewObject<ULevelViewportContext>();
+			ContextObject->LevelViewport = SharedThis(this);
+			ViewportToolbarContext.AddObject(ContextObject);
+		}
 	}
 
 	const TSharedRef<SWidget> NewViewportToolbar = SNew(SBox)
diff --git a/Engine/Source/Editor/LevelEditor/Private/ViewportToolbar/LevelEditorViewportToolbarSections.cpp b/Engine/Source/Editor/LevelEditor/Private/ViewportToolbar/LevelEditorViewportToolbarSections.cpp
--- a/Engine/Source/Editor/LevelEditor/Private/ViewportToolbar/LevelEditorViewportToolbarSections.cpp
+++ b/Engine/Source/Editor/LevelEditor/Private/ViewportToolbar/LevelEditorViewportToolbarSections.cpp
@@ -884,49 +884,25 @@ void PopulateViewModesMenu(UToolMenu* InMenu, TSharedRef<::SLevelViewport> InVie
 	}
 }
 
-FToolMenuEntry CreateViewportToolbarViewModesSubmenu()
+void ExtendViewModesSubmenu(FName InViewModesSubmenuName)
 {
-	// This has to be a dynamic entry for the ViewModes submenu's label to be able to access the context.
-	return FToolMenuEntry::InitDynamicEntry(
-		"DynamicViewModes",
-		FNewToolMenuSectionDelegate::CreateLambda(
-			[](FToolMenuSection& InDynamicSection) -> void
+	UToolMenu* const Submenu = UToolMenus::Get()->ExtendMenu(InViewModesSubmenuName);
+
+	Submenu->AddDynamicSection(
+		"LevelEditorViewModesExtensionDynamicSection",
+		FNewToolMenuDelegate::CreateLambda(
+			[](UToolMenu* InDynamicMenu)
 			{
-				TAttribute<FText> LabelAttribute = UE::UnrealEd::GetViewModesSubmenuLabel(nullptr);
-				if (ULevelViewportContext* const LevelViewportContext =
-						InDynamicSection.FindContext<ULevelViewportContext>())
+				ULevelViewportContext* const LevelViewportContext = InDynamicMenu->FindContext<ULevelViewportContext>();
+				if (!LevelViewportContext)
 				{
-					TWeakPtr<SEditorViewport> EditorViewport = LevelViewportContext->LevelViewport;
-					LabelAttribute = TAttribute<FText>::CreateLambda(
-						[EditorViewport]()
-						{
-							return UE::UnrealEd::GetViewModesSubmenuLabel(EditorViewport);
-						}
-					);
+					return;
 				}
 
-				InDynamicSection.AddSubMenu(
-					"ViewModes",
-					LabelAttribute,
-					LOCTEXT("ViewModesSubmenuTooltip", "View mode settings for the current viewport."),
-					FNewToolMenuDelegate::CreateLambda(
-						[](UToolMenu* Submenu) -> void
-						{
-							ULevelViewportContext* const LevelViewportContext =
-								Submenu->FindContext<ULevelViewportContext>();
-							if (!LevelViewportContext)
-							{
-								return;
-							}
-
-							if (const TSharedPtr<::SLevelViewport> LevelViewport = LevelViewportContext->LevelViewport.Pin())
-							{
-								UE::UnrealEd::PopulateViewModesMenu(Submenu, LevelViewport.ToSharedRef());
-								PopulateViewModesMenu(Submenu, LevelViewport.ToSharedRef());
-							}
-						}
-					)
-				);
+				if (const TSharedPtr<::SLevelViewport> LevelViewport = LevelViewportContext->LevelViewport.Pin())
+				{
+					PopulateViewModesMenu(InDynamicMenu, LevelViewport.ToSharedRef());
+				}
 			}
 		)
 	);
diff --git a/Engine/Source/Editor/LevelEditor/Private/ViewportToolbar/LevelEditorViewportToolbarSections.h b/Engine/Source/Editor/LevelEditor/Private/ViewportToolbar/LevelEditorViewportToolbarSections.h
--- a/Engine/Source/Editor/LevelEditor/Private/ViewportToolbar/LevelEditorViewportToolbarSections.h
+++ b/Engine/Source/Editor/LevelEditor/Private/ViewportToolbar/LevelEditorViewportToolbarSections.h
@@ -23,7 +23,7 @@ FToolMenuEntry CreateViewportToolbarSelectionSection();
 
 TSharedPtr<FExtender> GetViewModesLegacyExtenders();
 void PopulateViewModesMenu(UToolMenu* InMenu, TSharedRef<::SLevelViewport> InViewport);
-FToolMenuEntry CreateViewportToolbarViewModesSubmenu();
+void ExtendViewModesSubmenu(FName InViewModesSubmenuName);
 
 FToolMenuEntry CreateShowFoliageSubmenu();
 FToolMenuEntry CreateShowHLODsSubmenu();
diff --git a/Engine/Source/Editor/UnrealEd/Private/SEditorViewportViewMenu.cpp b/Engine/Source/Editor/UnrealEd/Private/SEditorViewportViewMenu.cpp
--- a/Engine/Source/Editor/UnrealEd/Private/SEditorViewportViewMenu.cpp
+++ b/Engine/Source/Editor/UnrealEd/Private/SEditorViewportViewMenu.cpp
@@ -9,6 +9,7 @@
 #include "RenderResource.h"
 #include "SEditorViewport.h"
 #include "ViewportToolbar/UnrealEdViewportToolbar.h"
+#include "ViewportToolbar/UnrealEdViewportToolbarContext.h"
 
 #define LOCTEXT_NAMESPACE "EditorViewportViewMenu"
 
diff --git a/Engine/Source/Editor/UnrealEd/Private/ViewportToolbar/UnrealEdViewportToolbar.cpp b/Engine/Source/Editor/UnrealEd/Private/ViewportToolbar/UnrealEdViewportToolbar.cpp
--- a/Engine/Source/Editor/UnrealEd/Private/ViewportToolbar/UnrealEdViewportToolbar.cpp
+++ b/Engine/Source/Editor/UnrealEd/Private/ViewportToolbar/UnrealEdViewportToolbar.cpp
@@ -12,7 +12,9 @@
 #include "SEditorViewport.h"
 #include "Templates/SharedPointer.h"
 #include "ToolMenu.h"
+#include "ToolMenuEntry.h"
 #include "ToolMenuSection.h"
+#include "ViewportToolbar/UnrealEdViewportToolbarContext.h"
 
 #define LOCTEXT_NAMESPACE "UnrealEdViewportToolbar"
 
@@ -319,6 +321,53 @@ void PopulateViewModesMenu(
 	}
 }
 
+FToolMenuEntry CreateViewportToolbarViewModesSubmenu()
+{
+	// This has to be a dynamic entry for the ViewModes submenu's label to be able to access the context.
+	return FToolMenuEntry::InitDynamicEntry(
+		"DynamicViewModes",
+		FNewToolMenuSectionDelegate::CreateLambda(
+			[](FToolMenuSection& InDynamicSection) -> void
+			{
+				// Base the label on the current view mode.
+				TAttribute<FText> LabelAttribute = UE::UnrealEd::GetViewModesSubmenuLabel(nullptr);
+				if (UUnrealEdViewportToolbarContext* const Context =
+						InDynamicSection.FindContext<UUnrealEdViewportToolbarContext>())
+				{
+					LabelAttribute = TAttribute<FText>::CreateLambda(
+						[WeakViewport = Context->Viewport]()
+						{
+							return UE::UnrealEd::GetViewModesSubmenuLabel(WeakViewport);
+						}
+					);
+				}
+
+				InDynamicSection.AddSubMenu(
+					"ViewModes",
+					LabelAttribute,
+					LOCTEXT("ViewModesSubmenuTooltip", "View mode settings for the current viewport."),
+					FNewToolMenuDelegate::CreateLambda(
+						[](UToolMenu* Submenu) -> void
+						{
+							UUnrealEdViewportToolbarContext* const Context =
+								Submenu->FindContext<UUnrealEdViewportToolbarContext>();
+							if (!Context)
+							{
+								return;
+							}
+
+							if (const TSharedPtr<SEditorViewport> Viewport = Context->Viewport.Pin())
+							{
+								PopulateViewModesMenu(Submenu, Viewport.ToSharedRef(), Context->IsViewModeSupported);
+							}
+						}
+					)
+				);
+			}
+		)
+	);
+}
+
 } // namespace UE::UnrealEd
 
 #undef LOCTEXT_NAMESPACE
diff --git a/Engine/Source/Editor/UnrealEd/Public/ViewportToolbar/UnrealEdViewportToolbar.h b/Engine/Source/Editor/UnrealEd/Public/ViewportToolbar/UnrealEdViewportToolbar.h
--- a/Engine/Source/Editor/UnrealEd/Public/ViewportToolbar/UnrealEdViewportToolbar.h
+++ b/Engine/Source/Editor/UnrealEd/Public/ViewportToolbar/UnrealEdViewportToolbar.h
@@ -5,18 +5,18 @@
 #include "Delegates/DelegateCombinations.h"
 #include "Engine/EngineBaseTypes.h"
 #include "Templates/SharedPointerFwd.h"
+#include "UnrealEdViewportToolbarContext.h"
 
 class FText;
 class SEditorViewport;
 class UToolMenu;
+struct FToolMenuEntry;
 
 namespace UE::UnrealEd
 {
 
 UNREALED_API FText GetViewModesSubmenuLabel(TWeakPtr<SEditorViewport> InViewport);
 
-DECLARE_DELEGATE_RetVal_OneParam(bool, IsViewModeSupportedDelegate, EViewModeIndex);
-
 /**
  * Populate a given UToolMenu with entries for a View Modes viewport toolbar submenu.
  *
@@ -28,4 +28,6 @@ UNREALED_API void PopulateViewModesMenu(UToolMenu* InMenu,
 	TSharedRef<SEditorViewport> InViewport,
 	IsViewModeSupportedDelegate InIsViewModeSupported = IsViewModeSupportedDelegate());
 
+UNREALED_API FToolMenuEntry CreateViewportToolbarViewModesSubmenu();
+
 } // namespace UE::UnrealEd
diff --git a/Engine/Source/Editor/UnrealEd/Public/ViewportToolbar/UnrealEdViewportToolbarContext.h b/Engine/Source/Editor/UnrealEd/Public/ViewportToolbar/UnrealEdViewportToolbarContext.h
new file mode 100644
--- /dev/null
+++ b/Engine/Source/Editor/UnrealEd/Public/ViewportToolbar/UnrealEdViewportToolbarContext.h
@@ -0,0 +1,28 @@
+﻿// Copyright Epic Games, Inc. All Rights Reserved.
+
+#pragma once
+
+#include "Engine/EngineBaseTypes.h"
+#include "Templates/SharedPointer.h"
+#include "UObject/Object.h"
+
+#include "UnrealEdViewportToolbarContext.generated.h"
+
+namespace UE::UnrealEd
+{
+
+DECLARE_DELEGATE_RetVal_OneParam(bool, IsViewModeSupportedDelegate, EViewModeIndex);
+
+}
+
+class SEditorViewport;
+
+UCLASS()
+class UNREALED_API UUnrealEdViewportToolbarContext : public UObject
+{
+	GENERATED_BODY()
+
+public:
+	TWeakPtr<SEditorViewport> Viewport;
+	UE::UnrealEd::IsViewModeSupportedDelegate IsViewModeSupported;
+};
