Move new viewport toolbar functions CreateViewportToolbarTransformsSection() and CreateViewportToolbarSelectionSection() to UnrealEd

They used to live in the leve editor. Now asset editors can use them, too.

diff --git a/Engine/Source/Editor/LevelEditor/Private/SLevelViewport.cpp b/Engine/Source/Editor/LevelEditor/Private/SLevelViewport.cpp
--- a/Engine/Source/Editor/LevelEditor/Private/SLevelViewport.cpp
+++ b/Engine/Source/Editor/LevelEditor/Private/SLevelViewport.cpp
@@ -1926,14 +1926,14 @@ TSharedPtr<SWidget> SLevelViewport::MakeViewportToolbar()
 
 				// Add the "Transforms" sub menu.
 				{
-					FToolMenuEntry TransformsSubmenu = UE::LevelEditor::CreateViewportToolbarTransformsSection();
+					FToolMenuEntry TransformsSubmenu = UE::UnrealEd::CreateViewportToolbarTransformsSection();
 					TransformsSubmenu.InsertPosition.Position = EToolMenuInsertType::First;
 					LeftSection.AddEntry(TransformsSubmenu);
 				}
 
 				// Add the "Selection" sub menu.
 				{
-					FToolMenuEntry SelectionSubmenu = UE::LevelEditor::CreateViewportToolbarSelectionSection();
+					FToolMenuEntry SelectionSubmenu = UE::UnrealEd::CreateViewportToolbarSelectionSection();
 					SelectionSubmenu.InsertPosition.Position = EToolMenuInsertType::First;
 					LeftSection.AddEntry(SelectionSubmenu);
 				}
diff --git a/Engine/Source/Editor/LevelEditor/Private/ViewportToolbar/LevelEditorViewportToolbarSections.cpp b/Engine/Source/Editor/LevelEditor/Private/ViewportToolbar/LevelEditorViewportToolbarSections.cpp
--- a/Engine/Source/Editor/LevelEditor/Private/ViewportToolbar/LevelEditorViewportToolbarSections.cpp
+++ b/Engine/Source/Editor/LevelEditor/Private/ViewportToolbar/LevelEditorViewportToolbarSections.cpp
@@ -4,7 +4,6 @@
 
 #include "EditorViewportCommands.h"
 #include "FoliageType.h"
-#include "Framework/Commands/GenericCommands.h"
 #include "GameFramework/ActorPrimitiveColorHandler.h"
 #include "GroomVisualizationData.h"
 #include "Layers/LayersSubsystem.h"
@@ -147,7 +146,7 @@ FToolMenuEntry CreateSurfaceSnapOffsetEntry()
 	SurfaceOffset.MakeCustomWidget.BindLambda(
 		[Label, Tooltip, WidgetsMargin](const FToolMenuContext& InContext, const FToolMenuCustomWidgetContext& InWidgetContext)
 		{
-			// clang-format on
+			// clang-format off
 			return SNew(SHorizontalBox).IsEnabled_Lambda([]()
 				{
 					return GetDefault<ULevelEditorViewportSettings>()->SnapToSurface.bEnabled;
@@ -181,7 +180,7 @@ FToolMenuEntry CreateSurfaceSnapOffsetEntry()
 						})
 					]
 				];
-			// clang-format off
+			// clang-format on
 		});
 
 	return SurfaceOffset;
@@ -251,7 +250,7 @@ FToolMenuEntry CreateActorSnapDistanceEntry()
 	SnapDistance.MakeCustomWidget.BindLambda(
 		[Label, Tooltip, WidgetsMargin](const FToolMenuContext& InContext, const FToolMenuCustomWidgetContext& InWidgetContext)
 		{
-			// clang-format on
+			// clang-format off
 			return SNew(SHorizontalBox).IsEnabled_Lambda([]()
 				{
 					return !!GetDefault<ULevelEditorViewportSettings>()->bEnableActorSnap;
@@ -282,7 +281,7 @@ FToolMenuEntry CreateActorSnapDistanceEntry()
 						})
 					]
 				];
-			// clang-format off
+			// clang-format on
 		});
 
 	return SnapDistance;
@@ -341,166 +340,6 @@ bool ShowViewportRealtimeWarning(FLevelEditorViewportClient& ViewportClient)
 	return !ViewportClient.IsRealtime() && !ViewportClient.IsRealtimeOverrideSet() && ViewportClient.IsPerspective();
 }
 
-// TODO: Move this outside the level editor and make it publicly available to anyone building a viewport toolbar.
-FToolMenuEntry CreateViewportToolbarTransformsSection()
-{
-	return FToolMenuEntry::InitSubMenu(
-		"Transforms",
-		LOCTEXT("TransformsSubmenuLabel", "Transforms"),
-		LOCTEXT("TransformsSubmenuTooltip", "Viewport-related transforms tools"),
-		FNewToolMenuDelegate::CreateLambda(
-			[](UToolMenu* Submenu) -> void
-			{
-				{
-					FToolMenuSection& TransformToolsSection =
-						Submenu->FindOrAddSection("TransformTools", LOCTEXT("TransformToolsLabel", "Transform Tools"));
-
-					FToolMenuEntry SelectMode = FToolMenuEntry::InitMenuEntry(FEditorViewportCommands::Get().SelectMode);
-					SelectMode.UserInterfaceActionType = EUserInterfaceActionType::RadioButton;
-					SelectMode.SetShowInToolbarTopLevel(true);
-					TransformToolsSection.AddEntry(SelectMode);
-
-					FToolMenuEntry TranslateMode =
-						FToolMenuEntry::InitMenuEntry(FEditorViewportCommands::Get().TranslateMode);
-					TranslateMode.UserInterfaceActionType = EUserInterfaceActionType::RadioButton;
-					TranslateMode.SetShowInToolbarTopLevel(true);
-					TransformToolsSection.AddEntry(TranslateMode);
-
-					FToolMenuEntry RotateMode = FToolMenuEntry::InitMenuEntry(FEditorViewportCommands::Get().RotateMode);
-					RotateMode.UserInterfaceActionType = EUserInterfaceActionType::RadioButton;
-					RotateMode.SetShowInToolbarTopLevel(true);
-					TransformToolsSection.AddEntry(RotateMode);
-
-					FToolMenuEntry ScaleMode = FToolMenuEntry::InitMenuEntry(FEditorViewportCommands::Get().ScaleMode);
-					ScaleMode.UserInterfaceActionType = EUserInterfaceActionType::RadioButton;
-					ScaleMode.SetShowInToolbarTopLevel(true);
-					TransformToolsSection.AddEntry(ScaleMode);
-				}
-
-				{
-					FToolMenuSection& SpacesSection = Submenu->FindOrAddSection("Spaces", LOCTEXT("SpacesLabel", "Spaces"));
-
-					FToolMenuEntry WorldSpace =
-						FToolMenuEntry::InitMenuEntry(FEditorViewportCommands::Get().RelativeCoordinateSystem_World);
-					WorldSpace.UserInterfaceActionType = EUserInterfaceActionType::RadioButton;
-					WorldSpace.SetShowInToolbarTopLevel(true);
-					SpacesSection.AddEntry(WorldSpace);
-
-					FToolMenuEntry LocalSpace =
-						FToolMenuEntry::InitMenuEntry(FEditorViewportCommands::Get().RelativeCoordinateSystem_Local);
-					LocalSpace.UserInterfaceActionType = EUserInterfaceActionType::RadioButton;
-					LocalSpace.SetShowInToolbarTopLevel(true);
-					SpacesSection.AddEntry(LocalSpace);
-				}
-
-				{
-					FToolMenuSection& GizmoSection = Submenu->FindOrAddSection("Gizmo", LOCTEXT("GizmoLabel", "Gizmo"));
-
-					GizmoSection.AddMenuEntry(
-						FLevelEditorCommands::Get().ShowTransformWidget,
-						LOCTEXT("ShowTransformGizmoLabel", "Show Transform Gizmo")
-					);
-					
-					TSharedRef<SWidget> GizmoScaleWidget =
-						// clang-format off
-						SNew(SHorizontalBox)
-						+ SHorizontalBox::Slot()
-						.FillWidth(0.9f)
-						[
-							SNew(SSpinBox<int32>)
-								.MinValue(-10)
-								.MaxValue(150)
-								.ToolTipText_Lambda(
-									[]() -> FText
-									{
-										return FText::AsNumber(
-											GetDefault<ULevelEditorViewportSettings>()->TransformWidgetSizeAdjustment
-										);
-									}
-								)
-								.Value_Lambda(
-									[]() -> float
-									{
-										return GetDefault<ULevelEditorViewportSettings>()->TransformWidgetSizeAdjustment;
-									}
-								)
-								.OnValueChanged_Lambda(
-									[](float InValue)
-									{
-										ULevelEditorViewportSettings* ViewportSettings =
-											GetMutableDefault<ULevelEditorViewportSettings>();
-										ViewportSettings->TransformWidgetSizeAdjustment = InValue;
-										ViewportSettings->PostEditChange();
-									}
-								)
-						]
-						+ SHorizontalBox::Slot()
-						.FillWidth(0.1f);
-					// clang-format on
-					GizmoSection.AddEntry(FToolMenuEntry::InitWidget(
-						"GizmoScale", GizmoScaleWidget, LOCTEXT("GizmoScaleLabel", "Gizmo Scale")
-					));
-				}
-			}
-		)
-	);
-}
-
-FToolMenuEntry CreateViewportToolbarSelectionSection()
-{
-	return FToolMenuEntry::InitSubMenu(
-		"Select",
-		LOCTEXT("SelectonSubmenuLabel", "Select"),
-		LOCTEXT("SelectionSubmenuTooltip", "Viewport-related selection tools"),
-		FNewToolMenuDelegate::CreateLambda(
-			[](UToolMenu* Submenu) -> void
-			{
-				{
-					FToolMenuSection& UnnamedSection = Submenu->FindOrAddSection(NAME_None);
-
-					UnnamedSection.AddMenuEntry(FGenericCommands::Get().SelectAll);
-					UnnamedSection.AddMenuEntry(FLevelEditorCommands::Get().SelectNone);
-					UnnamedSection.AddMenuEntry(FLevelEditorCommands::Get().InvertSelection);
-
-					UnnamedSection.AddSeparator("Advanced");
-
-					UnnamedSection.AddMenuEntry(FLevelEditorCommands::Get().SelectAllActorsOfSameClass);
-				}
-
-				{
-					FToolMenuSection& ByTypeSection =
-						Submenu->FindOrAddSection("ByTypeSection", LOCTEXT("ByTypeSectionLabel", "By Type"));
-
-					ByTypeSection.AddSubMenu(
-						"BSP",
-						LOCTEXT("BspLabel", "BSP"),
-						LOCTEXT("BspTooltip", "BSP-related tools"),
-						FNewToolMenuDelegate::CreateLambda(
-							[](UToolMenu* BspMenu)
-							{
-								FToolMenuSection& SelectAllSection = BspMenu->FindOrAddSection(
-									"SelectAllBSP", LOCTEXT("SelectAllBSPLabel", "Select All BSP")
-								);
-
-								SelectAllSection.AddMenuEntry(FLevelEditorCommands::Get().SelectAllAddditiveBrushes);
-								SelectAllSection.AddMenuEntry(FLevelEditorCommands::Get().SelectAllSubtractiveBrushes);
-								SelectAllSection.AddMenuEntry(FLevelEditorCommands::Get().SelectAllSurfaces);
-							}
-						)
-					);
-				}
-
-				{
-					FToolMenuSection& OptionsSection =
-						Submenu->FindOrAddSection("Options", LOCTEXT("OptionsLabel", "Options"));
-
-					OptionsSection.AddMenuEntry(FLevelEditorCommands::Get().AllowTranslucentSelection);
-				}
-			}
-		)
-	);
-}
-
 TSharedPtr<FExtender> GetViewModesLegacyExtenders()
 {
 	FLevelEditorModule& LevelEditorModule = FModuleManager::GetModuleChecked<FLevelEditorModule>(TEXT("LevelEditor"));
diff --git a/Engine/Source/Editor/LevelEditor/Private/ViewportToolbar/LevelEditorViewportToolbarSections.h b/Engine/Source/Editor/LevelEditor/Private/ViewportToolbar/LevelEditorViewportToolbarSections.h
--- a/Engine/Source/Editor/LevelEditor/Private/ViewportToolbar/LevelEditorViewportToolbarSections.h
+++ b/Engine/Source/Editor/LevelEditor/Private/ViewportToolbar/LevelEditorViewportToolbarSections.h
@@ -17,10 +17,6 @@ namespace UE::LevelEditor
 
 bool ShowViewportRealtimeWarning(FLevelEditorViewportClient& ViewportClient);
 
-FToolMenuEntry CreateViewportToolbarTransformsSection();
-
-FToolMenuEntry CreateViewportToolbarSelectionSection();
-
 TSharedPtr<FExtender> GetViewModesLegacyExtenders();
 void PopulateViewModesMenu(UToolMenu* InMenu, TSharedRef<::SLevelViewport> InViewport);
 void ExtendViewModesSubmenu(FName InViewModesSubmenuName);
diff --git a/Engine/Source/Editor/UnrealEd/Private/ViewportToolbar/UnrealEdViewportToolbar.cpp b/Engine/Source/Editor/UnrealEd/Private/ViewportToolbar/UnrealEdViewportToolbar.cpp
--- a/Engine/Source/Editor/UnrealEd/Private/ViewportToolbar/UnrealEdViewportToolbar.cpp
+++ b/Engine/Source/Editor/UnrealEd/Private/ViewportToolbar/UnrealEdViewportToolbar.cpp
@@ -6,15 +6,19 @@
 #include "DebugViewModeHelpers.h"
 #include "EditorViewportClient.h"
 #include "EditorViewportCommands.h"
+#include "Framework/Commands/GenericCommands.h"
 #include "GPUSkinCache.h"
 #include "GPUSkinCacheVisualizationMenuCommands.h"
+#include "LevelEditorActions.h"
 #include "RayTracingDebugVisualizationMenuCommands.h"
 #include "SEditorViewport.h"
+#include "Settings/LevelEditorViewportSettings.h"
 #include "Templates/SharedPointer.h"
 #include "ToolMenu.h"
 #include "ToolMenuEntry.h"
 #include "ToolMenuSection.h"
 #include "ViewportToolbar/UnrealEdViewportToolbarContext.h"
+#include "Widgets/Input/SSpinBox.h"
 
 #define LOCTEXT_NAMESPACE "UnrealEdViewportToolbar"
 
@@ -45,6 +49,165 @@ bool ShowNewViewportToolbars()
 	return Private::CVarToolMenusViewportToolbarsValue >= 1;
 }
 
+FToolMenuEntry CreateViewportToolbarTransformsSection()
+{
+	return FToolMenuEntry::InitSubMenu(
+		"Transforms",
+		LOCTEXT("TransformsSubmenuLabel", "Transforms"),
+		LOCTEXT("TransformsSubmenuTooltip", "Viewport-related transforms tools"),
+		FNewToolMenuDelegate::CreateLambda(
+			[](UToolMenu* Submenu) -> void
+			{
+				{
+					FToolMenuSection& TransformToolsSection =
+						Submenu->FindOrAddSection("TransformTools", LOCTEXT("TransformToolsLabel", "Transform Tools"));
+
+					FToolMenuEntry SelectMode = FToolMenuEntry::InitMenuEntry(FEditorViewportCommands::Get().SelectMode);
+					SelectMode.UserInterfaceActionType = EUserInterfaceActionType::RadioButton;
+					SelectMode.SetShowInToolbarTopLevel(true);
+					TransformToolsSection.AddEntry(SelectMode);
+
+					FToolMenuEntry TranslateMode =
+						FToolMenuEntry::InitMenuEntry(FEditorViewportCommands::Get().TranslateMode);
+					TranslateMode.UserInterfaceActionType = EUserInterfaceActionType::RadioButton;
+					TranslateMode.SetShowInToolbarTopLevel(true);
+					TransformToolsSection.AddEntry(TranslateMode);
+
+					FToolMenuEntry RotateMode = FToolMenuEntry::InitMenuEntry(FEditorViewportCommands::Get().RotateMode);
+					RotateMode.UserInterfaceActionType = EUserInterfaceActionType::RadioButton;
+					RotateMode.SetShowInToolbarTopLevel(true);
+					TransformToolsSection.AddEntry(RotateMode);
+
+					FToolMenuEntry ScaleMode = FToolMenuEntry::InitMenuEntry(FEditorViewportCommands::Get().ScaleMode);
+					ScaleMode.UserInterfaceActionType = EUserInterfaceActionType::RadioButton;
+					ScaleMode.SetShowInToolbarTopLevel(true);
+					TransformToolsSection.AddEntry(ScaleMode);
+				}
+
+				{
+					FToolMenuSection& SpacesSection = Submenu->FindOrAddSection("Spaces", LOCTEXT("SpacesLabel", "Spaces"));
+
+					FToolMenuEntry WorldSpace =
+						FToolMenuEntry::InitMenuEntry(FEditorViewportCommands::Get().RelativeCoordinateSystem_World);
+					WorldSpace.UserInterfaceActionType = EUserInterfaceActionType::RadioButton;
+					WorldSpace.SetShowInToolbarTopLevel(true);
+					SpacesSection.AddEntry(WorldSpace);
+
+					FToolMenuEntry LocalSpace =
+						FToolMenuEntry::InitMenuEntry(FEditorViewportCommands::Get().RelativeCoordinateSystem_Local);
+					LocalSpace.UserInterfaceActionType = EUserInterfaceActionType::RadioButton;
+					LocalSpace.SetShowInToolbarTopLevel(true);
+					SpacesSection.AddEntry(LocalSpace);
+				}
+
+				{
+					FToolMenuSection& GizmoSection = Submenu->FindOrAddSection("Gizmo", LOCTEXT("GizmoLabel", "Gizmo"));
+
+					GizmoSection.AddMenuEntry(
+						FLevelEditorCommands::Get().ShowTransformWidget,
+						LOCTEXT("ShowTransformGizmoLabel", "Show Transform Gizmo")
+					);
+
+					TSharedRef<SWidget> GizmoScaleWidget =
+						// clang-format off
+						SNew(SHorizontalBox)
+						+ SHorizontalBox::Slot()
+						.FillWidth(0.9f)
+						[
+							SNew(SSpinBox<int32>)
+								.MinValue(-10)
+								.MaxValue(150)
+								.ToolTipText_Lambda(
+									[]() -> FText
+									{
+										return FText::AsNumber(
+											GetDefault<ULevelEditorViewportSettings>()->TransformWidgetSizeAdjustment
+										);
+									}
+								)
+								.Value_Lambda(
+									[]() -> float
+									{
+										return GetDefault<ULevelEditorViewportSettings>()->TransformWidgetSizeAdjustment;
+									}
+								)
+								.OnValueChanged_Lambda(
+									[](float InValue)
+									{
+										ULevelEditorViewportSettings* ViewportSettings =
+											GetMutableDefault<ULevelEditorViewportSettings>();
+										ViewportSettings->TransformWidgetSizeAdjustment = InValue;
+										ViewportSettings->PostEditChange();
+									}
+								)
+						]
+						+ SHorizontalBox::Slot()
+						.FillWidth(0.1f);
+					// clang-format on
+					GizmoSection.AddEntry(FToolMenuEntry::InitWidget(
+						"GizmoScale", GizmoScaleWidget, LOCTEXT("GizmoScaleLabel", "Gizmo Scale")
+					));
+				}
+			}
+		)
+	);
+}
+
+FToolMenuEntry CreateViewportToolbarSelectionSection()
+{
+	return FToolMenuEntry::InitSubMenu(
+		"Select",
+		LOCTEXT("SelectonSubmenuLabel", "Select"),
+		LOCTEXT("SelectionSubmenuTooltip", "Viewport-related selection tools"),
+		FNewToolMenuDelegate::CreateLambda(
+			[](UToolMenu* Submenu) -> void
+			{
+				{
+					FToolMenuSection& UnnamedSection = Submenu->FindOrAddSection(NAME_None);
+
+					UnnamedSection.AddMenuEntry(FGenericCommands::Get().SelectAll);
+					UnnamedSection.AddMenuEntry(FLevelEditorCommands::Get().SelectNone);
+					UnnamedSection.AddMenuEntry(FLevelEditorCommands::Get().InvertSelection);
+
+					UnnamedSection.AddSeparator("Advanced");
+
+					UnnamedSection.AddMenuEntry(FLevelEditorCommands::Get().SelectAllActorsOfSameClass);
+				}
+
+				{
+					FToolMenuSection& ByTypeSection =
+						Submenu->FindOrAddSection("ByTypeSection", LOCTEXT("ByTypeSectionLabel", "By Type"));
+
+					ByTypeSection.AddSubMenu(
+						"BSP",
+						LOCTEXT("BspLabel", "BSP"),
+						LOCTEXT("BspTooltip", "BSP-related tools"),
+						FNewToolMenuDelegate::CreateLambda(
+							[](UToolMenu* BspMenu)
+							{
+								FToolMenuSection& SelectAllSection = BspMenu->FindOrAddSection(
+									"SelectAllBSP", LOCTEXT("SelectAllBSPLabel", "Select All BSP")
+								);
+
+								SelectAllSection.AddMenuEntry(FLevelEditorCommands::Get().SelectAllAddditiveBrushes);
+								SelectAllSection.AddMenuEntry(FLevelEditorCommands::Get().SelectAllSubtractiveBrushes);
+								SelectAllSection.AddMenuEntry(FLevelEditorCommands::Get().SelectAllSurfaces);
+							}
+						)
+					);
+				}
+
+				{
+					FToolMenuSection& OptionsSection =
+						Submenu->FindOrAddSection("Options", LOCTEXT("OptionsLabel", "Options"));
+
+					OptionsSection.AddMenuEntry(FLevelEditorCommands::Get().AllowTranslucentSelection);
+				}
+			}
+		)
+	);
+}
+
 FText GetViewModesSubmenuLabel(TWeakPtr<SEditorViewport> InViewport)
 {
 	FText Label = LOCTEXT("ViewMenuTitle_Default", "View");
diff --git a/Engine/Source/Editor/UnrealEd/Public/ViewportToolbar/UnrealEdViewportToolbar.h b/Engine/Source/Editor/UnrealEd/Public/ViewportToolbar/UnrealEdViewportToolbar.h
--- a/Engine/Source/Editor/UnrealEd/Public/ViewportToolbar/UnrealEdViewportToolbar.h
+++ b/Engine/Source/Editor/UnrealEd/Public/ViewportToolbar/UnrealEdViewportToolbar.h
@@ -21,6 +21,10 @@ UNREALED_API bool ShowOldViewportToolbars();
 /** The value of this function is controlled by the CVAR "ToolMenusViewportToolbars". */
 UNREALED_API bool ShowNewViewportToolbars();
 
+UNREALED_API FToolMenuEntry CreateViewportToolbarTransformsSection();
+
+UNREALED_API FToolMenuEntry CreateViewportToolbarSelectionSection();
+
 UNREALED_API FText GetViewModesSubmenuLabel(TWeakPtr<SEditorViewport> InViewport);
 
 /**
