[Viewport Toolbar] Level Editor toolbar: add view mode extensions support

diff --git a/Engine/Source/Editor/LevelEditor/Private/SLevelViewport.cpp b/Engine/Source/Editor/LevelEditor/Private/SLevelViewport.cpp
--- a/Engine/Source/Editor/LevelEditor/Private/SLevelViewport.cpp
+++ b/Engine/Source/Editor/LevelEditor/Private/SLevelViewport.cpp
@@ -2030,16 +2030,26 @@ TSharedPtr<SWidget> SLevelViewport::MakeViewportToolbar()
 			// Stay backward-compatible with legacy view menu extenders.
 			ViewportToolbarContext.AddExtender(UE::LevelEditor::GetViewModesLegacyExtenders());
 
+			FLevelEditorModule& LevelEditorModule =
+					FModuleManager::GetModuleChecked<FLevelEditorModule>(TEXT("LevelEditor"));
+
+			TSharedRef<FUICommandList> CommandListRef = GetCommandList().ToSharedRef();
+
 			// Stay backward-compatible with legacy show menu extenders.
 			{
-				FLevelEditorModule& LevelEditorModule =
-					FModuleManager::GetModuleChecked<FLevelEditorModule>(TEXT("LevelEditor"));
-				TSharedRef<FUICommandList> CommandListRef = GetCommandList().ToSharedRef();
 				TSharedPtr<FExtender> Extenders = LevelEditorModule.AssembleExtenders(
 					CommandListRef, LevelEditorModule.GetAllLevelViewportShowMenuExtenders()
 				);
 				ViewportToolbarContext.AddExtender(Extenders);
 			}
+
+			// Stay backward-compatible with legacy view menu extenders
+			{
+				TSharedPtr<FExtender> ViewMenuExtenders = LevelEditorModule.AssembleExtenders(
+					CommandListRef, LevelEditorModule.GetAllLevelEditorToolbarViewMenuExtenders()
+				);
+				ViewportToolbarContext.AddExtender(ViewMenuExtenders);
+			}
 		}
 
 		// Add the UnrealEd viewport toolbar context.
