Actor Coloration
 * Add new view mode index in order to have the Actor Coloration mode behave like the other view modes (mutually exclusive).
 * Add tooltips for most of the registered actor color handlers.
 * Use FText in PropertyColorSettings and NSLOCTEXT in .ini

diff --git a/Engine/Config/BaseEditorSettings.ini b/Engine/Config/BaseEditorSettings.ini
--- a/Engine/Config/BaseEditorSettings.ini
+++ b/Engine/Config/BaseEditorSettings.ini
@@ -45,7 +45,7 @@ IncludeAssetPaths=True
 IncludeCollectionNames=True
 
 [/Script/UnrealEd.PropertyColorSettings]
-+CustomProperties=(Name="AffectsNavmesh",Text="Affects NavMesh",PropertyChain="RootComponent.bCanEverAffectNavigation",PropertyValue="True",PropertyColor=(R=0,G=0,B=255))
++CustomProperties=(Name="AffectsNavmesh",Text=NSLOCTEXT("ActorColoration", "AffectsNavmesh", "Affects NavMesh"), TextToolTip=NSLOCTEXT("ActorColoration", "AffectsNavmesh_ToolTip", "Colorize actor in Blue if it is affecting the navmesh."), PropertyChain="RootComponent.bCanEverAffectNavigation",PropertyValue="True",PropertyColor=(R=0,G=0,B=255))
 
 [/Script/VREditor.VRModeSettings]
 bEnableAutoVREditMode=False
diff --git a/Engine/Source/Editor/DataLayerEditor/Private/DataLayer/DataLayerEditorSubsystem.cpp b/Engine/Source/Editor/DataLayerEditor/Private/DataLayer/DataLayerEditorSubsystem.cpp
--- a/Engine/Source/Editor/DataLayerEditor/Private/DataLayer/DataLayerEditorSubsystem.cpp
+++ b/Engine/Source/Editor/DataLayerEditor/Private/DataLayer/DataLayerEditorSubsystem.cpp
@@ -202,7 +202,9 @@ void FDataLayersBroadcast::Initialize()
 				}
 			}
 			return FLinearColor::Gray;
-		});
+		},
+		[]() {},
+		LOCTEXT("CurrentDataLayerColor_ToolTip", "Colorize actor using its Data Layer Debug Color only if the Data Layer is in the Actor Editor Context. White means multiple values, the rest is Gray.") );
 
 		auto GetActorExternalDataLayerInstance = [](AActor* InActor) -> const UDataLayerInstance*
 		{
@@ -250,14 +252,18 @@ void FDataLayersBroadcast::Initialize()
 				}
 			}
 			return FLinearColor::Gray;
-		});
+		},
+		[]() {},
+		LOCTEXT("RuntimeDataLayerColor_ToolTip", "Colorize actor using its Data Layer Debug Color only if the Data Layer is in the Actor Editor Context. White means multiple values, the rest is Gray."));
 
 		// Colorize actor using its External Data Layer Debug Color (Use gray if none)
 		FActorPrimitiveColorHandler::Get().RegisterPrimitiveColorHandler(NAME_ExternalDataLayerColor, LOCTEXT("ExternalDataLayerColor", "External Data Layer Color"), [&GetActorExternalDataLayerInstance](const UPrimitiveComponent* InPrimitiveComponent) -> FLinearColor
 		{
 			const UDataLayerInstance* ExternalDataLayerInstance = GetActorExternalDataLayerInstance(InPrimitiveComponent->GetOwner());
 			return ExternalDataLayerInstance ? ExternalDataLayerInstance->GetDebugColor() : FLinearColor::Gray;
-		});
+		},
+		[]() {},
+		LOCTEXT("ExternalDataLayerColor_ToolTip", "Colorize actor using its External Data Layer Debug Color (Use Gray if none)."));
 #endif
 	}
 }
diff --git a/Engine/Source/Editor/EditorStyle/Private/StarshipStyle.cpp b/Engine/Source/Editor/EditorStyle/Private/StarshipStyle.cpp
--- a/Engine/Source/Editor/EditorStyle/Private/StarshipStyle.cpp
+++ b/Engine/Source/Editor/EditorStyle/Private/StarshipStyle.cpp
@@ -4297,6 +4297,7 @@ void FStarshipEditorStyle::FStyle::SetupLevelEditorStyle()
 		Set( "EditorViewport.VisualizeSubstrateMode", new IMAGE_BRUSH_SVG("Starship/Common/BufferVisualization", Icon16x16 ) );
 		Set( "EditorViewport.VisualizeGroomMode",     new IMAGE_BRUSH_SVG("Starship/Common/BufferVisualization", Icon16x16 ) );
 		Set( "EditorViewport.VisualizeVirtualShadowMapMode", new IMAGE_BRUSH_SVG("Starship/Common/BufferVisualization", Icon16x16 ) );
+		Set( "EditorViewport.VisualizeActorColorationMode", new IMAGE_BRUSH("Icons/icon_ViewMode_LODColoration_16px", Icon16x16));
 		Set( "EditorViewport.CollisionPawn",          new IMAGE_BRUSH_SVG("Starship/Common/PlayerCollision", Icon16x16 ) );
 		Set( "EditorViewport.CollisionVisibility",    new IMAGE_BRUSH_SVG("Starship/Common/VisibilityCollision", Icon16x16 ) );
 		
diff --git a/Engine/Source/Editor/LevelEditor/Private/ViewportToolbar/LevelEditorViewportToolbarSections.cpp b/Engine/Source/Editor/LevelEditor/Private/ViewportToolbar/LevelEditorViewportToolbarSections.cpp
--- a/Engine/Source/Editor/LevelEditor/Private/ViewportToolbar/LevelEditorViewportToolbarSections.cpp
+++ b/Engine/Source/Editor/LevelEditor/Private/ViewportToolbar/LevelEditorViewportToolbarSections.cpp
@@ -335,7 +335,7 @@ void PopulateViewModesMenu(UToolMenu* InMenu, TSharedRef<::SLevelViewport> InVie
 				SubMenuSection.AddMenuEntry(
 					NAME_None,
 					PrimitiveColorHandler.HandlerText,
-					FText(),
+					PrimitiveColorHandler.HandlerToolTipText,
 					FSlateIcon(),
 					FUIAction(
 						FExecuteAction::CreateLambda(
@@ -343,33 +343,8 @@ void PopulateViewModesMenu(UToolMenu* InMenu, TSharedRef<::SLevelViewport> InVie
 							{
 								if (TSharedPtr<::SLevelViewport> Viewport = WeakViewport.Pin())
 								{
-									const bool bActorColorationEnabled =
-										Viewport->GetLevelViewportClient().HandleIsShowFlagEnabled(
-											FEngineShowFlags::EShowFlag::SF_ActorColoration
-										);
-
-									if (PrimitiveColorHandler.HandlerName.IsNone())
-									{
-										if (bActorColorationEnabled)
-										{
-											Viewport->GetLevelViewportClient().HandleToggleShowFlag(
-												FEngineShowFlags::EShowFlag::SF_ActorColoration
-											);
-										}
-									}
-									else
-									{
-										if (!bActorColorationEnabled)
-										{
-											Viewport->GetLevelViewportClient().HandleToggleShowFlag(
-												FEngineShowFlags::EShowFlag::SF_ActorColoration
-											);
-										}
-
-										FActorPrimitiveColorHandler::Get().SetActivePrimitiveColorHandler(
-											PrimitiveColorHandler.HandlerName, GWorld
-										);
-									}
+									FLevelEditorViewportClient& ViewportClient = Viewport->GetLevelViewportClient();
+									ViewportClient.ChangeActorColorationVisualizationMode(PrimitiveColorHandler.HandlerName);
 								}
 							}
 						),
@@ -388,26 +363,8 @@ void PopulateViewModesMenu(UToolMenu* InMenu, TSharedRef<::SLevelViewport> InVie
 							{
 								if (TSharedPtr<::SLevelViewport> Viewport = WeakViewport.Pin())
 								{
-									const bool bActorColorationEnabled =
-										Viewport->GetLevelViewportClient().HandleIsShowFlagEnabled(
-											FEngineShowFlags::EShowFlag::SF_ActorColoration
-										);
-
-									if (PrimitiveColorHandler.HandlerName.IsNone())
-									{
-										return bActorColorationEnabled ? ECheckBoxState::Unchecked
-																	   : ECheckBoxState::Checked;
-									}
-									else
-									{
-										if (bActorColorationEnabled)
-										{
-											return FActorPrimitiveColorHandler::Get().GetActivePrimitiveColorHandler()
-														== PrimitiveColorHandler.HandlerName
-													 ? ECheckBoxState::Checked
-													 : ECheckBoxState::Unchecked;
-										}
-									}
+									FLevelEditorViewportClient& ViewportClient = Viewport->GetLevelViewportClient();
+									return ViewportClient.IsActorColorationVisualizationModeSelected(PrimitiveColorHandler.HandlerName) ? ECheckBoxState::Checked : ECheckBoxState::Unchecked;
 								}
 
 								return ECheckBoxState::Unchecked;
@@ -421,9 +378,9 @@ void PopulateViewModesMenu(UToolMenu* InMenu, TSharedRef<::SLevelViewport> InVie
 
 		FToolMenuSection& Section = InMenu->FindOrAddSection("ViewMode");
 		Section.AddSubMenu(
-			"ActorColoration",
-			LOCTEXT("ActorColorationDisplayName", "Actor Coloration"),
-			LOCTEXT("ActorColorationMenu_ToolTip", "Override Actor Coloration mode"),
+			"VisualizeActorColorationViewMode",
+			LOCTEXT("VisualizeActorColorationViewModeDisplayName", "Actor Coloration"),
+			LOCTEXT("ActorColorationVisualizationMenu_ToolTip", "Select a mode for actor coloration visualization."),
 			FNewToolMenuDelegate::CreateLambda(BuildActorColorationMenu),
 			FUIAction(
 				FExecuteAction(),
@@ -431,19 +388,16 @@ void PopulateViewModesMenu(UToolMenu* InMenu, TSharedRef<::SLevelViewport> InVie
 				FIsActionChecked::CreateLambda(
 					[WeakViewport = InViewport.ToWeakPtr()]()
 					{
-						if (const TSharedPtr<::SLevelViewport> Viewport = WeakViewport.Pin())
-						{
-							return Viewport->GetLevelViewportClient().HandleIsShowFlagEnabled(
-								FEngineShowFlags::EShowFlag::SF_ActorColoration
-							);
-						}
-						return false;
+						const TSharedPtr<::SLevelViewport> Viewport = WeakViewport.Pin();
+						check(Viewport.IsValid());
+						FLevelEditorViewportClient& ViewportClient = Viewport->GetLevelViewportClient();
+						return ViewportClient.IsViewModeEnabled(VMI_VisualizeActorColoration);
 					}
 				)
 			),
 			EUserInterfaceActionType::RadioButton,
-			/*bInOpenSubMenuOnClick=*/false,
-			FSlateIcon(FAppStyle::GetAppStyleSetName(), "EditorViewport.LODColorationMode")
+			/* bInOpenSubMenuOnClick = */ false,
+			FSlateIcon(FAppStyle::GetAppStyleSetName(), "EditorViewport.VisualizeActorColorationMode")
 		);
 	}
 
diff --git a/Engine/Source/Editor/UnrealEd/Private/EditorEngine.cpp b/Engine/Source/Editor/UnrealEd/Private/EditorEngine.cpp
--- a/Engine/Source/Editor/UnrealEd/Private/EditorEngine.cpp
+++ b/Engine/Source/Editor/UnrealEd/Private/EditorEngine.cpp
@@ -609,12 +609,12 @@ UEditorEngine::UEditorEngine(const FObjectInitializer& ObjectInitializer)
 			{
 				const FString EmptyString;
 				SetPropertyColorationTarget(GWorld, EmptyString, nullptr, nullptr, nullptr);
-			});
+			},
+			LOCTEXT("PropertyColor_ToopTip", "Colorize actor if property matches. Red means a property match, otherwise the color is White."));
 
 		for (const FPropertyColorCustomProperty& PropertyColorCustomProperty : GetDefault<UPropertyColorSettings>()->CustomProperties)
 		{
-			FActorPrimitiveColorHandler::Get().RegisterPrimitiveColorHandler(PropertyColorCustomProperty.Name, 
-				FInternationalization::ForUseOnlyByLocMacroAndGraphNodeTextLiterals_CreateText(*PropertyColorCustomProperty.Text, TEXT("PropertyColor"), *PropertyColorCustomProperty.Name.ToString()),
+			FActorPrimitiveColorHandler::Get().RegisterPrimitiveColorHandler(PropertyColorCustomProperty.Name, PropertyColorCustomProperty.Text,
 				[this, PropertyColorCustomProperty](const UPrimitiveComponent* InPrimitiveComponent)
 				{
 					if (AActor* Actor = InPrimitiveComponent->GetOwner())
@@ -657,7 +657,8 @@ UEditorEngine::UEditorEngine(const FObjectInitializer& ObjectInitializer)
 							SetPropertyColorationTarget(GWorld, PropertyColorCustomProperty.PropertyValue, PropertyChain->GetTail()->GetValue(), AActor::StaticClass(), &PropertyChain);
 						}
 					}
-				});
+				},
+				PropertyColorCustomProperty.TextToolTip);
 		}
 	}
 #endif
diff --git a/Engine/Source/Editor/UnrealEd/Private/EditorViewportClient.cpp b/Engine/Source/Editor/UnrealEd/Private/EditorViewportClient.cpp
--- a/Engine/Source/Editor/UnrealEd/Private/EditorViewportClient.cpp
+++ b/Engine/Source/Editor/UnrealEd/Private/EditorViewportClient.cpp
@@ -2803,6 +2803,23 @@ FText FEditorViewportClient::GetCurrentVirtualShadowMapVisualizationModeDisplayN
 	return GetVirtualShadowMapVisualizationData().GetModeDisplayName(CurrentVirtualShadowMapVisualizationMode);
 }
 
+void FEditorViewportClient::ChangeActorColorationVisualizationMode(FName InName)
+{
+	SetViewMode(VMI_VisualizeActorColoration);
+	FActorPrimitiveColorHandler::Get().SetActivePrimitiveColorHandler(InName, GWorld);
+}
+
+bool FEditorViewportClient::IsActorColorationVisualizationModeSelected(FName InName) const
+{
+	return IsViewModeEnabled(VMI_VisualizeActorColoration) && FActorPrimitiveColorHandler::Get().GetActivePrimitiveColorHandler() == InName;
+}
+
+FText FEditorViewportClient::GetCurrentActorColorationVisualizationModeDisplayName() const
+{
+	checkf(IsViewModeEnabled(VMI_VisualizeActorColoration), TEXT("In order to call GetCurrentActorColorationVisualizationModeDisplayName(), first you must set ViewMode to VMI_VisualizeActorColoration."));
+	return FActorPrimitiveColorHandler::Get().GetActivePrimitiveColorHandlerDisplayName();
+}
+
 void FEditorViewportClient::ChangeSubstrateVisualizationMode(FName InName)
 {
 	SetViewMode(VMI_VisualizeSubstrate);
diff --git a/Engine/Source/Editor/UnrealEd/Private/PropertyColorSettings.h b/Engine/Source/Editor/UnrealEd/Private/PropertyColorSettings.h
--- a/Engine/Source/Editor/UnrealEd/Private/PropertyColorSettings.h
+++ b/Engine/Source/Editor/UnrealEd/Private/PropertyColorSettings.h
@@ -16,7 +16,10 @@ struct FPropertyColorCustomProperty
 	FName Name;
 
 	UPROPERTY(Config)
-	FString Text;
+	FText Text;
+
+	UPROPERTY(Config)
+	FText TextToolTip;
 
 	UPROPERTY(Config)
 	FString PropertyChain;
diff --git a/Engine/Source/Editor/UnrealEd/Private/ViewportToolbar/UnrealEdViewportToolbar.cpp b/Engine/Source/Editor/UnrealEd/Private/ViewportToolbar/UnrealEdViewportToolbar.cpp
--- a/Engine/Source/Editor/UnrealEd/Private/ViewportToolbar/UnrealEdViewportToolbar.cpp
+++ b/Engine/Source/Editor/UnrealEd/Private/ViewportToolbar/UnrealEdViewportToolbar.cpp
@@ -1010,6 +1010,10 @@ FText GetViewModesSubmenuLabel(TWeakPtr<SEditorViewport> InViewport)
 		{
 			Label = ViewportClient->GetCurrentVirtualShadowMapVisualizationModeDisplayName();
 		}
+		else if (ViewMode == VMI_VisualizeActorColoration)
+		{
+			Label = ViewportClient->GetCurrentActorColorationVisualizationModeDisplayName();
+		}
 		else if (ViewMode == VMI_VisualizeGPUSkinCache)
 		{
 			Label = ViewportClient->GetCurrentGPUSkinCacheVisualizationModeDisplayName();
diff --git a/Engine/Source/Editor/UnrealEd/Public/EditorViewportClient.h b/Engine/Source/Editor/UnrealEd/Public/EditorViewportClient.h
--- a/Engine/Source/Editor/UnrealEd/Public/EditorViewportClient.h
+++ b/Engine/Source/Editor/UnrealEd/Public/EditorViewportClient.h
@@ -1430,6 +1430,26 @@ public:
 	*/
 	UNREALED_API FText GetCurrentVirtualShadowMapVisualizationModeDisplayName() const;
 
+	/**
+	* Changes the actor coloration visualization mode for this viewport.
+	*
+	* @param InName	The ID of the required visualization mode
+	*/
+	UNREALED_API void ChangeActorColorationVisualizationMode(FName InName);
+
+	/**
+	* Checks if an actor coloration visualization mode is selected.
+	*
+	* @param InName	The ID of the required visualization mode
+	* @return	true if the supplied actor coloration visualization mode is checked
+	*/
+	UNREALED_API bool IsActorColorationVisualizationModeSelected(FName InName) const;
+
+	/**
+	* Returns the FText display name associated with current actor coloration visualization mode.
+	*/
+	UNREALED_API FText GetCurrentActorColorationVisualizationModeDisplayName() const;
+
 	/**
 	* Returns whether visualize debug material is enabled.
 	*/
diff --git a/Engine/Source/Editor/WorldPartitionEditor/Private/WorldPartition/ContentBundle/ContentBundleEditorSubsystem.cpp b/Engine/Source/Editor/WorldPartitionEditor/Private/WorldPartition/ContentBundle/ContentBundleEditorSubsystem.cpp
--- a/Engine/Source/Editor/WorldPartitionEditor/Private/WorldPartition/ContentBundle/ContentBundleEditorSubsystem.cpp
+++ b/Engine/Source/Editor/WorldPartitionEditor/Private/WorldPartition/ContentBundle/ContentBundleEditorSubsystem.cpp
@@ -44,7 +44,9 @@ void UContentBundleEditingSubmodule::DoInitialize()
 			}
 		}
 		return FLinearColor::Gray;
-	});
+	},
+	[](){}, 
+	LOCTEXT("ContentBundleColor_ToolTip", "Colorize actor using a random color per content bundle, the rest is Gray.") );
 #endif
 }
 
diff --git a/Engine/Source/Editor/WorldPartitionEditor/Private/WorldPartition/HLOD/HLODEditorSubsystem.cpp b/Engine/Source/Editor/WorldPartitionEditor/Private/WorldPartition/HLOD/HLODEditorSubsystem.cpp
--- a/Engine/Source/Editor/WorldPartitionEditor/Private/WorldPartition/HLOD/HLODEditorSubsystem.cpp
+++ b/Engine/Source/Editor/WorldPartitionEditor/Private/WorldPartition/HLOD/HLODEditorSubsystem.cpp
@@ -53,7 +53,9 @@ UWorldPartitionHLODEditorSubsystem::UWorldPartitionHLODEditorSubsystem()
 				}
 			}
 			return FLinearColor::Red;
-		});
+		},
+		[]() {},
+		LOCTEXT("HLODRelevantColor_ToolTip", "Colorize actor if relevant to the HLOD system. Green means relevant, otherwise the color is Red."));
 	}
 #endif
 
diff --git a/Engine/Source/Runtime/Engine/Classes/Engine/EngineBaseTypes.h b/Engine/Source/Runtime/Engine/Classes/Engine/EngineBaseTypes.h
--- a/Engine/Source/Runtime/Engine/Classes/Engine/EngineBaseTypes.h
+++ b/Engine/Source/Runtime/Engine/Classes/Engine/EngineBaseTypes.h
@@ -928,12 +928,14 @@ enum EViewModeIndex : int
 	/** Visualize Groom debug views */
 	VMI_VisualizeGroom = 35 UMETA(DisplayName = "Groom Visualization"),
 
-
 	VMI_LWCComplexity = 36 UMETA(DisplayName = "Material LWC Function Usage"),
 	
 	/** Lit Wireframe. */
 	VMI_Lit_Wireframe = 37 UMETA(DisplayName = "Lit Wireframe"),
 
+	/** Visualize Actor Coloration. */
+	VMI_VisualizeActorColoration = 38 UMETA(DisplayName = "Actor Coloration Visualization"),
+
 	VMI_Max UMETA(Hidden),
 
 	// VMI_Unknown - The value assigned to VMI_Unknown must be the highest possible of any member of EViewModeIndex, or GetViewModeName might seg-fault
diff --git a/Engine/Source/Runtime/Engine/Classes/GameFramework/ActorPrimitiveColorHandler.h b/Engine/Source/Runtime/Engine/Classes/GameFramework/ActorPrimitiveColorHandler.h
--- a/Engine/Source/Runtime/Engine/Classes/GameFramework/ActorPrimitiveColorHandler.h
+++ b/Engine/Source/Runtime/Engine/Classes/GameFramework/ActorPrimitiveColorHandler.h
@@ -22,9 +22,10 @@ class ENGINE_API FActorPrimitiveColorHandler
 public:
 	struct FPrimitiveColorHandler
 	{
-		FPrimitiveColorHandler(FName InHandlerName, FText InHandlerText, bool bInAvailalbleInEditor, const FGetColorFunc& InGetColorFunc, const FActivateFunc& InActivateFunc)
+		FPrimitiveColorHandler(FName InHandlerName, FText InHandlerText, bool bInAvailalbleInEditor, const FGetColorFunc& InGetColorFunc, const FActivateFunc& InActivateFunc, FText InHandlerToolTipText = FText())
 			: HandlerName(InHandlerName)
 			, HandlerText(InHandlerText)
+			, HandlerToolTipText(InHandlerToolTipText)
 			, bAvailalbleInEditor(bInAvailalbleInEditor)
 			, GetColorFunc(InGetColorFunc)
 			, ActivateFunc(InActivateFunc)
@@ -32,6 +33,7 @@ public:
 
 		FName HandlerName;
 		FText HandlerText;
+		FText HandlerToolTipText;
 		bool bAvailalbleInEditor;
 		FGetColorFunc GetColorFunc;
 		FActivateFunc ActivateFunc;
@@ -40,14 +42,16 @@ public:
 	FActorPrimitiveColorHandler();
 	static FActorPrimitiveColorHandler& Get();
 
-	void RegisterPrimitiveColorHandler(FName InHandlerName, FText InHandlerText, const FGetColorFunc& InHandlerFunc, const FActivateFunc& InActivateFunc = []() {});
-	void RegisterPrimitiveColorHandler(FName InHandlerName, FText InHandlerText, bool bInAvailalbleInEditor, const FGetColorFunc& InHandlerFunc, const FActivateFunc& InActivateFunc = []() {});
+	void RegisterPrimitiveColorHandler(FName InHandlerName, FText InHandlerText, const FGetColorFunc& InHandlerFunc, const FActivateFunc& InActivateFunc = []() {}, FText InHandlerToolTipText = FText());
+	void RegisterPrimitiveColorHandler(FName InHandlerName, FText InHandlerText, bool bInAvailalbleInEditor, const FGetColorFunc& InHandlerFunc, const FActivateFunc& InActivateFunc = []() {}, FText InHandlerToolTipText = FText());
 	void UnregisterPrimitiveColorHandler(FName InHandlerName);
 	void GetRegisteredPrimitiveColorHandlers(TArray<FPrimitiveColorHandler>& OutPrimitiveColorHandlers) const;
 
 	FName GetActivePrimitiveColorHandler() const;
 	bool SetActivePrimitiveColorHandler(FName InHandlerName, UWorld* InWorld);
 
+	FText GetActivePrimitiveColorHandlerDisplayName() const;
+
 	void RefreshPrimitiveColorHandler(FName InHandlerName, UWorld* InWorld);
 	void RefreshPrimitiveColorHandler(FName InHandlerName, const TArray<AActor*>& InActors);
 	void RefreshPrimitiveColorHandler(FName InHandlerName, const TArray<UPrimitiveComponent*>& InPrimitiveComponents);
@@ -55,9 +59,10 @@ public:
 	FLinearColor GetPrimitiveColor(const UPrimitiveComponent* InPrimitiveComponent) const;
 
 private:
+	void InitActivePrimitiveColorHandler();
+
 #if ENABLE_ACTOR_PRIMITIVE_COLOR_HANDLER
 	FName ActivePrimitiveColorHandlerName;
-	FText ActivePrimitiveColorHandlerText;
 	FPrimitiveColorHandler* ActivePrimitiveColorHandler;
 	TMap<FName, FPrimitiveColorHandler> Handlers;
 #endif
diff --git a/Engine/Source/Runtime/Engine/Private/ActorPrimitiveColorHandler.cpp b/Engine/Source/Runtime/Engine/Private/ActorPrimitiveColorHandler.cpp
--- a/Engine/Source/Runtime/Engine/Private/ActorPrimitiveColorHandler.cpp
+++ b/Engine/Source/Runtime/Engine/Private/ActorPrimitiveColorHandler.cpp
@@ -10,13 +10,15 @@
 FActorPrimitiveColorHandler::FActorPrimitiveColorHandler()
 {
 #if ENABLE_ACTOR_PRIMITIVE_COLOR_HANDLER
-	RegisterPrimitiveColorHandler(NAME_None, LOCTEXT("Disable", "Disable"), [](const UPrimitiveComponent*)
-	{
-		return FLinearColor::White;
-	});
+	InitActivePrimitiveColorHandler();
+#endif
+}
 
+void FActorPrimitiveColorHandler::InitActivePrimitiveColorHandler()
+{
+#if ENABLE_ACTOR_PRIMITIVE_COLOR_HANDLER
 	ActivePrimitiveColorHandlerName = NAME_None;
-	ActivePrimitiveColorHandler = Handlers.Find(NAME_None);
+	ActivePrimitiveColorHandler = nullptr;
 #endif
 }
 
@@ -25,17 +27,16 @@ FActorPrimitiveColorHandler& FActorPrimitiveColorHandler::Get()
 	return TLazySingleton<FActorPrimitiveColorHandler>::Get();
 }
 
-void FActorPrimitiveColorHandler::RegisterPrimitiveColorHandler(FName InHandlerName, FText InHandlerText, const FGetColorFunc& InHandlerFunc, const FActivateFunc& InActivateFunc)
+void FActorPrimitiveColorHandler::RegisterPrimitiveColorHandler(FName InHandlerName, FText InHandlerText, const FGetColorFunc& InHandlerFunc, const FActivateFunc& InActivateFunc, FText InHandlerToolTipText)
 {
-	RegisterPrimitiveColorHandler(InHandlerName, InHandlerText, true, InHandlerFunc, InActivateFunc);
+	RegisterPrimitiveColorHandler(InHandlerName, InHandlerText, true, InHandlerFunc, InActivateFunc, InHandlerToolTipText);
 }
 
-void FActorPrimitiveColorHandler::RegisterPrimitiveColorHandler(FName InHandlerName, FText InHandlerText, bool bInAvailalbleInEditor, const FGetColorFunc& InHandlerFunc, const FActivateFunc& InActivateFunc)
+void FActorPrimitiveColorHandler::RegisterPrimitiveColorHandler(FName InHandlerName, FText InHandlerText, bool bInAvailalbleInEditor, const FGetColorFunc& InHandlerFunc, const FActivateFunc& InActivateFunc, FText InHandlerToolTipText)
 {
 #if ENABLE_ACTOR_PRIMITIVE_COLOR_HANDLER
 	check(!Handlers.Contains(InHandlerName));
-	Handlers.Add(InHandlerName, { InHandlerName, InHandlerText, bInAvailalbleInEditor, InHandlerFunc, InActivateFunc });
-	ActivePrimitiveColorHandler = &Handlers.FindChecked(ActivePrimitiveColorHandlerName);
+	Handlers.Add(InHandlerName, { InHandlerName, InHandlerText, bInAvailalbleInEditor, InHandlerFunc, InActivateFunc, InHandlerToolTipText });
 #endif
 }
 
@@ -48,10 +49,9 @@ void FActorPrimitiveColorHandler::UnregisterPrimitiveColorHandler(FName InHandle
 	
 	if (InHandlerName == ActivePrimitiveColorHandlerName)
 	{
-		ActivePrimitiveColorHandlerName = NAME_None;
+		InitActivePrimitiveColorHandler();
 	}
 
-	ActivePrimitiveColorHandler = Handlers.Find(ActivePrimitiveColorHandlerName);
 #endif
 }
 
@@ -134,6 +134,15 @@ FName FActorPrimitiveColorHandler::GetActivePrimitiveColorHandler() const
 #endif
 }
 
+FText FActorPrimitiveColorHandler::GetActivePrimitiveColorHandlerDisplayName() const
+{
+#if ENABLE_ACTOR_PRIMITIVE_COLOR_HANDLER
+	return ActivePrimitiveColorHandler != nullptr ? ActivePrimitiveColorHandler->HandlerText : FText();
+#else
+	return FText();
+#endif
+}
+
 void FActorPrimitiveColorHandler::GetRegisteredPrimitiveColorHandlers(TArray<FPrimitiveColorHandler>& OutPrimitiveColorHandlers) const
 {
 #if ENABLE_ACTOR_PRIMITIVE_COLOR_HANDLER
@@ -144,7 +153,7 @@ void FActorPrimitiveColorHandler::GetRegisteredPrimitiveColorHandlers(TArray<FPr
 FLinearColor FActorPrimitiveColorHandler::GetPrimitiveColor(const UPrimitiveComponent* InPrimitiveComponent) const
 {
 #if ENABLE_ACTOR_PRIMITIVE_COLOR_HANDLER
-	return ActivePrimitiveColorHandler->GetColorFunc(InPrimitiveComponent);
+	return ActivePrimitiveColorHandler != nullptr ? ActivePrimitiveColorHandler->GetColorFunc(InPrimitiveComponent) : FLinearColor::White;
 #else
 	return FLinearColor::White;
 #endif
diff --git a/Engine/Source/Runtime/Engine/Private/LevelInstance/LevelInstanceSubsystem.cpp b/Engine/Source/Runtime/Engine/Private/LevelInstance/LevelInstanceSubsystem.cpp
--- a/Engine/Source/Runtime/Engine/Private/LevelInstance/LevelInstanceSubsystem.cpp
+++ b/Engine/Source/Runtime/Engine/Private/LevelInstance/LevelInstanceSubsystem.cpp
@@ -3177,7 +3177,9 @@ void ULevelInstanceSubsystem::RegisterPrimitiveColorHandler()
 				return Actor->HasLevelInstancePropertyOverrides() ? FLinearColor::Green : FLinearColor::Red;
 			}
 			return FLinearColor::White;
-		});
+		},
+		[]() {},
+		LOCTEXT("LevelInstancePropertyOverride_ToolTip", "Colorize actor with his level instance color, Green means the level has some property overrides, otherwise it will be Red. Rest is White."));
 		bPrimitiveColorHandlerRegistered = true;
 	}
 #endif
diff --git a/Engine/Source/Runtime/Engine/Private/LevelStreaming.cpp b/Engine/Source/Runtime/Engine/Private/LevelStreaming.cpp
--- a/Engine/Source/Runtime/Engine/Private/LevelStreaming.cpp
+++ b/Engine/Source/Runtime/Engine/Private/LevelStreaming.cpp
@@ -424,7 +424,9 @@ ULevelStreaming::ULevelStreaming(const FObjectInitializer& ObjectInitializer)
 				}
 			}
 			return FLinearColor::White;
-		});
+		},
+		[]() {},
+		LOCTEXT("LevelColor_ToopTip", "Colorize actor with its level color, otherwise the color is White."));
 	}
 #endif
 }
diff --git a/Engine/Source/Runtime/Engine/Private/ShowFlags.cpp b/Engine/Source/Runtime/Engine/Private/ShowFlags.cpp
--- a/Engine/Source/Runtime/Engine/Private/ShowFlags.cpp
+++ b/Engine/Source/Runtime/Engine/Private/ShowFlags.cpp
@@ -418,6 +418,7 @@ void ApplyViewMode(EViewModeIndex ViewModeIndex, bool bPerspective, FEngineShowF
 	EngineShowFlags.SetVisualizeGPUSkinCache(ViewModeIndex == VMI_VisualizeGPUSkinCache);
 	EngineShowFlags.SetVisualizeLWCComplexity(ViewModeIndex == VMI_LWCComplexity);
 	EngineShowFlags.SetMeshEdges(ViewModeIndex == VMI_Lit_Wireframe);
+	EngineShowFlags.SetActorColoration(ViewModeIndex == VMI_VisualizeActorColoration);
 }
 
 void EngineShowFlagOverride(EShowFlagInitMode ShowFlagInitMode, EViewModeIndex ViewModeIndex, FEngineShowFlags& EngineShowFlags, bool bCanDisableTonemapper)
diff --git a/Engine/Source/Runtime/Engine/Private/ViewModeNames.cpp b/Engine/Source/Runtime/Engine/Private/ViewModeNames.cpp
--- a/Engine/Source/Runtime/Engine/Private/ViewModeNames.cpp
+++ b/Engine/Source/Runtime/Engine/Private/ViewModeNames.cpp
@@ -182,6 +182,11 @@ TArray<FText> FillViewModeDisplayNames()
 			ViewModeDisplayNames.Emplace(LOCTEXT("UViewModeUtils_VMI_VisualizeVirtualShadowMap", "Virtual Shadow Map Visualization"));
 		}
 
+		else if (ViewModeIndex == VMI_VisualizeActorColoration)
+		{
+			ViewModeDisplayNames.Emplace(LOCTEXT("UViewModeUtils_VMI_VisualizeActorColoration", "Actor Coloration Visualization"));
+		}
+
 		// VMI_Max
 		else if (ViewModeIndex == VMI_Max)
 		{
@@ -207,7 +212,7 @@ const static TArray<FText> GViewModeDisplayNames = FillViewModeDisplayNames();
 FText UViewModeUtils::GetViewModeDisplayName(const EViewModeIndex ViewModeIndex)
 {
 	const FText ViewModeName = GViewModeDisplayNames[ViewModeIndex];
-	ensureMsgf(!ViewModeName.IsEmpty(), TEXT("Used an unknown value of EViewModeIndex (with value %d). Consider adding this new value in UViewModeUtils::GetViewModeName"), ViewModeIndex);
+	ensureMsgf(!ViewModeName.IsEmpty(), TEXT("Used an unknown value of EViewModeIndex (with value %d). Consider adding this new value in FillViewModeDisplayNames()"), ViewModeIndex);
 	return ViewModeName;
 }
 
@@ -396,6 +401,12 @@ TArray<const FSlateBrush*> FillViewModeDisplayIcons()
 			ViewModeDisplayIcons.Emplace(FAppStyle::Get().GetBrush("EditorViewport.VirtualTexturePendingMipsMode"));
 		}
 
+		// Actor Coloration
+		else if (ViewModeIndex == VMI_VisualizeActorColoration)
+		{
+			ViewModeDisplayIcons.Emplace(FAppStyle::Get().GetBrush("EditorViewport.VisualizeActorColorationMode"));
+		}
+
 		// VMI_Max
 		else if (ViewModeIndex == VMI_Max)
 		{
