[Viewport Toolbar] Create the "Camera" submenu for the new level editor viewport toolbar

Adding first version of Camera menu.

Camera speed, Pilot and other Camera options are not finalized yet.
Future work will also include generalizing this code so that it can be shared with Editors other than the Level one.

Bookmark menu and Camera spawn functions moved and/or refactored, so that the new and old toolbar use the same code, which matches the new design.

diff --git a/Engine/Source/Editor/LevelEditor/Private/SLevelViewport.cpp b/Engine/Source/Editor/LevelEditor/Private/SLevelViewport.cpp
--- a/Engine/Source/Editor/LevelEditor/Private/SLevelViewport.cpp
+++ b/Engine/Source/Editor/LevelEditor/Private/SLevelViewport.cpp
@@ -1987,6 +1987,13 @@ TSharedPtr<SWidget> SLevelViewport::MakeViewportToolbar()
 					RightSection.AddEntry(ViewModesSubmenu);
 				}
 
+				// Add the "Camera" submenu.
+				{
+					FToolMenuEntry CameraSubmenu = UE::LevelEditor::CreateLevelEditorViewportToolbarCameraSubmenu();
+					CameraSubmenu.InsertPosition.Position = EToolMenuInsertType::First;
+					RightSection.AddEntry(CameraSubmenu);
+				}
+
 				// Add the "Show" submenu.
 				{
 					// Stay backward-compatible with the old viewport toolbar.
@@ -2050,6 +2057,18 @@ TSharedPtr<SWidget> SLevelViewport::MakeViewportToolbar()
 				);
 				ViewportToolbarContext.AddExtender(ViewMenuExtenders);
 			}
+
+			// Stay backward-compatible with legacy options menu extenders
+			{
+				TSharedPtr<FExtender> ViewMenuExtenders = LevelEditorModule.AssembleExtenders(
+					CommandListRef, LevelEditorModule.GetAllLevelViewportOptionsMenuExtenders()
+				);
+
+				if (ViewMenuExtenders.IsValid())
+				{
+					ViewportToolbarContext.AddExtender(ViewMenuExtenders);
+				}
+			}
 		}
 
 		// Add the UnrealEd viewport toolbar context.
diff --git a/Engine/Source/Editor/LevelEditor/Private/SLevelViewportToolBar.cpp b/Engine/Source/Editor/LevelEditor/Private/SLevelViewportToolBar.cpp
--- a/Engine/Source/Editor/LevelEditor/Private/SLevelViewportToolBar.cpp
+++ b/Engine/Source/Editor/LevelEditor/Private/SLevelViewportToolBar.cpp
@@ -426,187 +426,6 @@ bool SLevelViewportToolBar::IsPerspectiveViewport() const
 	return false;
 }
 
-/**
- * Called to generate the set bookmark submenu
- */
-static void OnGenerateSetBookmarkMenu(UToolMenu* Menu, TWeakPtr<class SLevelViewport> Viewport)
-{
-	FToolMenuSection& Section = Menu->AddSection("Section");
-
-	// Add a menu entry for each bookmark
-	TSharedPtr<SLevelViewport> SharedViewport = Viewport.Pin();
-	FLevelEditorViewportClient& ViewportClient = SharedViewport->GetLevelViewportClient();
-
-	const int32 NumberOfBookmarks = static_cast<int32>(IBookmarkTypeTools::Get().GetMaxNumberOfBookmarks(&ViewportClient));
-	const int32 NumberOfMappedBookmarks = FMath::Min<int32>(AWorldSettings::NumMappedBookmarks, NumberOfBookmarks);
-
-	for( int32 BookmarkIndex = 0; BookmarkIndex < NumberOfMappedBookmarks; ++BookmarkIndex )
-	{
-		Section.AddMenuEntry(
-			NAME_None,
-			FLevelViewportCommands::Get().SetBookmarkCommands[BookmarkIndex],
-			FBookmarkUI::GetPlainLabel(BookmarkIndex)
-		);
-	}
-
-	// Only mapped bookmarks will have predefined actions.
-	// So, create any additional actions we need to hit the max number of bookmarks.
-	for (int32 BookmarkIndex = NumberOfMappedBookmarks; BookmarkIndex < NumberOfBookmarks; ++BookmarkIndex)
-	{
-		FUIAction Action;
-		Action.ExecuteAction.BindSP(SharedViewport.ToSharedRef(), &SLevelViewport::OnSetBookmark, BookmarkIndex);
-
-		Section.AddMenuEntry(
-			NAME_None,
-			FBookmarkUI::GetPlainLabel(BookmarkIndex),
-			FBookmarkUI::GetSetTooltip(BookmarkIndex),
-			FBookmarkUI::GetDefaultIcon(),
-			Action);
-	}
-}
-
-/**
- * Called to generate the clear bookmark submenu
- */
-static void OnGenerateClearBookmarkMenu(UToolMenu* Menu, TWeakPtr<class SLevelViewport> Viewport)
-{
-	FToolMenuSection& Section = Menu->AddSection("Section");
-
-	// Add a menu entry for each bookmark
-	FEditorModeTools& Tools = GLevelEditorModeTools();
-	TSharedPtr<SLevelViewport> SharedViewport = Viewport.Pin();
-	FLevelEditorViewportClient& ViewportClient = SharedViewport->GetLevelViewportClient();
-
-	const int32 NumberOfBookmarks = static_cast<int32>(IBookmarkTypeTools::Get().GetMaxNumberOfBookmarks(&ViewportClient));
-	const int32 NumberOfMappedBookmarks = FMath::Min<int32>(AWorldSettings::NumMappedBookmarks, NumberOfBookmarks);
-
-	for( int32 BookmarkIndex = 0; BookmarkIndex < NumberOfMappedBookmarks; ++BookmarkIndex )
-	{
-		if ( IBookmarkTypeTools::Get().CheckBookmark( BookmarkIndex , &ViewportClient ) )
-		{
-			Section.AddMenuEntry(
-				NAME_None,
-				FLevelViewportCommands::Get().ClearBookmarkCommands[ BookmarkIndex ],
-				FBookmarkUI::GetPlainLabel(BookmarkIndex)
-			);
-		}
-	}
-
-	for (int32 BookmarkIndex = NumberOfMappedBookmarks; BookmarkIndex < NumberOfBookmarks; ++BookmarkIndex)
-	{
-		if ( IBookmarkTypeTools::Get().CheckBookmark(BookmarkIndex, &ViewportClient) )
-		{
-			FUIAction Action;
-			Action.ExecuteAction.BindSP(SharedViewport.ToSharedRef(), &SLevelViewport::OnClearBookmark, BookmarkIndex);
-			
-			Section.AddMenuEntry(
-				NAME_None,
-				FBookmarkUI::GetPlainLabel(BookmarkIndex),
-				FBookmarkUI::GetClearTooltip(BookmarkIndex),
-				FBookmarkUI::GetDefaultIcon(),
-				Action);
-		}
-	}
-}
-
-/**
- * Called to generate the jump to bookmark menu.
- */
-static bool GenerateJumpToBookmarkMenu(UToolMenu* Menu, TWeakPtr<class SLevelViewport> Viewport)
-{
-	FToolMenuSection& Section = Menu->AddSection("Section");
-
-	// Add a menu entry for each bookmark
-	
-	FEditorModeTools& Tools = GLevelEditorModeTools();
-	TSharedPtr<SLevelViewport> SharedViewport = Viewport.Pin();
-	FLevelEditorViewportClient& ViewportClient = SharedViewport->GetLevelViewportClient();
-
-	const int32 NumberOfBookmarks = static_cast<int32>(IBookmarkTypeTools::Get().GetMaxNumberOfBookmarks(&ViewportClient));
-	const int32 NumberOfMappedBookmarks = FMath::Min<int32>(AWorldSettings::NumMappedBookmarks, NumberOfBookmarks);
-
-	bool bFoundAnyBookmarks = false;
-
-	for( int32 BookmarkIndex = 0; BookmarkIndex < NumberOfMappedBookmarks; ++BookmarkIndex )
-	{
-		if ( IBookmarkTypeTools::Get().CheckBookmark( BookmarkIndex , &ViewportClient ) )
-		{
-			bFoundAnyBookmarks = true;
-			Section.AddMenuEntry(
-				NAME_None,
-				FLevelViewportCommands::Get().JumpToBookmarkCommands[BookmarkIndex]				
-			);
-		}
-	}
-
-	for (int32 BookmarkIndex = NumberOfMappedBookmarks; BookmarkIndex < NumberOfBookmarks; ++BookmarkIndex)
-	{
-		if ( IBookmarkTypeTools::Get().CheckBookmark(BookmarkIndex, &ViewportClient) )
-		{
-			bFoundAnyBookmarks = true;
-
-			FUIAction Action;
-			Action.ExecuteAction.BindSP(SharedViewport.ToSharedRef(), &SLevelViewport::OnJumpToBookmark, BookmarkIndex);
-			
-			Section.AddMenuEntry(
-				NAME_None,
-				FBookmarkUI::GetJumpToLabel(BookmarkIndex),
-				FBookmarkUI::GetJumpToTooltip(BookmarkIndex),
-				FBookmarkUI::GetDefaultIcon(),
-				Action);
-		}
-	}
-	
-	return bFoundAnyBookmarks;
-}
-
-/**
- * Called to generate the bookmark submenu
- */
-static void OnGenerateBookmarkMenu(UToolMenu* Menu, TWeakPtr<class SLevelViewport> Viewport)
-{
-	FEditorModeTools& Tools = GLevelEditorModeTools();
-
-	// true if a bookmark was found. 
-	bool bFoundBookmark = false;
-
-	// Get the viewport client to pass down to the CheckBookmark function
-	FLevelEditorViewportClient& ViewportClient = Viewport.Pin()->GetLevelViewportClient();
-
-	bool bFoundBookmarks = false;
-	{
-		FToolMenuSection& Section = Menu->AddSection("LevelViewportActiveBoookmarks", LOCTEXT("JumpToBookmarkHeader", "Active Bookmarks"));
-		bFoundBookmarks = GenerateJumpToBookmarkMenu(Menu, Viewport);
-	}
-
-	{
-		FToolMenuSection& Section = Menu->AddSection("LevelViewportBookmarkSubmenus");
-		Section.AddSubMenu(
-			"SetBookmark",
-			LOCTEXT("SetBookmarkSubMenu", "Set Bookmark"),
-			LOCTEXT("SetBookmarkSubMenu_ToolTip", "Set viewport bookmarks"),
-			FNewToolMenuDelegate::CreateStatic( &OnGenerateSetBookmarkMenu, Viewport )
-			);
-
-		const FLevelViewportCommands& Actions = FLevelViewportCommands::Get();
-		Section.AddMenuEntry( Actions.CompactBookmarks );
-
-		if( bFoundBookmarks )
-		{
-			Section.AddSubMenu(
-				"ClearBookmark",
-				LOCTEXT("ClearBookmarkSubMenu", "Clear Bookmark"),
-				LOCTEXT("ClearBookmarkSubMenu_ToolTip", "Clear viewport bookmarks"),
-				FNewToolMenuDelegate::CreateStatic( &OnGenerateClearBookmarkMenu, Viewport ),
-				false,
-				FSlateIcon(FAppStyle::Get().GetStyleSetName(), "EditorViewport.SubMenu.Bookmarks")
-				);
-
-			Section.AddMenuEntry( Actions.ClearAllBookmarks );
-		}
-	}
-}
-
 TSharedRef<SWidget> SLevelViewportToolBar::GenerateOptionsMenu() 
 {
 	static const FName MenuName("LevelEditor.LevelViewportToolBar.Options");
@@ -679,8 +498,8 @@ void SLevelViewportToolBar::FillOptionsMenu(UToolMenu* Menu)
 
 			if( bIsPerspective )
 			{
-				Section.AddEntry(FToolMenuEntry::InitWidget("FOVAngle", GenerateFOVMenu(), LOCTEXT("FOVAngle", "Field of View (H)")));
-				Section.AddEntry(FToolMenuEntry::InitWidget("FarViewPlane", GenerateFarViewPlaneMenu(), LOCTEXT("FarViewPlane", "Far View Plane")));
+				Section.AddEntry(UE::LevelEditor::CreateFOVMenu(Viewport));
+				Section.AddEntry(UE::LevelEditor::CreateFarViewPlaneMenu(Viewport));
 			}
 
 			FEditorViewportClient& ViewportClient = Viewport.Pin()->GetLevelViewportClient();
@@ -715,7 +534,7 @@ void SLevelViewportToolBar::FillOptionsMenu(UToolMenu* Menu)
 					"Bookmark",
 					LOCTEXT("BookmarkSubMenu", "Bookmarks"),
 					LOCTEXT("BookmarkSubMenu_ToolTip", "Viewport location bookmarking"),
-					FNewToolMenuDelegate::CreateStatic(&OnGenerateBookmarkMenu, Viewport),
+					FNewToolMenuDelegate::CreateStatic(&UE::LevelEditor::CreateBookmarksMenu, Viewport),
 					false,
 					FSlateIcon(FAppStyle::Get().GetStyleSetName(), "EditorViewport.SubMenu.Bookmarks")
 				);
@@ -724,7 +543,7 @@ void SLevelViewportToolBar::FillOptionsMenu(UToolMenu* Menu)
 					"Camera",
 					LOCTEXT("CameraSubMeun", "Create Camera Here"),
 					LOCTEXT("CameraSubMenu_ToolTip", "Select a camera type to create at current viewport's location"),
-					FNewToolMenuDelegate::CreateSP(this, &SLevelViewportToolBar::GenerateCameraSpawnMenu),
+					FNewToolMenuDelegate::CreateStatic(&UE::LevelEditor::CreateCameraSpawnMenu),
 					false,
 					FSlateIcon(FAppStyle::Get().GetStyleSetName(), "EditorViewport.SubMenu.CreateCamera")
 				);
@@ -1136,100 +955,6 @@ TSharedRef<SWidget> SLevelViewportToolBar::GenerateViewModeOptionsMenu() const
 	return BuildViewModeOptionsMenu(Viewport.Pin()->GetCommandList(), ViewClient.GetViewMode(), World ? World->GetFeatureLevel() : GMaxRHIFeatureLevel, ViewClient.GetViewModeParamNameMap());
 }
 
-TSharedRef<SWidget> SLevelViewportToolBar::GenerateFOVMenu() const
-{
-	const float FOVMin = 5.f;
-	const float FOVMax = 170.f;
-
-	return
-		SNew( SBox )
-		.HAlign( HAlign_Right )
-		[
-			SNew( SBox )
-			.Padding( FMargin(4.0f, 0.0f, 0.0f, 0.0f) )
-			.WidthOverride( 100.0f )
-			[
-				SNew ( SBorder )
-				.BorderImage(FAppStyle::Get().GetBrush("Menu.WidgetBorder"))
-				.Padding(FMargin(1.0f))
-				[
-					SNew(SSpinBox<float>)
-					.Style(&FAppStyle::Get(), "Menu.SpinBox")
-					.Font( FAppStyle::GetFontStyle( TEXT( "MenuItem.Font" ) ) )
-					.MinValue(FOVMin)
-					.MaxValue(FOVMax)
-					.Value( this, &SLevelViewportToolBar::OnGetFOVValue )
-					.OnValueChanged( const_cast<SLevelViewportToolBar*>(this), &SLevelViewportToolBar::OnFOVValueChanged )
-				]
-			]
-		];
-}
-
-float SLevelViewportToolBar::OnGetFOVValue( ) const
-{
-	return Viewport.Pin()->GetLevelViewportClient().ViewFOV;
-}
-
-void SLevelViewportToolBar::OnFOVValueChanged( float NewValue )
-{
-	bool bUpdateStoredFOV = true;
-	FLevelEditorViewportClient& ViewportClient = Viewport.Pin()->GetLevelViewportClient();
-	if (ViewportClient.GetActiveActorLock().IsValid())
-	{
-		ACameraActor* CameraActor = Cast< ACameraActor >( ViewportClient.GetActiveActorLock().Get() );
-		if( CameraActor != NULL )
-		{
-			CameraActor->GetCameraComponent()->FieldOfView = NewValue;
-			bUpdateStoredFOV = false;
-		}
-	}
-
-	if ( bUpdateStoredFOV )
-	{
-		ViewportClient.FOVAngle = NewValue;
-	}
-
-	ViewportClient.ViewFOV = NewValue;
-	ViewportClient.Invalidate();
-}
-
-TSharedRef<SWidget> SLevelViewportToolBar::GenerateFarViewPlaneMenu() const
-{
-	return
-		SNew(SBox)
-		.HAlign(HAlign_Right)
-		[
-			SNew(SBox)
-			.Padding(FMargin(4.0f, 0.0f, 0.0f, 0.0f))
-			.WidthOverride(100.0f)
-			[
-				SNew ( SBorder )
-				.BorderImage(FAppStyle::Get().GetBrush("Menu.WidgetBorder"))
-				.Padding(FMargin(1.0f))
-				[
-					SNew(SSpinBox<float>)
-					.Style(&FAppStyle::Get(), "Menu.SpinBox")
-					.ToolTipText(LOCTEXT("FarViewPlaneTooltip", "Distance to use as the far view plane, or zero to enable an infinite far view plane"))
-					.MinValue(0.0f)
-					.MaxValue(100000.0f)
-					.Font(FAppStyle::GetFontStyle(TEXT("MenuItem.Font")))
-					.Value(this, &SLevelViewportToolBar::OnGetFarViewPlaneValue)
-					.OnValueChanged(const_cast<SLevelViewportToolBar*>(this), &SLevelViewportToolBar::OnFarViewPlaneValueChanged)
-				]
-			]
-		];
-}
-
-float SLevelViewportToolBar::OnGetFarViewPlaneValue() const
-{
-	return Viewport.Pin()->GetLevelViewportClient().GetFarClipPlaneOverride();
-}
-
-void SLevelViewportToolBar::OnFarViewPlaneValueChanged( float NewValue )
-{
-	Viewport.Pin()->GetLevelViewportClient().OverrideFarClipPlane(NewValue);
-}
-
 double SLevelViewportToolBar::OnGetHLODInEditorMaxDrawDistanceValue() const
 {
 	IWorldPartitionEditorModule* WorldPartitionEditorModule = FModuleManager::GetModulePtr<IWorldPartitionEditorModule>("WorldPartitionEditor");
diff --git a/Engine/Source/Editor/LevelEditor/Private/SLevelViewportToolBar.h b/Engine/Source/Editor/LevelEditor/Private/SLevelViewportToolBar.h
--- a/Engine/Source/Editor/LevelEditor/Private/SLevelViewportToolBar.h
+++ b/Engine/Source/Editor/LevelEditor/Private/SLevelViewportToolBar.h
@@ -203,28 +203,6 @@ private:
 	 */
 	TSharedRef<SWidget> GenerateViewModeOptionsMenu() const;
 
-	/**
-	 * @return The widget containing the perspective only FOV window.
-	 */
-	TSharedRef<SWidget> GenerateFOVMenu() const;
-
-	/** Called by the FOV slider in the perspective viewport to get the FOV value */
-	float OnGetFOVValue() const;
-
-	/** Called when the FOV slider is adjusted in the perspective viewport */
-	void OnFOVValueChanged( float NewValue );
-
-	/**
-	 * @return The widget containing the far view plane slider.
-	 */
-	TSharedRef<SWidget> GenerateFarViewPlaneMenu() const;
-
-	/** Called by the far view plane slider in the perspective viewport to get the far view plane value */
-	float OnGetFarViewPlaneValue() const;
-
-	/** Called when the far view plane slider is adjusted in the perspective viewport */
-	void OnFarViewPlaneValueChanged( float NewValue );
-
 	FReply OnRealtimeWarningClicked();
 	EVisibility GetRealtimeWarningVisibility() const;
 
diff --git a/Engine/Source/Editor/LevelEditor/Private/ViewportToolbar/LevelEditorViewportToolbarSections.cpp b/Engine/Source/Editor/LevelEditor/Private/ViewportToolbar/LevelEditorViewportToolbarSections.cpp
--- a/Engine/Source/Editor/LevelEditor/Private/ViewportToolbar/LevelEditorViewportToolbarSections.cpp
+++ b/Engine/Source/Editor/LevelEditor/Private/ViewportToolbar/LevelEditorViewportToolbarSections.cpp
@@ -2,9 +2,13 @@
 
 #include "ViewportToolbar/LevelEditorViewportToolbarSections.h"
 
+#include "Bookmarks/BookmarkUI.h"
+#include "Bookmarks/IBookmarkTypeTools.h"
+#include "Camera/CameraActor.h"
 #include "EditorViewportCommands.h"
 #include "FoliageType.h"
 #include "GameFramework/ActorPrimitiveColorHandler.h"
+#include "GameFramework/WorldSettings.h"
 #include "GroomVisualizationData.h"
 #include "Layers/LayersSubsystem.h"
 #include "LevelEditor.h"
@@ -128,6 +132,156 @@ void PopulateShowLayersSubmenu(UToolMenu* InMenu, TWeakPtr<::SLevelViewport> InV
 	}
 }
 
+void SetLevelViewportFOV(const TSharedRef<::SLevelViewport>& InLevelViewport, float InValue)
+{
+	bool bUpdateStoredFOV = true;
+
+	if (InLevelViewport->GetLevelViewportClient().GetActiveActorLock().IsValid())
+	{
+		if (ACameraActor* CameraActor =
+				Cast<ACameraActor>(InLevelViewport->GetLevelViewportClient().GetActiveActorLock().Get()))
+		{
+			CameraActor->GetCameraComponent()->FieldOfView = InValue;
+			bUpdateStoredFOV = false;
+		}
+	}
+
+	if (bUpdateStoredFOV)
+	{
+		InLevelViewport->GetLevelViewportClient().FOVAngle = InValue;
+	}
+
+	InLevelViewport->GetLevelViewportClient().ViewFOV = InValue;
+	InLevelViewport->GetLevelViewportClient().Invalidate();
+}
+
+void SetFarViewPlaneValue(const TSharedRef<::SLevelViewport>& InLevelViewport, float InValue)
+{
+	FLevelEditorViewportClient& ViewportClient = InLevelViewport->GetLevelViewportClient();
+	return ViewportClient.OverrideFarClipPlane(InValue);
+}
+
+float GetLevelViewportFOV(const TSharedRef<::SLevelViewport>& InLevelViewport)
+{
+	FLevelEditorViewportClient& ViewportClient = InLevelViewport->GetLevelViewportClient();
+	return ViewportClient.ViewFOV;
+}
+
+float GetFarViewPlaneValue(const TSharedRef<::SLevelViewport>& InLevelViewport)
+{
+	FLevelEditorViewportClient& ViewportClient = InLevelViewport->GetLevelViewportClient();
+	return ViewportClient.GetFarClipPlaneOverride();
+}
+
+// TODO: properly implement
+void SetCameraSpeed(const TSharedRef<::SLevelViewport>& InLevelViewport, float NewValue)
+{
+}
+
+void SetCameraSpeedScalarValue(const TSharedRef<::SLevelViewport>& InLevelViewport, float NewValue)
+{
+	if (InLevelViewport->GetViewportClient().IsValid())
+	{
+		InLevelViewport->GetViewportClient()->SetCameraSpeedScalar(NewValue);
+
+		// TODO: make sure something like this gets called if needed (e.g. future menus sharing code)
+		// Also, verify where/how to deal with this callback
+		// OnCamSpeedScalarChanged.ExecuteIfBound(NewValue);
+	}
+}
+
+// TODO: properly implement
+float GetCamSpeedSliderPosition(const TSharedRef<::SLevelViewport>& InLevelViewport)
+{
+	return 1.0f;
+}
+
+float GetCamSpeedScalarSliderPosition(const TSharedRef<::SLevelViewport>& InLevelViewport)
+{
+	float CamSpeedScalar = 1.0f;
+
+	if (InLevelViewport->GetViewportClient().IsValid())
+	{
+		CamSpeedScalar = (InLevelViewport->GetViewportClient()->GetCameraSpeedScalar());
+	}
+
+	return CamSpeedScalar;
+}
+
+bool AddJumpToBookmarkMenu(UToolMenu* InMenu, const TWeakPtr<::SLevelViewport>& InViewport)
+{
+	FToolMenuSection& Section =
+		InMenu->FindOrAddSection("JumpToBookmark", LOCTEXT("JumpToBookmarksSectionName", "Jump to Bookmark"));
+
+	// Add a menu entry for each bookmark
+	TSharedPtr<::SLevelViewport> SharedViewport = InViewport.Pin();
+	FLevelEditorViewportClient& ViewportClient = SharedViewport->GetLevelViewportClient();
+
+	const int32 NumberOfBookmarks = static_cast<int32>(IBookmarkTypeTools::Get().GetMaxNumberOfBookmarks(&ViewportClient));
+	const int32 NumberOfMappedBookmarks = FMath::Min<int32>(AWorldSettings::NumMappedBookmarks, NumberOfBookmarks);
+
+	bool bFoundAnyBookmarks = false;
+
+	for (int32 BookmarkIndex = 0; BookmarkIndex < NumberOfMappedBookmarks; ++BookmarkIndex)
+	{
+		if (IBookmarkTypeTools::Get().CheckBookmark(BookmarkIndex, &ViewportClient))
+		{
+			bFoundAnyBookmarks = true;
+			Section.AddMenuEntry(
+				NAME_None,
+				FLevelViewportCommands::Get().JumpToBookmarkCommands[BookmarkIndex],
+				FBookmarkUI::GetPlainLabel(BookmarkIndex),
+				FBookmarkUI::GetJumpToTooltip(BookmarkIndex),
+				FSlateIcon(FAppStyle::Get().GetStyleSetName(), "EditorViewport.SubMenu.Bookmarks")
+			);
+		}
+	}
+
+	return bFoundAnyBookmarks;
+}
+
+void AddClearBookmarkMenu(UToolMenu* InMenu, const TWeakPtr<::SLevelViewport>& InViewport)
+{
+	FToolMenuSection& Section = InMenu->AddSection("Section");
+
+	// Add a menu entry for each bookmark
+	// FEditorModeTools& Tools = GLevelEditorModeTools();
+	TSharedPtr<::SLevelViewport> SharedViewport = InViewport.Pin();
+	FLevelEditorViewportClient& ViewportClient = SharedViewport->GetLevelViewportClient();
+
+	const int32 NumberOfBookmarks = static_cast<int32>(IBookmarkTypeTools::Get().GetMaxNumberOfBookmarks(&ViewportClient));
+	const int32 NumberOfMappedBookmarks = FMath::Min<int32>(AWorldSettings::NumMappedBookmarks, NumberOfBookmarks);
+
+	for (int32 BookmarkIndex = 0; BookmarkIndex < NumberOfMappedBookmarks; ++BookmarkIndex)
+	{
+		if (IBookmarkTypeTools::Get().CheckBookmark(BookmarkIndex, &ViewportClient))
+		{
+			Section.AddMenuEntry(
+				NAME_None,
+				FLevelViewportCommands::Get().ClearBookmarkCommands[BookmarkIndex],
+				FBookmarkUI::GetPlainLabel(BookmarkIndex)
+			);
+		}
+	}
+
+	for (int32 BookmarkIndex = NumberOfMappedBookmarks; BookmarkIndex < NumberOfBookmarks; ++BookmarkIndex)
+	{
+		if (IBookmarkTypeTools::Get().CheckBookmark(BookmarkIndex, &ViewportClient))
+		{
+			FUIAction Action;
+			Action.ExecuteAction.BindSP(SharedViewport.ToSharedRef(), &::SLevelViewport::OnClearBookmark, BookmarkIndex);
+
+			Section.AddMenuEntry(
+				NAME_None,
+				FBookmarkUI::GetPlainLabel(BookmarkIndex),
+				FBookmarkUI::GetClearTooltip(BookmarkIndex),
+				FSlateIcon(FAppStyle::GetAppStyleSetName(), "GraphEditor.Clean"),
+				Action
+			);
+		}
+	}
+}
+
 } // namespace UE::LevelEditor::Private
 
 namespace UE::LevelEditor
@@ -817,13 +971,23 @@ FToolMenuEntry CreateShowHLODsSubmenu()
 
 								auto CreateDrawDistanceWidget = [](TSharedRef<SSpinBox<double>> InSpinBoxWidget)
 								{
-									return SNew(SBox).HAlign(HAlign_Right
-									)[SNew(SBox)
+									// clang-format off
+									return SNew(SBox)
+										.HAlign(HAlign_Right)
+										[
+											SNew(SBox)
 										  .Padding(FMargin(0.0f, 0.0f, 0.0f, 0.0f))
-										  .WidthOverride(100.0f
-										  )[SNew(SBorder)
+										  .WidthOverride(100.0f)
+											[
+												SNew(SBorder)
 												.BorderImage(FAppStyle::Get().GetBrush("Menu.WidgetBorder"))
-												.Padding(FMargin(1.0f))[InSpinBoxWidget]]];
+												.Padding(FMargin(1.0f))
+												[
+													InSpinBoxWidget
+												]
+											]
+										];
+									// clang-format on
 								};
 
 								FToolMenuEntry MinDrawDistanceMenuEntry = FToolMenuEntry::InitWidget(
@@ -1615,6 +1779,451 @@ FToolMenuEntry CreateLevelEditorViewportToolbarSettingsSubmenu()
 	);
 }
 
+void CreateCameraSpawnMenu(UToolMenu* InMenu)
+{
+	FToolMenuSection& Section = InMenu->AddSection("Section");
+	const FLevelViewportCommands& Actions = FLevelViewportCommands::Get();
+
+	for (TSharedPtr<FUICommandInfo> Camera : Actions.CreateCameras)
+	{
+		Section.AddMenuEntry(NAME_None, Camera);
+	}
+}
+
+void CreateBookmarksMenu(UToolMenu* InMenu, TWeakPtr<::SLevelViewport> InViewport)
+{
+	// Add a menu entry for each bookmark
+	TSharedPtr<::SLevelViewport> SharedViewport = InViewport.Pin();
+	FLevelEditorViewportClient& ViewportClient = SharedViewport->GetLevelViewportClient();
+
+	FToolMenuSection& ManageBookmarksSection =
+		InMenu->FindOrAddSection("ManageBookmarks", LOCTEXT("ManageBookmarkSectionName", "Manage Bookmarks"));
+
+	bool bFoundBookmarks = false;
+
+	// Jump to Bookmark Section
+	{
+		bFoundBookmarks = Private::AddJumpToBookmarkMenu(InMenu, InViewport);
+	}
+
+	// Manage Bookmarks Section
+	{
+		// Set Bookmark Submenu
+		{
+			const int32 NumberOfBookmarks =
+				static_cast<int32>(IBookmarkTypeTools::Get().GetMaxNumberOfBookmarks(&ViewportClient));
+			const int32 NumberOfMappedBookmarks = FMath::Min<int32>(AWorldSettings::NumMappedBookmarks, NumberOfBookmarks);
+
+			ManageBookmarksSection.AddSubMenu(
+				"SetBookmark",
+				LOCTEXT("SetBookmarkSubMenu", "Set Bookmark"),
+				LOCTEXT("SetBookmarkSubMenu_ToolTip", "Setting bookmarks"),
+				FNewToolMenuDelegate::CreateLambda(
+					[NumberOfMappedBookmarks](UToolMenu* InMenu)
+					{
+						const FLevelViewportCommands& Actions = FLevelViewportCommands::Get();
+
+						FToolMenuSection& SetBookmarksSection =
+							InMenu->FindOrAddSection("SetBookmark", LOCTEXT("SetBookmarkSectionName", "Set Bookmark"));
+
+						for (int32 BookmarkIndex = 0; BookmarkIndex < NumberOfMappedBookmarks; ++BookmarkIndex)
+						{
+							SetBookmarksSection.AddMenuEntry(
+								NAME_None,
+								Actions.SetBookmarkCommands[BookmarkIndex],
+								FBookmarkUI::GetPlainLabel(BookmarkIndex),
+								FBookmarkUI::GetSetTooltip(BookmarkIndex),
+								FSlateIcon(FAppStyle::Get().GetStyleSetName(), "LevelViewport.ToggleActorPilotCameraView")
+							);
+						}
+					}
+				),
+				false,
+				FSlateIcon(FSlateIcon(FAppStyle::GetAppStyleSetName(), "LevelViewport.ToggleActorPilotCameraView"))
+			);
+		}
+
+		// Manage Bookmarks Submenu
+		{
+			if (bFoundBookmarks)
+			{
+				ManageBookmarksSection.AddSubMenu(
+					"ManageBookmarks",
+					LOCTEXT("ManageBookmarksSubMenu", "Manage Bookmarks"),
+					LOCTEXT("ManageBookmarksSubMenu_ToolTip", "Bookmarks related actions"),
+					FNewToolMenuDelegate::CreateLambda(
+						[bFoundBookmarks, InViewport](UToolMenu* InMenu)
+						{
+							if (!bFoundBookmarks)
+							{
+								return;
+							}
+
+							const FLevelViewportCommands& Actions = FLevelViewportCommands::Get();
+
+							FToolMenuSection& ManageBookmarksSubsection = InMenu->FindOrAddSection(
+								"ManageBookmarks", LOCTEXT("ManageBookmarkSectionName", "Manage Bookmarks")
+							);
+
+							ManageBookmarksSubsection.AddSubMenu(
+								"ClearBookmark",
+								LOCTEXT("ClearBookmarkSubMenu", "Clear Bookmark"),
+								LOCTEXT("ClearBookmarkSubMenu_ToolTip", "Clear viewport bookmarks"),
+								FNewToolMenuDelegate::CreateLambda(&Private::AddClearBookmarkMenu, InViewport),
+								false,
+								FSlateIcon(FAppStyle::Get().GetStyleSetName(), "EditorViewport.SubMenu.Bookmarks")
+							);
+
+							FToolMenuEntry& CompactBookmarks =
+								ManageBookmarksSubsection.AddMenuEntry(Actions.CompactBookmarks);
+							CompactBookmarks.Icon =
+								FSlateIcon(FAppStyle::GetAppStyleSetName(), "AnimationEditor.ApplyCompression");
+
+							FToolMenuEntry& ClearBookmarks =
+								ManageBookmarksSubsection.AddMenuEntry(Actions.ClearAllBookmarks);
+							ClearBookmarks.Icon = FSlateIcon(FAppStyle::GetAppStyleSetName(), "GraphEditor.Clean");
+						}
+					),
+					false,
+					FSlateIcon(FAppStyle::Get().GetStyleSetName(), "EditorViewport.SubMenu.Bookmarks")
+				);
+			}
+		}
+	}
+}
+
+FToolMenuEntry CreateFOVMenu(TWeakPtr<::SLevelViewport> InLevelViewportWeak)
+{
+	constexpr float FOVMin = 5.0f;
+	constexpr float FOVMax = 170.0f;
+
+	return UnrealEd::CreateNumericEntry(
+		"FOVAngle",
+		LOCTEXT("FOVAngle", "Field of View"),
+		LOCTEXT("FOVAngleTooltip", "Field of View"),
+		FCanExecuteAction(),
+		UnrealEd::FNumericEntryExecuteActionDelegate::CreateLambda(
+			[InLevelViewportWeak](float InValue)
+			{
+				if (TSharedPtr<::SLevelViewport> LevelViewport = InLevelViewportWeak.Pin())
+				{
+					Private::SetLevelViewportFOV(LevelViewport.ToSharedRef(), InValue);
+				}
+			}
+		),
+		TAttribute<float>::CreateLambda(
+			[InLevelViewportWeak]()
+			{
+				if (TSharedPtr<::SLevelViewport> Viewport = InLevelViewportWeak.Pin())
+				{
+					return Private::GetLevelViewportFOV(Viewport.ToSharedRef());
+				}
+
+				return FOVMin;
+			}
+		),
+		FOVMin,
+		FOVMax,
+		1
+	);
+}
+
+FToolMenuEntry CreateFarViewPlaneMenu(TWeakPtr<::SLevelViewport> InInLevelViewportWeak)
+{
+	constexpr float FarMin = 0.0f;
+	constexpr float FarMax = 100000.0f;
+
+	return UnrealEd::CreateNumericEntry(
+		"FarViewPlane",
+		LOCTEXT("FarViewPlane", "Far View Plane"),
+		LOCTEXT("FarViewPlaneTooltip", "Far View Plane"),
+		FCanExecuteAction(),
+		UnrealEd::FNumericEntryExecuteActionDelegate::CreateLambda(
+			[InInLevelViewportWeak](float InValue)
+			{
+				if (TSharedPtr<::SLevelViewport> LevelViewport = InInLevelViewportWeak.Pin())
+				{
+					Private::SetFarViewPlaneValue(LevelViewport.ToSharedRef(), InValue);
+				}
+			}
+		),
+		TAttribute<float>::CreateLambda(
+			[InInLevelViewportWeak]()
+			{
+				if (TSharedPtr<::SLevelViewport> Viewport = InInLevelViewportWeak.Pin())
+				{
+					return Private::GetFarViewPlaneValue(Viewport.ToSharedRef());
+				}
+
+				return FarMax;
+			}
+		),
+		FarMin,
+		FarMax,
+		1
+	);
+}
+
+FToolMenuEntry CreateCameraSpeedSlider(TWeakPtr<::SLevelViewport> InLevelViewportWeak)
+{
+	constexpr float MinSpeed = 0.033f;
+	constexpr float MaxSpeed = 32.0f;
+
+	return UnrealEd::CreateNumericEntry(
+		"CameraSpeed",
+		LOCTEXT("CameraSpeedLabel", "Camera Speed"),
+		LOCTEXT("CameraSpeedTooltip", "Camera Speed"),
+		FCanExecuteAction(),
+		UnrealEd::FNumericEntryExecuteActionDelegate::CreateLambda(
+			[InLevelViewportWeak](float InValue)
+			{
+				if (TSharedPtr<::SLevelViewport> LevelViewport = InLevelViewportWeak.Pin())
+				{
+					Private::SetCameraSpeed(LevelViewport.ToSharedRef(), InValue);
+				}
+			}
+		),
+		TAttribute<float>::CreateLambda(
+			[InLevelViewportWeak]()
+			{
+				if (TSharedPtr<::SLevelViewport> Viewport = InLevelViewportWeak.Pin())
+				{
+					return Private::GetCamSpeedSliderPosition(Viewport.ToSharedRef());
+				}
+
+				return 1.0f;
+			}
+		),
+		MinSpeed,
+		MaxSpeed,
+		3
+	);
+}
+
+FToolMenuEntry CreateCameraSpeedScalarSlider(TWeakPtr<::SLevelViewport> InLevelViewportWeak)
+{
+	constexpr float MinSpeed = 1.0f;
+	constexpr float MaxSpeed = 128.0f;
+
+	return UnrealEd::CreateNumericEntry(
+		"CameraSpeedScalar",
+		LOCTEXT("CameraSpeedScalarLabel", "Speed Scalar"),
+		LOCTEXT("CameraSpeedScalarTooltip", "Scalar to increase camera movement range"),
+		FCanExecuteAction(),
+		UnrealEd::FNumericEntryExecuteActionDelegate::CreateLambda(
+			[InLevelViewportWeak](float InValue)
+			{
+				if (TSharedPtr<::SLevelViewport> LevelViewport = InLevelViewportWeak.Pin())
+				{
+					Private::SetCameraSpeedScalarValue(LevelViewport.ToSharedRef(), InValue);
+				}
+			}
+		),
+		TAttribute<float>::CreateLambda(
+			[InLevelViewportWeak]()
+			{
+				if (TSharedPtr<::SLevelViewport> Viewport = InLevelViewportWeak.Pin())
+				{
+					return Private::GetCamSpeedScalarSliderPosition(Viewport.ToSharedRef());
+				}
+
+				return MinSpeed;
+			}
+		),
+		MinSpeed,
+		MaxSpeed,
+		1
+	);
+}
+
+void CreateCameraSpeedMenu(UToolMenu* InMenu, const TWeakPtr<::SLevelViewport>& InLevelViewportWeak)
+{
+	FToolMenuSection& Section = InMenu->AddSection("Section");
+
+	Section.AddEntry(CreateCameraSpeedSlider(InLevelViewportWeak));
+	Section.AddEntry(CreateCameraSpeedScalarSlider(InLevelViewportWeak));
+}
+
+FToolMenuEntry CreateLevelEditorViewportToolbarCameraSubmenu()
+{
+	return FToolMenuEntry::InitSubMenu(
+		"Camera",
+		LOCTEXT("CameraSubmenuLabel", "Camera"),
+		LOCTEXT("CameraSubmenuTooltip", "Viewport-related Camera settings"),
+		FNewToolMenuDelegate::CreateLambda(
+			[](UToolMenu* Submenu) -> void
+			{
+				TWeakPtr<::SLevelViewport> LevelViewportWeak = Submenu->FindContext<ULevelViewportContext>()->LevelViewport;
+
+				// TODO:
+				// add Select Active Camera
+
+				// Perspective Section
+				{
+					FToolMenuSection& PerspectiveSection =
+						Submenu->FindOrAddSection("Perspective", LOCTEXT("PerspectiveLabel", "Perspective"));
+
+					PerspectiveSection.AddMenuEntry(FEditorViewportCommands::Get().Perspective);
+
+					// TODO: show separator based on actual list of cameras
+					// PerspectiveSection.AddSeparator("PerspectiveSeparator");
+					// Camera list
+					// Camera types
+				}
+
+				// Orthographic Section
+				{
+					FToolMenuSection& OrthographicSection =
+						Submenu->FindOrAddSection("Orthographic", LOCTEXT("OrthographicLabel", "Orthographic"));
+					OrthographicSection.AddMenuEntry(FEditorViewportCommands::Get().Top);
+					OrthographicSection.AddMenuEntry(FEditorViewportCommands::Get().Bottom);
+					OrthographicSection.AddMenuEntry(FEditorViewportCommands::Get().Left);
+					OrthographicSection.AddMenuEntry(FEditorViewportCommands::Get().Right);
+					OrthographicSection.AddMenuEntry(FEditorViewportCommands::Get().Front);
+					OrthographicSection.AddMenuEntry(FEditorViewportCommands::Get().Back);
+
+					OrthographicSection.AddSeparator("PerspectiveSeparator");
+
+					OrthographicSection.AddEntry(CreateFOVMenu(LevelViewportWeak));
+					OrthographicSection.AddEntry(CreateFarViewPlaneMenu(LevelViewportWeak));
+				}
+
+				// Create Section
+				{
+					FToolMenuSection& CreateSection = Submenu->FindOrAddSection("Create", LOCTEXT("CreateLabel", "Create"));
+
+					CreateSection.AddSubMenu(
+						"CreateCamera",
+						LOCTEXT("CameraSubMenu", "Create Camera Here"),
+						LOCTEXT("CameraSubMenu_ToolTip", "Select a camera type to create at current viewport's location"),
+						FNewToolMenuDelegate::CreateLambda(
+							[](UToolMenu* InMenu)
+							{
+								CreateCameraSpawnMenu(InMenu);
+							}
+						),
+						false,
+						FSlateIcon(FAppStyle::Get().GetStyleSetName(), "EditorViewport.SubMenu.CreateCamera")
+					);
+
+					CreateSection.AddSubMenu(
+						"Bookmarks",
+						LOCTEXT("BookmarksSubMenu", "Bookmarks"),
+						LOCTEXT("BookmarksSubMenu_ToolTip", "Bookmarks related actions"),
+						FNewToolMenuDelegate::CreateLambda(
+							[LevelViewportWeak](UToolMenu* InMenu)
+							{
+								CreateBookmarksMenu(InMenu, LevelViewportWeak);
+							}
+						),
+						false,
+						FSlateIcon(FAppStyle::Get().GetStyleSetName(), "EditorViewport.SubMenu.Bookmarks")
+					);
+				}
+
+				// Positioning Section
+				{
+					FToolMenuSection& PositioningSection =
+						Submenu->FindOrAddSection("Positioning", LOCTEXT("PositioningLabel", "Positioning"));
+
+					// Camera Speed Submenu
+					{
+						PositioningSection.AddSubMenu(
+							"CameraSpeed",
+							LOCTEXT("CameraSpeedSubMenu", "Camera Speed"),
+							LOCTEXT("CameraSpeedSubMenu_ToolTip", "Camera Speed related actions"),
+							FNewToolMenuDelegate::CreateLambda(
+								[LevelViewportWeak](UToolMenu* InMenu)
+								{
+									CreateCameraSpeedMenu(InMenu, LevelViewportWeak);
+								}
+							),
+							false,
+							FSlateIcon(FSlateIcon(FAppStyle::GetAppStyleSetName(), "LevelViewport.ToggleActorPilotCameraView"))
+						);
+					}
+
+					PositioningSection.AddSeparator("PositioningSeparator_1");
+
+					// Frame Selection
+					{
+						FToolMenuEntry FocusViewportToSelection =
+							FToolMenuEntry::InitMenuEntry(FEditorViewportCommands::Get().FocusViewportToSelection);
+						FocusViewportToSelection.UserInterfaceActionType = EUserInterfaceActionType::ToggleButton;
+						FocusViewportToSelection.Icon = FSlateIcon(FAppStyle::GetAppStyleSetName(), "Icons.FrameActor");
+						PositioningSection.AddEntry(FocusViewportToSelection);
+					}
+
+					// Pilot Submenu
+					{
+						PositioningSection.AddSubMenu(
+							"Pilot",
+							LOCTEXT("PilotSubMenu", "Pilot"),
+							LOCTEXT("PilotSubMenu_ToolTip", "Pilot related actions"),
+							FNewToolMenuDelegate::CreateLambda(
+								[](UToolMenu* InMenu)
+								{
+									// TODO add create Pilot menu
+								}
+							),
+							false,
+							FSlateIcon(FSlateIcon(FAppStyle::GetAppStyleSetName(), "LevelViewport.PilotSelectedActor"))
+						);
+					}
+
+					PositioningSection.AddSeparator("PositioningSeparator_2");
+
+					// Move Camera/Object
+					{
+						FToolMenuEntry CameraToObjectMenu =
+							FToolMenuEntry::InitMenuEntry(FLevelEditorCommands::Get().SnapCameraToObject);
+						CameraToObjectMenu.UserInterfaceActionType = EUserInterfaceActionType::ToggleButton;
+						CameraToObjectMenu.Label = LOCTEXT("CameraToObjectLabel", "Move Camera to Object");
+						PositioningSection.AddEntry(CameraToObjectMenu);
+
+						FToolMenuEntry ObjectToCameraMenu =
+							FToolMenuEntry::InitMenuEntry(FLevelEditorCommands::Get().SnapObjectToCamera);
+						ObjectToCameraMenu.UserInterfaceActionType = EUserInterfaceActionType::ToggleButton;
+						ObjectToCameraMenu.Label = LOCTEXT("ObjectToCameraLabel", "Move Object to Camera");
+						PositioningSection.AddEntry(ObjectToCameraMenu);
+					}
+				}
+
+				// Options Section
+				{
+					FToolMenuSection& OptionsSection =
+						Submenu->FindOrAddSection("CameraOptions", LOCTEXT("OptionsLabel", "Options"));
+					// add Cinematic Viewport
+					// add Allow Cinematic Control
+					// add Game View
+
+					FToolMenuEntry AllowCinematicControl =
+						FToolMenuEntry::InitMenuEntry(FLevelViewportCommands::Get().ToggleCinematicPreview);
+					AllowCinematicControl.UserInterfaceActionType = EUserInterfaceActionType::ToggleButton;
+					OptionsSection.AddEntry(AllowCinematicControl);
+
+					FToolMenuEntry ToggleGameView =
+						FToolMenuEntry::InitMenuEntry(FLevelViewportCommands::Get().ToggleGameView);
+					ToggleGameView.UserInterfaceActionType = EUserInterfaceActionType::ToggleButton;
+					OptionsSection.AddEntry(ToggleGameView);
+
+					// This additional options section is used to force certain elements to appear after extensions
+					{
+						FToolMenuSection& AdditionalOptions =
+							Submenu->FindOrAddSection("AdditionalOptions", LOCTEXT("AdditionalOptionsLabel", ""));
+						AdditionalOptions.AddSeparator("AdditionalOptionsSeparator");
+
+						FToolMenuEntry HighResolutionScreenshot =
+							FToolMenuEntry::InitMenuEntry(FLevelViewportCommands::Get().HighResScreenshot);
+						HighResolutionScreenshot.UserInterfaceActionType = EUserInterfaceActionType::ToggleButton;
+						AdditionalOptions.AddEntry(HighResolutionScreenshot);
+					}
+				}
+			}
+		)
+	);
+}
+
 } // namespace UE::LevelEditor
 
 #undef LOCTEXT_NAMESPACE
diff --git a/Engine/Source/Editor/LevelEditor/Private/ViewportToolbar/LevelEditorViewportToolbarSections.h b/Engine/Source/Editor/LevelEditor/Private/ViewportToolbar/LevelEditorViewportToolbarSections.h
--- a/Engine/Source/Editor/LevelEditor/Private/ViewportToolbar/LevelEditorViewportToolbarSections.h
+++ b/Engine/Source/Editor/LevelEditor/Private/ViewportToolbar/LevelEditorViewportToolbarSections.h
@@ -14,6 +14,15 @@ struct FToolMenuSection;
 
 namespace UE::LevelEditor
 {
+void CreateCameraSpawnMenu(UToolMenu* InMenu);
+void CreateBookmarksMenu(UToolMenu* InMenu, TWeakPtr<SLevelViewport> InViewport);
+void CreateCameraSpeedMenu(UToolMenu* InMenu, const TWeakPtr<SLevelViewport>& InLevelViewportWeak);
+
+FToolMenuEntry CreateFOVMenu(TWeakPtr<SLevelViewport> InLevelViewportWeak);
+FToolMenuEntry CreateFarViewPlaneMenu(TWeakPtr<SLevelViewport> InInLevelViewportWeak);
+FToolMenuEntry CreateFarViewPlaneMenu(TWeakPtr<SLevelViewport> InInLevelViewportWeak);
+FToolMenuEntry CreateCameraSpeedSlider(TWeakPtr<SLevelViewport> InLevelViewportWeak);
+FToolMenuEntry CreateCameraSpeedScalarSlider(TWeakPtr<SLevelViewport> InLevelViewportWeak);
 
 bool ShowViewportRealtimeWarning(FLevelEditorViewportClient& ViewportClient);
 
@@ -39,4 +48,6 @@ void GenerateViewportLayoutsMenu(UToolMenu* InMenu, TSharedPtr<::SLevelViewport>
 TSharedRef<SWidget> BuildVolumeControlCustomWidget();
 FToolMenuEntry CreateLevelEditorViewportToolbarSettingsSubmenu();
 
+FToolMenuEntry CreateLevelEditorViewportToolbarCameraSubmenu();
+
 } // namespace UE::LevelEditor
diff --git a/Engine/Source/Editor/UnrealEd/Private/ViewportToolbar/UnrealEdViewportToolbar.cpp b/Engine/Source/Editor/UnrealEd/Private/ViewportToolbar/UnrealEdViewportToolbar.cpp
--- a/Engine/Source/Editor/UnrealEd/Private/ViewportToolbar/UnrealEdViewportToolbar.cpp
+++ b/Engine/Source/Editor/UnrealEd/Private/ViewportToolbar/UnrealEdViewportToolbar.cpp
@@ -1576,6 +1576,62 @@ FToolMenuEntry CreateCheckboxSubmenu(
 	return CheckBoxSubmenu;
 }
 
+FToolMenuEntry CreateNumericEntry(
+	const FName InName,
+	const FText& InLabel,
+	const FText& InTooltip,
+	const FCanExecuteAction& InCanExecuteAction,
+	const FNumericEntryExecuteActionDelegate& InOnValueChanged,
+	const TAttribute<float>& InGetValue,
+	float InMinValue,
+	float InMaxValue,
+	int32 InMaxFractionalDigits
+)
+{
+	const FMargin WidgetsMargin(2.0f, 0.0f, 3.0f, 0.0f);
+
+	FToolMenuEntry NumericEntry = FToolMenuEntry::InitMenuEntry(
+		InName,
+		FUIAction(FExecuteAction(), InCanExecuteAction),
+		// clang-format off
+		SNew(SHorizontalBox)
+		+ SHorizontalBox::Slot()
+		.VAlign(VAlign_Center)
+		.HAlign(HAlign_Left)
+		.Padding(WidgetsMargin)
+		.AutoWidth()
+		[
+			SNew(STextBlock)
+			.Text(InLabel)
+		]
+		+ SHorizontalBox::Slot()
+		.VAlign(VAlign_Center)
+		.HAlign(HAlign_Right)
+		.Padding(FMargin(6.0f, 0))
+		.FillContentWidth(1.0)
+		[
+			SNew(SBox)
+			.Padding(WidgetsMargin)
+			.MinDesiredWidth(80.0f)
+			[
+				SNew(SNumericEntryBox<float>)
+				.ToolTipText(InTooltip)
+				.MinValue(InMinValue)
+				.MaxValue(InMaxValue)
+				.MaxSliderValue(InMaxValue)
+				.AllowSpin(true)
+				.MaxFractionalDigits(InMaxFractionalDigits)
+				.Font(FAppStyle::GetFontStyle(TEXT("MenuItem.Font")))
+				.OnValueChanged_Lambda([InOnValueChanged](float InValue){ InOnValueChanged.Execute(InValue); })
+				.Value_Lambda([InGetValue](){ return InGetValue.Get(); })
+			]
+		]
+		// clang-format on
+	);
+
+	return NumericEntry;
+}
+
 } // namespace UE::UnrealEd
 
 #undef LOCTEXT_NAMESPACE
diff --git a/Engine/Source/Editor/UnrealEd/Public/ViewportToolbar/UnrealEdViewportToolbar.h b/Engine/Source/Editor/UnrealEd/Public/ViewportToolbar/UnrealEdViewportToolbar.h
--- a/Engine/Source/Editor/UnrealEd/Public/ViewportToolbar/UnrealEdViewportToolbar.h
+++ b/Engine/Source/Editor/UnrealEd/Public/ViewportToolbar/UnrealEdViewportToolbar.h
@@ -54,6 +54,8 @@ DECLARE_DELEGATE_RetVal_OneParam(bool, FLocationGridCheckboxListIsCheckedDelegat
 DECLARE_DELEGATE_OneParam(FScaleGridCheckboxListExecuteActionDelegate, int);
 DECLARE_DELEGATE_RetVal_OneParam(bool, FScaleGridCheckboxListIsCheckedDelegate, int);
 
+DECLARE_DELEGATE_OneParam(FNumericEntryExecuteActionDelegate, float);
+
 UNREALED_API TSharedRef<SWidget> BuildRotationGridCheckBoxList(
 	FName InExtentionHook,
 	const FText& InHeading,
@@ -99,4 +101,16 @@ UNREALED_API FToolMenuEntry CreateCheckboxSubmenu(
 	const FNewToolMenuChoice& InMakeMenu
 );
 
+UNREALED_API FToolMenuEntry CreateNumericEntry(
+	const FName InName,
+	const FText& InLabel,
+	const FText& InTooltip,
+	const FCanExecuteAction& InCanExecuteAction,
+	const FNumericEntryExecuteActionDelegate& InOnValueChanged,
+	const TAttribute<float>& InGetValue,
+	float InMinValue = 0.0f,
+	float InMaxValue = 1.0f,
+	int32 InMaxFractionalDigits = 2
+);
+
 } // namespace UE::UnrealEd
