Hlod in Editor
* Fixed :  Setting values for Min or Max Draw Distance property of HLODs updates both properties in Editor

diff --git a/Engine/Source/Editor/LevelEditor/Private/ViewportToolbar/LevelEditorViewportToolbarSections.cpp b/Engine/Source/Editor/LevelEditor/Private/ViewportToolbar/LevelEditorViewportToolbarSections.cpp
--- a/Engine/Source/Editor/LevelEditor/Private/ViewportToolbar/LevelEditorViewportToolbarSections.cpp
+++ b/Engine/Source/Editor/LevelEditor/Private/ViewportToolbar/LevelEditorViewportToolbarSections.cpp
@@ -940,18 +940,40 @@ FToolMenuEntry CreateShowHLODsSubmenu()
 											}
 										);
 
+								// double SLevelViewportToolBar::OnGetHLODInEditorMaxDrawDistanceValue() const
+								auto OnGetHLODInEditorMaxDrawDistanceValue = []() -> double
+								{
+									IWorldPartitionEditorModule* WorldPartitionEditorModule =
+										FModuleManager::GetModulePtr<IWorldPartitionEditorModule>("WorldPartitionEditor");
+									return WorldPartitionEditorModule
+										? WorldPartitionEditorModule->GetHLODInEditorMaxDrawDistance()
+										: 0;
+								};
+
+								// void SLevelViewportToolBar::OnHLODInEditorMaxDrawDistanceValueChanged(double NewValue) const
+								auto OnHLODInEditorMaxDrawDistanceValueChanged = [](double NewValue) -> void
+								{
+									IWorldPartitionEditorModule* WorldPartitionEditorModule =
+										FModuleManager::GetModulePtr<IWorldPartitionEditorModule>("WorldPartitionEditor");
+									if (WorldPartitionEditorModule)
+									{
+										WorldPartitionEditorModule->SetHLODInEditorMaxDrawDistance(NewValue);
+										GEditor->RedrawLevelEditingViewports(true);
+									}
+								};
+
 								TSharedRef<SSpinBox<double>> MaxDrawDistanceSpinBox =
 									SNew(SSpinBox<double>)
 										.MinValue(MaxDrawDistanceMinValue)
 										.MaxValue(MaxDrawDistanceMaxValue)
 										.IsEnabled(bHLODInEditorAllowed)
-										.Value_Lambda(OnGetHLODInEditorMinDrawDistanceValue)
-										.OnValueChanged_Lambda(OnHLODInEditorMinDrawDistanceValueChanged)
+										.Value_Lambda(OnGetHLODInEditorMaxDrawDistanceValue)
+										.OnValueChanged_Lambda(OnHLODInEditorMaxDrawDistanceValueChanged)
 										.ToolTipText(
 											bHLODInEditorAllowed
 												? LOCTEXT(
 													  "HLODsInEditor_MaxDrawDistance_Tooltip",
-													  "Sets the maximum distance at which HLODs will be rendered"
+													  "Sets the maximum distance at which HLODs will be rendered (0.0 means infinite)"
 												  )
 												: HLODInEditorDisallowedReason
 										)
