[Viewport Toolbar] Remove unnecessary Viewport argument from menu entries API

Removing Viewport argument from methods which can retrieve it from the available ViewportContext.

diff --git a/Engine/Source/Editor/LevelEditor/Private/SLevelViewportToolBar.cpp b/Engine/Source/Editor/LevelEditor/Private/SLevelViewportToolBar.cpp
--- a/Engine/Source/Editor/LevelEditor/Private/SLevelViewportToolBar.cpp
+++ b/Engine/Source/Editor/LevelEditor/Private/SLevelViewportToolBar.cpp
@@ -534,7 +534,7 @@ void SLevelViewportToolBar::FillOptionsMenu(UToolMenu* Menu)
 					"Bookmark",
 					LOCTEXT("BookmarkSubMenu", "Bookmarks"),
 					LOCTEXT("BookmarkSubMenu_ToolTip", "Viewport location bookmarking"),
-					FNewToolMenuDelegate::CreateStatic(&UE::LevelEditor::CreateBookmarksMenu, Viewport),
+					FNewToolMenuDelegate::CreateStatic(&UE::LevelEditor::CreateBookmarksMenu),
 					false,
 					FSlateIcon(FAppStyle::Get().GetStyleSetName(), "EditorViewport.SubMenu.Bookmarks")
 				);
@@ -556,8 +556,10 @@ void SLevelViewportToolBar::FillOptionsMenu(UToolMenu* Menu)
 		{
 			FToolMenuSection& Section = Menu->AddSection("LevelViewportLayouts");
 			Section.AddSubMenu("Configs", LOCTEXT("ConfigsSubMenu", "Layouts"), FText::GetEmpty(),
-				FNewToolMenuDelegate::CreateLambda([WeakViewport = Viewport](UToolMenu* InMenu) {
-					UE::LevelEditor::GenerateViewportLayoutsMenu(InMenu, WeakViewport.Pin());
+				FNewToolMenuDelegate::CreateLambda(
+					[](UToolMenu* InMenu)
+					{
+						UE::LevelEditor::GenerateViewportLayoutsMenu(InMenu);
 				}),
 				false, FSlateIcon(FAppStyle::Get().GetStyleSetName(), "Icons.Layout"));
 		}
@@ -982,10 +984,7 @@ TWeakObjectPtr<UWorld> SLevelViewportToolBar::GetWorld() const
 
 void SLevelViewportToolBar::FillViewMenu(UToolMenu* Menu)
 {
-	if (TSharedPtr<SLevelViewport> ViewportPinned = Viewport.Pin())
-	{
-		UE::LevelEditor::PopulateViewModesMenu(Menu, ViewportPinned.ToSharedRef());
-	}
+	UE::LevelEditor::PopulateViewModesMenu(Menu);
 }
 
 
diff --git a/Engine/Source/Editor/LevelEditor/Private/ViewportToolbar/LevelEditorViewportToolbarSections.cpp b/Engine/Source/Editor/LevelEditor/Private/ViewportToolbar/LevelEditorViewportToolbarSections.cpp
--- a/Engine/Source/Editor/LevelEditor/Private/ViewportToolbar/LevelEditorViewportToolbarSections.cpp
+++ b/Engine/Source/Editor/LevelEditor/Private/ViewportToolbar/LevelEditorViewportToolbarSections.cpp
@@ -298,10 +298,22 @@ TSharedPtr<FExtender> GetViewModesLegacyExtenders()
 	return LevelEditorModule.GetMenuExtensibilityManager()->GetAllExtenders();
 }
 
-void PopulateViewModesMenu(UToolMenu* InMenu, TSharedRef<::SLevelViewport> InViewport)
+void PopulateViewModesMenu(UToolMenu* InMenu)
 {
 	FToolMenuInsert InsertPosition("ViewMode", EToolMenuInsertType::After);
 
+	ULevelViewportContext* const LevelViewportContext = InMenu->FindContext<ULevelViewportContext>();
+	if (!LevelViewportContext)
+	{
+		return;
+	}
+
+	const TSharedPtr<::SLevelViewport> LevelViewport = LevelViewportContext->LevelViewport.Pin();
+	if (!LevelViewport)
+	{
+		return;
+	}
+
 	{
 		FToolMenuSection& Section = InMenu->AddSection(
 			"LevelViewportDeferredRendering", LOCTEXT("DeferredRenderingHeader", "Deferred Rendering"), InsertPosition
@@ -319,7 +331,7 @@ void PopulateViewModesMenu(UToolMenu* InMenu, TSharedRef<::SLevelViewport> InVie
 				FExecuteAction(),
 				FCanExecuteAction(),
 				FIsActionChecked::CreateLambda(
-					[WeakViewport = InViewport.ToWeakPtr()]()
+					[WeakViewport = LevelViewport.ToWeakPtr()]()
 					{
 						const TSharedPtr<::SLevelViewport> Viewport = WeakViewport.Pin();
 						check(Viewport.IsValid());
@@ -345,7 +357,7 @@ void PopulateViewModesMenu(UToolMenu* InMenu, TSharedRef<::SLevelViewport> InVie
 				FExecuteAction(),
 				FCanExecuteAction(),
 				FIsActionChecked::CreateLambda(
-					[WeakViewport = InViewport.ToWeakPtr()]()
+					[WeakViewport = LevelViewport.ToWeakPtr()]()
 					{
 						const TSharedPtr<::SLevelViewport> Viewport = WeakViewport.Pin();
 						check(Viewport.IsValid());
@@ -371,7 +383,7 @@ void PopulateViewModesMenu(UToolMenu* InMenu, TSharedRef<::SLevelViewport> InVie
 				FExecuteAction(),
 				FCanExecuteAction(),
 				FIsActionChecked::CreateLambda(
-					[WeakViewport = InViewport.ToWeakPtr()]()
+					[WeakViewport = LevelViewport.ToWeakPtr()]()
 					{
 						const TSharedPtr<::SLevelViewport> Viewport = WeakViewport.Pin();
 						check(Viewport.IsValid());
@@ -398,7 +410,7 @@ void PopulateViewModesMenu(UToolMenu* InMenu, TSharedRef<::SLevelViewport> InVie
 				FExecuteAction(),
 				FCanExecuteAction(),
 				FIsActionChecked::CreateLambda(
-					[WeakViewport = InViewport.ToWeakPtr()]()
+					[WeakViewport = LevelViewport.ToWeakPtr()]()
 					{
 						const TSharedPtr<::SLevelViewport> Viewport = WeakViewport.Pin();
 						check(Viewport.IsValid());
@@ -425,7 +437,7 @@ void PopulateViewModesMenu(UToolMenu* InMenu, TSharedRef<::SLevelViewport> InVie
 				FExecuteAction(),
 				FCanExecuteAction(),
 				FIsActionChecked::CreateLambda(
-					[WeakViewport = InViewport.ToWeakPtr()]()
+					[WeakViewport = LevelViewport.ToWeakPtr()]()
 					{
 						const TSharedPtr<::SLevelViewport> Viewport = WeakViewport.Pin();
 						check(Viewport.IsValid());
@@ -455,7 +467,7 @@ void PopulateViewModesMenu(UToolMenu* InMenu, TSharedRef<::SLevelViewport> InVie
 				FExecuteAction(),
 				FCanExecuteAction(),
 				FIsActionChecked::CreateLambda(
-					[WeakViewport = InViewport.ToWeakPtr()]()
+					[WeakViewport = LevelViewport.ToWeakPtr()]()
 					{
 						const TSharedPtr<::SLevelViewport> Viewport = WeakViewport.Pin();
 						check(Viewport.IsValid());
@@ -471,7 +483,7 @@ void PopulateViewModesMenu(UToolMenu* InMenu, TSharedRef<::SLevelViewport> InVie
 	}
 
 	{
-		auto BuildActorColorationMenu = [WeakViewport = InViewport.ToWeakPtr()](UToolMenu* InMenu)
+		auto BuildActorColorationMenu = [WeakViewport = LevelViewport.ToWeakPtr()](UToolMenu* InMenu)
 		{
 			FToolMenuSection& SubMenuSection =
 				InMenu->AddSection("LevelViewportActorColoration", LOCTEXT("ActorColorationHeader", "Actor Coloration"));
@@ -540,7 +552,7 @@ void PopulateViewModesMenu(UToolMenu* InMenu, TSharedRef<::SLevelViewport> InVie
 				FExecuteAction(),
 				FCanExecuteAction(),
 				FIsActionChecked::CreateLambda(
-					[WeakViewport = InViewport.ToWeakPtr()]()
+					[WeakViewport = LevelViewport.ToWeakPtr()]()
 					{
 						const TSharedPtr<::SLevelViewport> Viewport = WeakViewport.Pin();
 						check(Viewport.IsValid());
@@ -559,7 +571,7 @@ void PopulateViewModesMenu(UToolMenu* InMenu, TSharedRef<::SLevelViewport> InVie
 		FToolMenuSection& Section =
 			InMenu->AddSection("LevelViewportLandscape", LOCTEXT("LandscapeHeader", "Landscape"), InsertPosition);
 
-		auto BuildLandscapeLODMenu = [WeakViewport = InViewport.ToWeakPtr()](UToolMenu* InMenu)
+		auto BuildLandscapeLODMenu = [WeakViewport = LevelViewport.ToWeakPtr()](UToolMenu* InMenu)
 		{
 			FToolMenuSection& SubMenuSection =
 				InMenu->AddSection("LevelViewportLandScapeLOD", LOCTEXT("LandscapeLODHeader", "Landscape LOD"));
@@ -638,16 +650,7 @@ void ExtendViewModesSubmenu(FName InViewModesSubmenuName)
 		FNewToolMenuDelegate::CreateLambda(
 			[](UToolMenu* InDynamicMenu)
 			{
-				ULevelViewportContext* const LevelViewportContext = InDynamicMenu->FindContext<ULevelViewportContext>();
-				if (!LevelViewportContext)
-				{
-					return;
-				}
-
-				if (const TSharedPtr<::SLevelViewport> LevelViewport = LevelViewportContext->LevelViewport.Pin())
-				{
-					PopulateViewModesMenu(InDynamicMenu, LevelViewport.ToSharedRef());
-				}
+				PopulateViewModesMenu(InDynamicMenu);
 			}
 		)
 	);
@@ -1531,9 +1534,20 @@ FToolMenuEntry CreateViewportToolbarPerformanceAndScalabilitySubmenu()
 	);
 }
 
-void GenerateViewportLayoutsMenu(UToolMenu* InMenu, TSharedPtr<::SLevelViewport> InViewport)
+void GenerateViewportLayoutsMenu(UToolMenu* InMenu)
 {
-	TSharedPtr<FUICommandList> CommandList = InViewport->GetCommandList();
+	ULevelViewportContext* const LevelViewportContext = InMenu->FindContext<ULevelViewportContext>();
+	if (!LevelViewportContext)
+	{
+		return;
+	}
+
+	const TSharedPtr<::SLevelViewport> LevelViewport = LevelViewportContext->LevelViewport.Pin();
+	if (!LevelViewport)
+	{
+		return;
+	}
+	TSharedPtr<FUICommandList> CommandList = LevelViewport->GetCommandList();
 
 	// Disable searching in this menu because it only contains visual representations of
 	// viewport layouts without any searchable text.
@@ -1716,7 +1730,7 @@ FToolMenuEntry CreateLevelEditorViewportToolbarSettingsSubmenu()
 								if (const TSharedPtr<::SLevelViewport> LevelViewport =
 										LevelViewportContext->LevelViewport.Pin())
 								{
-									GenerateViewportLayoutsMenu(InMenu, LevelViewport);
+									GenerateViewportLayoutsMenu(InMenu);
 								}
 							}
 						),
@@ -1812,11 +1826,23 @@ void CreateCameraSpawnMenu(UToolMenu* InMenu)
 	}
 }
 
-void CreateBookmarksMenu(UToolMenu* InMenu, TWeakPtr<::SLevelViewport> InViewport)
+void CreateBookmarksMenu(UToolMenu* InMenu)
 {
+	ULevelViewportContext* const LevelViewportContext = InMenu->FindContext<ULevelViewportContext>();
+	if (!LevelViewportContext)
+	{
+		return;
+	}
+
+	const TSharedPtr<::SLevelViewport> LevelViewport = LevelViewportContext->LevelViewport.Pin();
+
+	if (!LevelViewport)
+	{
+		return;
+	}
+
 	// Add a menu entry for each bookmark
-	TSharedPtr<::SLevelViewport> SharedViewport = InViewport.Pin();
-	FLevelEditorViewportClient& ViewportClient = SharedViewport->GetLevelViewportClient();
+	FLevelEditorViewportClient& ViewportClient = LevelViewport->GetLevelViewportClient();
 
 	FToolMenuSection& ManageBookmarksSection =
 		InMenu->FindOrAddSection("ManageBookmarks", LOCTEXT("ManageBookmarkSectionName", "Manage Bookmarks"));
@@ -1825,7 +1851,7 @@ void CreateBookmarksMenu(UToolMenu* InMenu, TWeakPtr<::SLevelViewport> InViewpor
 
 	// Jump to Bookmark Section
 	{
-		bFoundBookmarks = Private::AddJumpToBookmarkMenu(InMenu, InViewport);
+		bFoundBookmarks = Private::AddJumpToBookmarkMenu(InMenu, LevelViewport);
 	}
 
 	// Manage Bookmarks Section
@@ -1874,7 +1900,7 @@ void CreateBookmarksMenu(UToolMenu* InMenu, TWeakPtr<::SLevelViewport> InViewpor
 					LOCTEXT("ManageBookmarksSubMenu", "Manage Bookmarks"),
 					LOCTEXT("ManageBookmarksSubMenu_ToolTip", "Bookmarks related actions"),
 					FNewToolMenuDelegate::CreateLambda(
-						[bFoundBookmarks, InViewport](UToolMenu* InMenu)
+						[bFoundBookmarks, LevelViewport](UToolMenu* InMenu)
 						{
 							if (!bFoundBookmarks)
 							{
@@ -1891,7 +1917,7 @@ void CreateBookmarksMenu(UToolMenu* InMenu, TWeakPtr<::SLevelViewport> InViewpor
 								"ClearBookmark",
 								LOCTEXT("ClearBookmarkSubMenu", "Clear Bookmark"),
 								LOCTEXT("ClearBookmarkSubMenu_ToolTip", "Clear viewport bookmarks"),
-								FNewToolMenuDelegate::CreateLambda(&Private::AddClearBookmarkMenu, InViewport),
+								FNewToolMenuDelegate::CreateLambda(&Private::AddClearBookmarkMenu, LevelViewport),
 								false,
 								FSlateIcon(FAppStyle::Get().GetStyleSetName(), "EditorViewport.SubMenu.Bookmarks")
 							);
@@ -1950,7 +1976,7 @@ FToolMenuEntry CreateFOVMenu(TWeakPtr<::SLevelViewport> InLevelViewportWeak)
 	);
 }
 
-FToolMenuEntry CreateFarViewPlaneMenu(TWeakPtr<::SLevelViewport> InInLevelViewportWeak)
+FToolMenuEntry CreateFarViewPlaneMenu(TWeakPtr<::SLevelViewport> InLevelViewportWeak)
 {
 	constexpr float FarMin = 0.0f;
 	constexpr float FarMax = 100000.0f;
@@ -1961,18 +1987,18 @@ FToolMenuEntry CreateFarViewPlaneMenu(TWeakPtr<::SLevelViewport> InInLevelViewpo
 		LOCTEXT("FarViewPlaneTooltip", "Far View Plane"),
 		FCanExecuteAction(),
 		UnrealEd::FNumericEntryExecuteActionDelegate::CreateLambda(
-			[InInLevelViewportWeak](float InValue)
+			[InLevelViewportWeak](float InValue)
 			{
-				if (TSharedPtr<::SLevelViewport> LevelViewport = InInLevelViewportWeak.Pin())
+				if (TSharedPtr<::SLevelViewport> LevelViewport = InLevelViewportWeak.Pin())
 				{
 					Private::SetFarViewPlaneValue(LevelViewport.ToSharedRef(), InValue);
 				}
 			}
 		),
 		TAttribute<float>::CreateLambda(
-			[InInLevelViewportWeak]()
+			[InLevelViewportWeak]()
 			{
-				if (TSharedPtr<::SLevelViewport> Viewport = InInLevelViewportWeak.Pin())
+				if (TSharedPtr<::SLevelViewport> Viewport = InLevelViewportWeak.Pin())
 				{
 					return Private::GetFarViewPlaneValue(Viewport.ToSharedRef());
 				}
@@ -2058,12 +2084,30 @@ FToolMenuEntry CreateCameraSpeedScalarSlider(TWeakPtr<::SLevelViewport> InLevelV
 	);
 }
 
-void CreateCameraSpeedMenu(UToolMenu* InMenu, const TWeakPtr<::SLevelViewport>& InLevelViewportWeak)
+void CreateCameraSpeedMenu(UToolMenu* InMenu)
 {
+	if (!InMenu)
+	{
+		return;
+	}
+
 	FToolMenuSection& Section = InMenu->AddSection("Section");
 
-	Section.AddEntry(CreateCameraSpeedSlider(InLevelViewportWeak));
-	Section.AddEntry(CreateCameraSpeedScalarSlider(InLevelViewportWeak));
+	ULevelViewportContext* const LevelViewportContext = InMenu->FindContext<ULevelViewportContext>();
+	if (!LevelViewportContext)
+	{
+		return;
+	}
+
+	const TSharedPtr<::SLevelViewport> LevelViewport = LevelViewportContext->LevelViewport.Pin();
+
+	if (!LevelViewport)
+	{
+		return;
+	}
+
+	Section.AddEntry(CreateCameraSpeedSlider(LevelViewport));
+	Section.AddEntry(CreateCameraSpeedScalarSlider(LevelViewport));
 }
 
 FToolMenuEntry CreateLevelEditorViewportToolbarCameraSubmenu()
@@ -2133,9 +2177,9 @@ FToolMenuEntry CreateLevelEditorViewportToolbarCameraSubmenu()
 						LOCTEXT("BookmarksSubMenu", "Bookmarks"),
 						LOCTEXT("BookmarksSubMenu_ToolTip", "Bookmarks related actions"),
 						FNewToolMenuDelegate::CreateLambda(
-							[LevelViewportWeak](UToolMenu* InMenu)
+							[](UToolMenu* InMenu)
 							{
-								CreateBookmarksMenu(InMenu, LevelViewportWeak);
+								CreateBookmarksMenu(InMenu);
 							}
 						),
 						false,
@@ -2157,7 +2201,7 @@ FToolMenuEntry CreateLevelEditorViewportToolbarCameraSubmenu()
 							FNewToolMenuDelegate::CreateLambda(
 								[LevelViewportWeak](UToolMenu* InMenu)
 								{
-									CreateCameraSpeedMenu(InMenu, LevelViewportWeak);
+									CreateCameraSpeedMenu(InMenu);
 								}
 							),
 							false,
diff --git a/Engine/Source/Editor/LevelEditor/Private/ViewportToolbar/LevelEditorViewportToolbarSections.h b/Engine/Source/Editor/LevelEditor/Private/ViewportToolbar/LevelEditorViewportToolbarSections.h
--- a/Engine/Source/Editor/LevelEditor/Private/ViewportToolbar/LevelEditorViewportToolbarSections.h
+++ b/Engine/Source/Editor/LevelEditor/Private/ViewportToolbar/LevelEditorViewportToolbarSections.h
@@ -15,19 +15,18 @@ struct FToolMenuSection;
 namespace UE::LevelEditor
 {
 void CreateCameraSpawnMenu(UToolMenu* InMenu);
-void CreateBookmarksMenu(UToolMenu* InMenu, TWeakPtr<SLevelViewport> InViewport);
-void CreateCameraSpeedMenu(UToolMenu* InMenu, const TWeakPtr<SLevelViewport>& InLevelViewportWeak);
+void CreateBookmarksMenu(UToolMenu* InMenu);
+void CreateCameraSpeedMenu(UToolMenu* InMenu);
 
 FToolMenuEntry CreateFOVMenu(TWeakPtr<SLevelViewport> InLevelViewportWeak);
-FToolMenuEntry CreateFarViewPlaneMenu(TWeakPtr<SLevelViewport> InInLevelViewportWeak);
-FToolMenuEntry CreateFarViewPlaneMenu(TWeakPtr<SLevelViewport> InInLevelViewportWeak);
+FToolMenuEntry CreateFarViewPlaneMenu(TWeakPtr<SLevelViewport> InLevelViewportWeak);
 FToolMenuEntry CreateCameraSpeedSlider(TWeakPtr<SLevelViewport> InLevelViewportWeak);
 FToolMenuEntry CreateCameraSpeedScalarSlider(TWeakPtr<SLevelViewport> InLevelViewportWeak);
 
 bool ShowViewportRealtimeWarning(FLevelEditorViewportClient& ViewportClient);
 
 TSharedPtr<FExtender> GetViewModesLegacyExtenders();
-void PopulateViewModesMenu(UToolMenu* InMenu, TSharedRef<::SLevelViewport> InViewport);
+void PopulateViewModesMenu(UToolMenu* InMenu);
 void ExtendViewModesSubmenu(FName InViewModesSubmenuName);
 
 FToolMenuEntry CreateShowFoliageSubmenu();
@@ -44,7 +43,7 @@ FToolMenuEntry CreateFeatureLevelPreviewSubmenu();
 FToolMenuEntry CreateMaterialQualityLevelSubmenu();
 FToolMenuEntry CreateViewportToolbarPerformanceAndScalabilitySubmenu();
 
-void GenerateViewportLayoutsMenu(UToolMenu* InMenu, TSharedPtr<::SLevelViewport> InViewport);
+void GenerateViewportLayoutsMenu(UToolMenu* InMenu);
 TSharedRef<SWidget> BuildVolumeControlCustomWidget();
 FToolMenuEntry CreateLevelEditorViewportToolbarSettingsSubmenu();
 
diff --git a/Engine/Source/Editor/MaterialEditor/Private/SMaterialEditorViewport.cpp b/Engine/Source/Editor/MaterialEditor/Private/SMaterialEditorViewport.cpp
--- a/Engine/Source/Editor/MaterialEditor/Private/SMaterialEditorViewport.cpp
+++ b/Engine/Source/Editor/MaterialEditor/Private/SMaterialEditorViewport.cpp
@@ -127,7 +127,7 @@ void FMaterialEditorViewportClient::Draw(FViewport* InViewport,FCanvas* Canvas)
 {
 	FEditorViewportClient::Draw(InViewport, Canvas);
 
-	if (MaterialEditorPtr.IsValid())
+	if (!MaterialEditorPtr.IsValid())
 	{
 		MaterialEditorPtr.Pin()->DrawMessages(InViewport, Canvas);
 	}
@@ -994,7 +994,7 @@ void SMaterialEditor3DPreviewViewport::PopulateViewportOverlays(TSharedRef<class
 
 				// Add the "Camera" submenu.
 				{
-					FToolMenuEntry CameraSubmenu = UE::UnrealEd::CreateCameraSubmenu(SharedThis(this));
+					FToolMenuEntry CameraSubmenu = UE::UnrealEd::CreateCameraSubmenu();
 					CameraSubmenu.InsertPosition.Position = EToolMenuInsertType::First;
 					RightSection.AddEntry(CameraSubmenu);
 				}
@@ -1015,7 +1015,8 @@ void SMaterialEditor3DPreviewViewport::PopulateViewportOverlays(TSharedRef<class
 
 				// Add the Performance and Scalability submenu.
 				{
-					FToolMenuEntry PerformanceAndScalabilitySubmenu = UE::UnrealEd::CreatePerformanceAndScalabilitySubmenu(SharedThis(this));
+					FToolMenuEntry PerformanceAndScalabilitySubmenu =
+						UE::UnrealEd::CreatePerformanceAndScalabilitySubmenu();
 					PerformanceAndScalabilitySubmenu.InsertPosition.Position = EToolMenuInsertType::First;
 					RightSection.AddEntry(PerformanceAndScalabilitySubmenu);
 				}
diff --git a/Engine/Source/Editor/StaticMeshEditor/Private/SStaticMeshEditorViewport.cpp b/Engine/Source/Editor/StaticMeshEditor/Private/SStaticMeshEditorViewport.cpp
--- a/Engine/Source/Editor/StaticMeshEditor/Private/SStaticMeshEditorViewport.cpp
+++ b/Engine/Source/Editor/StaticMeshEditor/Private/SStaticMeshEditorViewport.cpp
@@ -638,7 +638,7 @@ TSharedPtr<SWidget> SStaticMeshEditorViewport::MakeViewportToolbar()
 				ViewModesSubmenu.InsertPosition.Position = EToolMenuInsertType::Last;
 				RightSection.AddEntry(ViewModesSubmenu);
 
-				FToolMenuEntry PerformanceAndScalabilitySubmenu = UE::UnrealEd::CreatePerformanceAndScalabilitySubmenu(SharedThis(this));
+				FToolMenuEntry PerformanceAndScalabilitySubmenu = UE::UnrealEd::CreatePerformanceAndScalabilitySubmenu();
 				PerformanceAndScalabilitySubmenu.InsertPosition.Position = EToolMenuInsertType::Last;
 				RightSection.AddEntry(PerformanceAndScalabilitySubmenu);
 			}
diff --git a/Engine/Source/Editor/UnrealEd/Private/ViewportToolbar/UnrealEdViewportToolbar.cpp b/Engine/Source/Editor/UnrealEd/Private/ViewportToolbar/UnrealEdViewportToolbar.cpp
--- a/Engine/Source/Editor/UnrealEd/Private/ViewportToolbar/UnrealEdViewportToolbar.cpp
+++ b/Engine/Source/Editor/UnrealEd/Private/ViewportToolbar/UnrealEdViewportToolbar.cpp
@@ -1737,33 +1737,47 @@ TSharedRef<SWidget> CreateFarViewPlaneMenuWidget(const TSharedRef<SEditorViewpor
 	// clang-format on
 }
 
-FToolMenuEntry CreateCameraSubmenu(TWeakPtr<SEditorViewport> InViewport)
+FToolMenuEntry CreateCameraSubmenu()
 {
 	return FToolMenuEntry::InitSubMenu(
 		"Camera",
 		LOCTEXT("CameraSubmenuLabel", "Camera"),
 		LOCTEXT("CameraSubmenuTooltip", "Camera options"),
 		FNewToolMenuDelegate::CreateLambda(
-			[InViewport](UToolMenu* Submenu) -> void
+			[](UToolMenu* Submenu) -> void
 			{
+				UUnrealEdViewportToolbarContext* const EditorViewportContext =
+					Submenu->FindContext<UUnrealEdViewportToolbarContext>();
+
+				if (!EditorViewportContext)
+				{
+					return;
+				}
+
+				const TSharedPtr<SEditorViewport> EditorViewport = EditorViewportContext->Viewport.Pin();
+				if (!EditorViewport)
+				{
+					return;
+				}
+
 				FToolMenuSection& UnnamedSection = Submenu->FindOrAddSection("", LOCTEXT("UnnamedLabel", ""));
 
 				UnnamedSection.AddEntry(FToolMenuEntry::InitWidget(
-					"CameraMenuItems", UE::UnrealEd::CreateCameraMenuWidget(InViewport.Pin().ToSharedRef()), FText(), true
+					"CameraMenuItems", UE::UnrealEd::CreateCameraMenuWidget(EditorViewport.ToSharedRef()), FText(), true
 				));
 
 				UnnamedSection.AddSeparator("CameraSubmenuSeparator");
 
 				UnnamedSection.AddEntry(FToolMenuEntry::InitWidget(
 					"CameraFOV",
-					UE::UnrealEd::CreateFOVMenuWidget(InViewport.Pin().ToSharedRef()),
+					UE::UnrealEd::CreateFOVMenuWidget(EditorViewport.ToSharedRef()),
 					LOCTEXT("CameraSubmenu_FieldOfViewLabel", "Field of View"),
 					true
 				));
 
 				UnnamedSection.AddEntry(FToolMenuEntry::InitWidget(
 					"CameraFarViewPlane",
-					UE::UnrealEd::CreateFarViewPlaneMenuWidget(InViewport.Pin().ToSharedRef()),
+					UE::UnrealEd::CreateFarViewPlaneMenuWidget(EditorViewport.ToSharedRef()),
 					LOCTEXT("CameraSubmenu_FarViewPlaneLabel", "Far View Plane"),
 					true
 				));
@@ -2049,9 +2063,21 @@ TSharedRef<SWidget> CreateCurrentScreenPercentageWidget(FEditorViewportClient& I
 	// clang-format on
 }
 
-static void ConstructScreenPercentageMenu(UToolMenu* InMenu, FEditorViewportClient* InViewportClient)
+static void ConstructScreenPercentageMenu(UToolMenu* InMenu)
 {
-	FEditorViewportClient& ViewportClient = *InViewportClient;
+	UUnrealEdViewportToolbarContext* const LevelViewportContext = InMenu->FindContext<UUnrealEdViewportToolbarContext>();
+	if (!LevelViewportContext)
+	{
+		return;
+	}
+
+	const TSharedPtr<SEditorViewport> LevelViewport = LevelViewportContext->Viewport.Pin();
+	if (!LevelViewport)
+	{
+		return;
+	}
+
+	FEditorViewportClient& ViewportClient = *LevelViewport->GetViewportClient();
 
 	const FEditorViewportCommands& BaseViewportCommands = FEditorViewportCommands::Get();
 
@@ -2115,33 +2141,124 @@ static void ConstructScreenPercentageMenu(UToolMenu* InMenu, FEditorViewportClie
 	}
 }
 
-UNREALED_API FToolMenuEntry CreatePerformanceAndScalabilitySubmenu(TWeakPtr<SEditorViewport> InViewport)
+UNREALED_API FToolMenuEntry CreatePerformanceAndScalabilitySubmenu()
 {
 	return FToolMenuEntry::InitSubMenu(
 		"PerformanceAndScalability",
 		LOCTEXT("PerformanceAndScalabilitySubmenuLabel", "Performance and Scalability"),
 		LOCTEXT("PerformanceAndScalabilitySubmenuTooltip", ""),
 		FNewToolMenuDelegate::CreateLambda(
-			[InViewport](UToolMenu* Submenu) -> void
+			[](UToolMenu* Submenu) -> void
 			{
-				TSharedPtr<SEditorViewport> Viewport = InViewport.Pin();
-				if (!Viewport)
-				{
-					return;
-				}
-
 				FToolMenuSection& UnnamedSection = Submenu->FindOrAddSection("", LOCTEXT("UnnamedLabel", ""));
 
 				UnnamedSection.AddMenuEntry(FEditorViewportCommands::Get().ToggleRealTime);
 
-				TSharedPtr<FEditorViewportClient> ViewportClient = Viewport->GetViewportClient();
-
 				UnnamedSection.AddSubMenu(
 					"ScreenPercentage",
 					LOCTEXT("ScreenPercentageSubMenu", "Screen Percentage"),
 					LOCTEXT("ScreenPercentageSubMenu_ToolTip", "Customize the viewport's screen percentage"),
-					FNewToolMenuDelegate::CreateStatic(&ConstructScreenPercentageMenu, ViewportClient.Get())
+					FNewToolMenuDelegate::CreateStatic(&ConstructScreenPercentageMenu)
+				);
+			}
+		)
+	);
+}
+
+FToolMenuEntry CreateToggleRealtimeEntry()
+{
+	return FToolMenuEntry::InitDynamicEntry(
+		"ToggleRealtimeDynamicSection",
+		FNewToolMenuSectionDelegate::CreateLambda(
+			[](FToolMenuSection& InnerSection) -> void
+			{
+				UUnrealEdViewportToolbarContext* const EditorViewportContext =
+					InnerSection.FindContext<UUnrealEdViewportToolbarContext>();
+				if (!EditorViewportContext)
+				{
+					return;
+				}
+
+				FToolUIAction RealtimeToggleAction;
+				if (EditorViewportContext)
+				{
+					if (TSharedPtr<SEditorViewport> EditorViewport = EditorViewportContext->Viewport.Pin())
+					{
+						RealtimeToggleAction.ExecuteAction = FToolMenuExecuteAction::CreateLambda(
+							[EditorViewport](const FToolMenuContext& Context) -> void
+							{
+								EditorViewport->OnToggleRealtime();
+							}
+						);
+
+						RealtimeToggleAction.GetActionCheckState = FToolMenuGetActionCheckState::CreateLambda(
+							[EditorViewport](const FToolMenuContext& Context) -> ECheckBoxState
+							{
+								return EditorViewport->IsRealtime() ? ECheckBoxState::Checked : ECheckBoxState::Unchecked;
+							}
+						);
+					}
+				}
+
+				TAttribute<FText> Tooltip;
+				{
+					const FText NonRealtimeTooltip = LOCTEXT(
+						"ToggleRealtimeTooltip_WarnRealtimeOff",
+						"This viewport is not updating in realtime.  Click to turn on realtime mode."
+					);
+					const FText RealtimeTooltip =
+						LOCTEXT("ToggleRealtimeTooltip", "Toggle realtime rendering of the viewport");
+
+					// If we can find a context with a viewport, use that to adjust the tooltip
+					// based on the viewport's realtime status.
+					if (EditorViewportContext)
+					{
+						Tooltip = TAttribute<FText>::CreateLambda(
+							[WeakViewport = EditorViewportContext->Viewport, NonRealtimeTooltip, RealtimeTooltip]() -> FText
+							{
+								bool bDisplayTopLevel = false;
+								if (const TSharedPtr<SEditorViewport> EditorViewport = WeakViewport.Pin())
+								{
+									bDisplayTopLevel = !EditorViewport->IsRealtime();
+								}
+
+								return bDisplayTopLevel ? NonRealtimeTooltip : RealtimeTooltip;
+							}
+						);
+					}
+					else
+					{
+						Tooltip = RealtimeTooltip;
+					}
+				}
+
+				FToolMenuEntry ToggleRealtime = FToolMenuEntry::InitMenuEntry(
+					"ToggleRealtime",
+					LOCTEXT("ToggleRealtimeLabel", "Realtime Viewport"),
+					Tooltip,
+					FSlateIcon(FAppStyle::GetAppStyleSetName(), "EditorViewport.ToggleRealTime"),
+					RealtimeToggleAction,
+					EUserInterfaceActionType::ToggleButton
 				);
+
+				// If we can find a context with a viewport, bind the top-level status of the
+				// realtime button to the viewport's realtime state where we show the realtime
+				// toggle in the top-level if the viewport is NOT realtime.
+				if (EditorViewportContext)
+				{
+					ToggleRealtime.SetShowInToolbarTopLevel(TAttribute<bool>::CreateLambda(
+						[WeakViewport = EditorViewportContext->Viewport]() -> bool
+						{
+							if (const TSharedPtr<SEditorViewport> EditorViewport = WeakViewport.Pin())
+							{
+								return !EditorViewport->IsRealtime();
+							}
+							return false;
+						}
+					));
+				}
+
+				InnerSection.AddEntry(ToggleRealtime);
 			}
 		)
 	);
diff --git a/Engine/Source/Editor/UnrealEd/Public/ViewportToolbar/UnrealEdViewportToolbar.h b/Engine/Source/Editor/UnrealEd/Public/ViewportToolbar/UnrealEdViewportToolbar.h
--- a/Engine/Source/Editor/UnrealEd/Public/ViewportToolbar/UnrealEdViewportToolbar.h
+++ b/Engine/Source/Editor/UnrealEd/Public/ViewportToolbar/UnrealEdViewportToolbar.h
@@ -115,9 +115,11 @@ UNREALED_API FToolMenuEntry CreateNumericEntry(
 	int32 InMaxFractionalDigits = 2
 );
 
-UNREALED_API FToolMenuEntry CreateCameraSubmenu(TWeakPtr<SEditorViewport> InViewport);
+UNREALED_API FToolMenuEntry CreateCameraSubmenu();
 
-UNREALED_API FToolMenuEntry CreatePerformanceAndScalabilitySubmenu(TWeakPtr<SEditorViewport> InViewport);
+UNREALED_API FToolMenuEntry CreatePerformanceAndScalabilitySubmenu();
+
+UNREALED_API FToolMenuEntry CreateToggleRealtimeEntry();
 
 // Camera Menu Widgets
 TSharedRef<SWidget> CreateCameraMenuWidget(const TSharedRef<SEditorViewport>& InViewport);
