[Viewport Toolbar] Add new "Performance & Scalability" submenu to the new Blueprint editor viewport toolbar

diff --git a/Engine/Source/Editor/Kismet/Private/SSCSEditorViewport.cpp b/Engine/Source/Editor/Kismet/Private/SSCSEditorViewport.cpp
--- a/Engine/Source/Editor/Kismet/Private/SSCSEditorViewport.cpp
+++ b/Engine/Source/Editor/Kismet/Private/SSCSEditorViewport.cpp
@@ -54,6 +54,8 @@ class SDockTab;
 class SWidget;
 struct FGeometry;
 
+#define LOCTEXT_NAMESPACE "SSCSEditorViewportToolBar"
+
 /*-----------------------------------------------------------------------------
    SSCSEditorViewportToolBar
 -----------------------------------------------------------------------------*/
@@ -387,13 +389,24 @@ TSharedPtr<SWidget> SSCSEditorViewport::MakeViewportToolbar()
 			// 	RightSection.AddEntry(ShowSubmenu);
 			// }
 
-			// TODO: Add this.
-			// // Add the "Performance & Scalability" submenu.
-			// {
-			// 	FToolMenuEntry PerfSubmenu = UE::LevelEditor::CreateViewportToolbarPerformanceAndScalabilitySubmenu();
-			// 	PerfSubmenu.InsertPosition.Position = EToolMenuInsertType::Last;
-			// 	RightSection.AddEntry(PerfSubmenu);
-			// }
+			// Add the "Performance & Scalability" submenu.
+			{
+				FToolMenuEntry PerfSubmenu = FToolMenuEntry::InitSubMenu(
+					"PerformanceAndScalability",
+					LOCTEXT("PerformanceAndScalabilityLabel", "Performance and Scalability"),
+					LOCTEXT("PerformanceAndScalabilityTooltip", "Performance and Scalability tooltip"),
+					FNewToolMenuDelegate::CreateLambda(
+						[](UToolMenu* Submenu) -> void
+						{
+							FToolMenuSection& UnnamedSection = Submenu->FindOrAddSection("", LOCTEXT("UnnamedLabel", ""));
+							UnnamedSection.AddEntry(UE::UnrealEd::CreateToggleRealtimeEntry());
+						}
+					)
+				);
+
+				PerfSubmenu.InsertPosition.Position = EToolMenuInsertType::Last;
+				RightSection.AddEntry(PerfSubmenu);
+			}
 		}
 	}
 
@@ -633,3 +646,5 @@ EActiveTimerReturnType SSCSEditorViewport::DeferredUpdatePreview(double InCurren
 	bIsActiveTimerRegistered = false;
 	return EActiveTimerReturnType::Stop;
 }
+
+#undef LOCTEXT_NAMESPACE
diff --git a/Engine/Source/Editor/LevelEditor/Private/ViewportToolbar/LevelEditorViewportToolbarSections.cpp b/Engine/Source/Editor/LevelEditor/Private/ViewportToolbar/LevelEditorViewportToolbarSections.cpp
--- a/Engine/Source/Editor/LevelEditor/Private/ViewportToolbar/LevelEditorViewportToolbarSections.cpp
+++ b/Engine/Source/Editor/LevelEditor/Private/ViewportToolbar/LevelEditorViewportToolbarSections.cpp
@@ -1350,118 +1350,7 @@ FToolMenuEntry CreateViewportToolbarPerformanceAndScalabilitySubmenu()
 				{
 					FToolMenuSection& UnnamedSection = Submenu->FindOrAddSection(NAME_None);
 
-					// Add realtime rendering toggle.
-					UnnamedSection.AddDynamicEntry(
-						"ToggleRealtimeDynamicSection",
-						FNewToolMenuSectionDelegate::CreateLambda(
-							[](FToolMenuSection& InnerSection) -> void
-							{
-								FToolUIAction RealtimeToggleAction;
-								RealtimeToggleAction.ExecuteAction = FToolMenuExecuteAction::CreateLambda(
-									[](const FToolMenuContext& Context) -> void
-									{
-										ULevelViewportContext* const LevelViewportContext =
-											Context.FindContext<ULevelViewportContext>();
-										if (!LevelViewportContext)
-										{
-											return;
-										}
-
-										if (const TSharedPtr<::SLevelViewport> LevelViewport =
-												LevelViewportContext->LevelViewport.Pin())
-										{
-											LevelViewport->OnToggleRealtime();
-										}
-									}
-								);
-
-								RealtimeToggleAction.GetActionCheckState = FToolMenuGetActionCheckState::CreateLambda(
-									[](const FToolMenuContext& Context) -> ECheckBoxState
-									{
-										ULevelViewportContext* const LevelViewportContext =
-											Context.FindContext<ULevelViewportContext>();
-										if (!LevelViewportContext)
-										{
-											return ECheckBoxState::Undetermined;
-										}
-
-										if (const TSharedPtr<::SLevelViewport> LevelViewport =
-												LevelViewportContext->LevelViewport.Pin())
-										{
-											return LevelViewport->IsRealtime() ? ECheckBoxState::Checked
-																			   : ECheckBoxState::Unchecked;
-										}
-
-										return ECheckBoxState::Undetermined;
-									}
-								);
-
-								TAttribute<FText> Tooltip;
-								{
-									const FText NonRealtimeTooltip = LOCTEXT(
-										"ToggleRealtimeTooltip_WarnRealtimeOff",
-										"This viewport is not updating in realtime.  Click to turn on realtime mode."
-									);
-									const FText RealtimeTooltip =
-										LOCTEXT("ToggleRealtimeTooltip", "Toggle realtime rendering of the viewport");
-
-									// If we can find a context with a viewport, use that to adjust the tooltip
-									// based on the viewport's realtime status.
-									if (ULevelViewportContext* const LevelViewportContext =
-											InnerSection.FindContext<ULevelViewportContext>())
-									{
-										Tooltip = TAttribute<FText>::CreateLambda(
-											[WeakViewport = LevelViewportContext->LevelViewport,
-											 NonRealtimeTooltip,
-											 RealtimeTooltip]() -> FText
-											{
-												bool bDisplayTopLevel = false;
-												if (const TSharedPtr<::SLevelViewport> LevelViewport = WeakViewport.Pin())
-												{
-													bDisplayTopLevel = !LevelViewport->IsRealtime();
-												}
-
-												return bDisplayTopLevel ? NonRealtimeTooltip : RealtimeTooltip;
-											}
-										);
-									}
-									else
-									{
-										Tooltip = RealtimeTooltip;
-									}
-								}
-
-								FToolMenuEntry ToggleRealtime = FToolMenuEntry::InitMenuEntry(
-									"ToggleRealtime",
-									LOCTEXT("ToggleRealtimeLabel", "Realtime Viewport"),
-									Tooltip,
-									FSlateIcon(FAppStyle::GetAppStyleSetName(), "EditorViewport.ToggleRealTime"),
-									RealtimeToggleAction,
-									EUserInterfaceActionType::ToggleButton
-								);
-
-								// If we can find a context with a viewport, bind the top-level status of the
-								// realtime button to the viewport's realtime state where we show the realtime
-								// toggle in the top-level if the viewport is NOT realtime.
-								if (ULevelViewportContext* const LevelViewportContext =
-										InnerSection.FindContext<ULevelViewportContext>())
-								{
-									ToggleRealtime.SetShowInToolbarTopLevel(TAttribute<bool>::CreateLambda(
-										[WeakViewport = LevelViewportContext->LevelViewport]() -> bool
-										{
-											if (const TSharedPtr<::SLevelViewport> LevelViewport = WeakViewport.Pin())
-											{
-												return !LevelViewport->IsRealtime();
-											}
-											return false;
-										}
-									));
-								}
-
-								InnerSection.AddEntry(ToggleRealtime);
-							}
-						)
-					);
+					UnnamedSection.AddEntry(UE::UnrealEd::CreateToggleRealtimeEntry());
 				}
 
 				{
diff --git a/Engine/Source/Editor/UnrealEd/Private/ViewportToolbar/UnrealEdViewportToolbar.cpp b/Engine/Source/Editor/UnrealEd/Private/ViewportToolbar/UnrealEdViewportToolbar.cpp
--- a/Engine/Source/Editor/UnrealEd/Private/ViewportToolbar/UnrealEdViewportToolbar.cpp
+++ b/Engine/Source/Editor/UnrealEd/Private/ViewportToolbar/UnrealEdViewportToolbar.cpp
@@ -2152,7 +2152,7 @@ UNREALED_API FToolMenuEntry CreatePerformanceAndScalabilitySubmenu()
 			{
 				FToolMenuSection& UnnamedSection = Submenu->FindOrAddSection("", LOCTEXT("UnnamedLabel", ""));
 
-				UnnamedSection.AddMenuEntry(FEditorViewportCommands::Get().ToggleRealTime);
+				UnnamedSection.AddEntry(CreateToggleRealtimeEntry());
 
 				UnnamedSection.AddSubMenu(
 					"ScreenPercentage",
@@ -2179,25 +2179,33 @@ FToolMenuEntry CreateToggleRealtimeEntry()
 					return;
 				}
 
+				TWeakPtr<SEditorViewport> EditorViewportWeak;
 				FToolUIAction RealtimeToggleAction;
 				if (EditorViewportContext)
 				{
-					if (TSharedPtr<SEditorViewport> EditorViewport = EditorViewportContext->Viewport.Pin())
-					{
-						RealtimeToggleAction.ExecuteAction = FToolMenuExecuteAction::CreateLambda(
-							[EditorViewport](const FToolMenuContext& Context) -> void
-							{
-								EditorViewport->OnToggleRealtime();
-							}
-						);
+					EditorViewportWeak = EditorViewportContext->Viewport;
 
-						RealtimeToggleAction.GetActionCheckState = FToolMenuGetActionCheckState::CreateLambda(
-							[EditorViewport](const FToolMenuContext& Context) -> ECheckBoxState
+					RealtimeToggleAction.ExecuteAction = FToolMenuExecuteAction::CreateLambda(
+					[EditorViewportWeak](const FToolMenuContext& Context) -> void
+						{
+						if (TSharedPtr<SEditorViewport> EditorViewport = EditorViewportWeak.Pin())
+						{
+							EditorViewport->OnToggleRealtime();
+						}
+						}
+					);
+
+					RealtimeToggleAction.GetActionCheckState = FToolMenuGetActionCheckState::CreateLambda(
+						[EditorViewportWeak](const FToolMenuContext& Context) -> ECheckBoxState
+						{
+							if (TSharedPtr<SEditorViewport> EditorViewport = EditorViewportWeak.Pin())
 							{
-								return EditorViewport->IsRealtime() ? ECheckBoxState::Checked : ECheckBoxState::Unchecked;
+								return EditorViewport->IsRealtime() ? ECheckBoxState::Checked
+																	: ECheckBoxState::Unchecked;
 							}
-						);
-					}
+							return ECheckBoxState::Undetermined;
+						}
+					);
 				}
 
 				TAttribute<FText> Tooltip;
@@ -2214,10 +2222,10 @@ FToolMenuEntry CreateToggleRealtimeEntry()
 					if (EditorViewportContext)
 					{
 						Tooltip = TAttribute<FText>::CreateLambda(
-							[WeakViewport = EditorViewportContext->Viewport, NonRealtimeTooltip, RealtimeTooltip]() -> FText
+							[EditorViewportWeak, NonRealtimeTooltip, RealtimeTooltip]() -> FText
 							{
 								bool bDisplayTopLevel = false;
-								if (const TSharedPtr<SEditorViewport> EditorViewport = WeakViewport.Pin())
+								if (const TSharedPtr<SEditorViewport> EditorViewport = EditorViewportWeak.Pin())
 								{
 									bDisplayTopLevel = !EditorViewport->IsRealtime();
 								}
