Typo - Fixed the HLOD icon variable name used in SLevelEditor

diff --git a/Engine/Source/Editor/LevelEditor/Private/SLevelEditor.cpp b/Engine/Source/Editor/LevelEditor/Private/SLevelEditor.cpp
--- a/Engine/Source/Editor/LevelEditor/Private/SLevelEditor.cpp
+++ b/Engine/Source/Editor/LevelEditor/Private/SLevelEditor.cpp
@@ -1496,14 +1496,14 @@ TSharedRef<SWidget> SLevelEditor::RestoreContentArea( const TSharedRef<SDockTab>
 		}
 
 		{
-			const FSlateIcon LayersIcon(FAppStyle::GetAppStyleSetName(), "LevelEditor.Tabs.HLOD");
+			const FSlateIcon HLODIcon(FAppStyle::GetAppStyleSetName(), "LevelEditor.Tabs.HLOD");
 			LevelEditorTabManager->RegisterTabSpawner(LevelEditorTabIds::LevelEditorHierarchicalLODOutliner, 
 				FOnSpawnTab::CreateSP(this, &SLevelEditor::SpawnLevelEditorTab, LevelEditorTabIds::LevelEditorHierarchicalLODOutliner, FString()),
 				FCanSpawnTab::CreateSP(this, &SLevelEditor::CanSpawnLevelEditorTab, LevelEditorTabIds::LevelEditorHierarchicalLODOutliner))
 				.SetDisplayName(NSLOCTEXT("LevelEditorTabs", "LevelEditorHierarchicalLODOutliner", "Hierarchical LOD Outliner"))
 				.SetTooltipText(NSLOCTEXT("LevelEditorTabs", "LevelEditorHierarchicalLODOutlinerTooltipText", "Open the Hierarchical LOD Outliner."))
 				.SetGroup(MenuStructure.GetLevelEditorCategory())
-				.SetIcon(LayersIcon);
+				.SetIcon(HLODIcon);
 		}
 		
 		{
