[Viewport Toolbar] Lumen and Nanite options missing from level viewport toolbars

SLevelViewportToolBar::FillViewMenu(UToolMenu* Menu) was not passing the ULevelViewportContext Object to UE::LevelEditor::PopulateViewModesMenu(Menu).
So, the function couldn't properly extend the menu entries.

diff --git a/Engine/Source/Editor/LevelEditor/Private/SLevelViewportToolBar.cpp b/Engine/Source/Editor/LevelEditor/Private/SLevelViewportToolBar.cpp
--- a/Engine/Source/Editor/LevelEditor/Private/SLevelViewportToolBar.cpp
+++ b/Engine/Source/Editor/LevelEditor/Private/SLevelViewportToolBar.cpp
@@ -984,6 +984,15 @@ TWeakObjectPtr<UWorld> SLevelViewportToolBar::GetWorld() const
 
 void SLevelViewportToolBar::FillViewMenu(UToolMenu* Menu)
 {
+	if (!Menu)
+	{
+		return;
+	}
+
+	ULevelViewportContext* const ContextObject = NewObject<ULevelViewportContext>();
+	ContextObject->LevelViewport = Viewport;
+	Menu->Context.AddObject(ContextObject);
+
 	UE::LevelEditor::PopulateViewModesMenu(Menu);
 }
 
