[Viewport Toolbar] Empty Bookmarks submenu in legacy toolbar

ULevelViewportContext Object was missing from options menu FToolMenuContext

diff --git a/Engine/Source/Editor/LevelEditor/Private/SLevelViewportToolBar.cpp b/Engine/Source/Editor/LevelEditor/Private/SLevelViewportToolBar.cpp
--- a/Engine/Source/Editor/LevelEditor/Private/SLevelViewportToolBar.cpp
+++ b/Engine/Source/Editor/LevelEditor/Private/SLevelViewportToolBar.cpp
@@ -448,10 +448,21 @@ TSharedRef<SWidget> SLevelViewportToolBar::GenerateOptionsMenu()
 	FLevelEditorModule& LevelEditorModule = FModuleManager::GetModuleChecked<FLevelEditorModule>( TEXT("LevelEditor") );
 	TSharedPtr<FExtender> MenuExtender = LevelEditorModule.AssembleExtenders(CommandList, LevelEditorModule.GetAllLevelViewportOptionsMenuExtenders());
 
-	ULevelViewportToolBarContext* ContextObject = NewObject<ULevelViewportToolBarContext>();
-	ContextObject->LevelViewportToolBarWidget = SharedThis(this);
+	FToolMenuContext MenuContext(CommandList, MenuExtender);
+	{
+		{
+			ULevelViewportToolBarContext* ToolbarContextObject = NewObject<ULevelViewportToolBarContext>();
+			ToolbarContextObject->LevelViewportToolBarWidget = SharedThis(this);
+			MenuContext.AddObject(ToolbarContextObject);
+		}
+
+		{
+			ULevelViewportContext* const LevelContextObject = NewObject<ULevelViewportContext>();
+			LevelContextObject->LevelViewport = Viewport;
+			MenuContext.AddObject(LevelContextObject);
+		}
+	}
 
-	FToolMenuContext MenuContext(CommandList, MenuExtender, ContextObject);
 	return UToolMenus::Get()->GenerateWidget(MenuName, MenuContext);
 }
 
