Editor Layout Change: Fix crash when Viewport->GetWorld() returns nullptr

#rnx

diff --git a/Engine/Source/Editor/LevelEditor/Private/ViewportToolbar/LevelEditorViewportToolbarSections.cpp b/Engine/Source/Editor/LevelEditor/Private/ViewportToolbar/LevelEditorViewportToolbarSections.cpp
--- a/Engine/Source/Editor/LevelEditor/Private/ViewportToolbar/LevelEditorViewportToolbarSections.cpp
+++ b/Engine/Source/Editor/LevelEditor/Private/ViewportToolbar/LevelEditorViewportToolbarSections.cpp
@@ -677,6 +677,11 @@ FToolMenuEntry CreateShowFoliageSubmenu()
 					return;
 				}
 
+				if (!Viewport->GetWorld())
+				{
+					return;
+				}
+
 				{
 					FToolMenuSection& Section = Submenu->AddSection("LevelViewportFoliageMeshes");
 					// Map 'Show All' and 'Hide All' commands
