Fix Text conflict

diff --git a/Engine/Source/Editor/LevelEditor/Private/ViewportToolbar/LevelEditorViewportToolbarSections.cpp b/Engine/Source/Editor/LevelEditor/Private/ViewportToolbar/LevelEditorViewportToolbarSections.cpp
--- a/Engine/Source/Editor/LevelEditor/Private/ViewportToolbar/LevelEditorViewportToolbarSections.cpp
+++ b/Engine/Source/Editor/LevelEditor/Private/ViewportToolbar/LevelEditorViewportToolbarSections.cpp
@@ -1305,7 +1305,7 @@ FToolMenuEntry CreateFeatureLevelPreviewSubmenu()
 			[](UToolMenu* InMenu) -> void
 			{
 				FToolMenuSection& Section =
-					InMenu->AddSection("EditorPreviewMode", LOCTEXT("EditorPreviewModeDevices", "Preview Platforms"));
+					InMenu->AddSection("EditorPreviewMode", LOCTEXT("EditorPreviewModePlatforms", "Preview Platforms"));
 
 				for (auto Iter = FLevelEditorCommands::Get().PlatformToPreviewPlatformOverrides.CreateConstIterator(); Iter; ++Iter)
 				{
