Remove label and add icon to the "Settings" submenu of the new level editor viewport toolbar

diff --git a/Engine/Source/Editor/LevelEditor/Private/ViewportToolbar/LevelEditorViewportToolbarSections.cpp b/Engine/Source/Editor/LevelEditor/Private/ViewportToolbar/LevelEditorViewportToolbarSections.cpp
--- a/Engine/Source/Editor/LevelEditor/Private/ViewportToolbar/LevelEditorViewportToolbarSections.cpp
+++ b/Engine/Source/Editor/LevelEditor/Private/ViewportToolbar/LevelEditorViewportToolbarSections.cpp
@@ -1608,7 +1608,7 @@ TSharedRef<SWidget> BuildVolumeControlCustomWidget()
 
 FToolMenuEntry CreateLevelEditorViewportToolbarSettingsSubmenu()
 {
-	return FToolMenuEntry::InitSubMenu(
+	FToolMenuEntry Entry = FToolMenuEntry::InitSubMenu(
 		"Settings",
 		LOCTEXT("SettingsSubmenuLabel", "Settings"),
 		LOCTEXT("SettingsSubmenuTooltip", "Viewport-related settings"),
@@ -1720,6 +1720,9 @@ FToolMenuEntry CreateLevelEditorViewportToolbarSettingsSubmenu()
 			}
 		)
 	);
+	Entry.Icon = FSlateIcon(FAppStyle::GetAppStyleSetName(), "LevelEditor.GameSettings");
+	Entry.ToolbarLabelOverride = FText();
+	return Entry;
 }
 
 void CreateCameraSpawnMenu(UToolMenu* InMenu)
