Add icon and hide label on the Show menu of the new level editor viewport toolbar

diff --git a/Engine/Source/Editor/LevelEditor/Private/ViewportToolbar/LevelEditorViewportToolbarSections.cpp b/Engine/Source/Editor/LevelEditor/Private/ViewportToolbar/LevelEditorViewportToolbarSections.cpp
--- a/Engine/Source/Editor/LevelEditor/Private/ViewportToolbar/LevelEditorViewportToolbarSections.cpp
+++ b/Engine/Source/Editor/LevelEditor/Private/ViewportToolbar/LevelEditorViewportToolbarSections.cpp
@@ -1224,7 +1224,7 @@ FToolMenuEntry CreateShowStatsSubmenu()
 
 FToolMenuEntry CreateViewportToolbarShowSubmenu()
 {
-	return FToolMenuEntry::InitSubMenu(
+	FToolMenuEntry Entry = FToolMenuEntry::InitSubMenu(
 		"Show",
 		LOCTEXT("ShowSubmenuLabel", "Show"),
 		LOCTEXT("ShowSubmenuTooltip", "Show flags related to the current viewport"),
@@ -1293,6 +1293,9 @@ FToolMenuEntry CreateViewportToolbarShowSubmenu()
 			}
 		)
 	);
+	Entry.ToolbarLabelOverride = FText();
+	Entry.Icon = FSlateIcon(FAppStyle::GetAppStyleSetName(), "Level.VisibleHighlightIcon16x");
+	return Entry;
 }
 
 FToolMenuEntry CreateFeatureLevelPreviewSubmenu()
