[ISM Selection] Make selecting an ISM work more like component selection from an actor.

Now - the typical selection rules are followed if an ISM is hit, and selection will follow the standard pattern Actor -> Component -> ISM where a double click takes you one level deeper, and a single selection takes you one level higher.

diff --git a/Engine/Source/Editor/LevelEditor/Private/Elements/Component/ComponentElementLevelEditorSelectionCustomization.cpp b/Engine/Source/Editor/LevelEditor/Private/Elements/Component/ComponentElementLevelEditorSelectionCustomization.cpp
--- a/Engine/Source/Editor/LevelEditor/Private/Elements/Component/ComponentElementLevelEditorSelectionCustomization.cpp
+++ b/Engine/Source/Editor/LevelEditor/Private/Elements/Component/ComponentElementLevelEditorSelectionCustomization.cpp
@@ -39,7 +39,12 @@ bool FComponentElementLevelEditorSelectionCustomization::DeselectElement(const T
 
 FTypedElementHandle FComponentElementLevelEditorSelectionCustomization::GetSelectionElement(const TTypedElement<ITypedElementSelectionInterface>& InElementSelectionHandle, FTypedElementListConstRef InCurrentSelection, const ETypedElementSelectionMethod InSelectionMethod)
 {
-	if (const UActorComponent* Component = ComponentElementDataUtil::GetComponentFromHandle(InElementSelectionHandle))
+	return GetSelectionElementStatic(InElementSelectionHandle, InCurrentSelection, InSelectionMethod);
+}
+
+FTypedElementHandle FComponentElementLevelEditorSelectionCustomization::GetSelectionElementStatic(const FTypedElementHandle& InComponentHandle, FTypedElementListConstRef InCurrentSelection, const ETypedElementSelectionMethod InSelectionMethod)
+{
+	if (const UActorComponent* Component = ComponentElementDataUtil::GetComponentFromHandle(InComponentHandle))
 	{
 		const AActor* ConsideredActor = Component->GetOwner();
 		const USceneComponent* ConsideredComponent = Cast<USceneComponent>(Component);
@@ -68,8 +73,9 @@ FTypedElementHandle FComponentElementLevelEditorSelectionCustomization::GetSelec
 			const bool bActorIsBlueprintable = FKismetEditorUtilities::CanCreateBlueprintOfClass(ConsideredActor->GetClass());
 			const bool bComponentAlreadySelected = InCurrentSelection->HasElementsOfType(NAME_Components);
 			const bool bWasDoubleClick = InSelectionMethod == ETypedElementSelectionMethod::Secondary;
+			const bool bWasFromSecondary = InSelectionMethod == ETypedElementSelectionMethod::FromSecondary;
 
-			const bool bSelectComponent = bActorAlreadySelectedExclusively && bActorIsBlueprintable && (bComponentAlreadySelected != bWasDoubleClick);
+			const bool bSelectComponent = bActorIsBlueprintable && (bWasFromSecondary || (bActorAlreadySelectedExclusively && (bComponentAlreadySelected != bWasDoubleClick)));
 
 			if (bSelectComponent && ConsideredComponent)
 			{
@@ -81,7 +87,7 @@ FTypedElementHandle FComponentElementLevelEditorSelectionCustomization::GetSelec
 			}
 		}
 	}
-	return InElementSelectionHandle;
+	return InComponentHandle;
 }
 
 void FComponentElementLevelEditorSelectionCustomization::GetNormalizedElements(const TTypedElement<ITypedElementSelectionInterface>& InElementSelectionHandle, FTypedElementListConstRef InSelectionSet, const FTypedElementSelectionNormalizationOptions& InNormalizationOptions, FTypedElementListRef OutNormalizedElements)
diff --git a/Engine/Source/Editor/LevelEditor/Private/Elements/Component/ComponentElementLevelEditorSelectionCustomization.h b/Engine/Source/Editor/LevelEditor/Private/Elements/Component/ComponentElementLevelEditorSelectionCustomization.h
--- a/Engine/Source/Editor/LevelEditor/Private/Elements/Component/ComponentElementLevelEditorSelectionCustomization.h
+++ b/Engine/Source/Editor/LevelEditor/Private/Elements/Component/ComponentElementLevelEditorSelectionCustomization.h
@@ -20,4 +20,6 @@ public:
 
 	bool SelectComponentElement(const TTypedElement<ITypedElementSelectionInterface>& InComponentSelectionHandle, FTypedElementListRef InSelectionSet, const FTypedElementSelectionOptions& InSelectionOptions);
 	bool DeselectComponentElement(const TTypedElement<ITypedElementSelectionInterface>& InComponentSelectionHandle, FTypedElementListRef InSelectionSet, const FTypedElementSelectionOptions& InSelectionOptions);
+
+	static FTypedElementHandle GetSelectionElementStatic(const FTypedElementHandle& InComponentHandle, FTypedElementListConstRef InCurrentSelection, const ETypedElementSelectionMethod InSelectionMethod);
 };
diff --git a/Engine/Source/Editor/LevelEditor/Private/Elements/SMInstance/SMInstanceElementLevelEditorSelectionCustomization.cpp b/Engine/Source/Editor/LevelEditor/Private/Elements/SMInstance/SMInstanceElementLevelEditorSelectionCustomization.cpp
--- a/Engine/Source/Editor/LevelEditor/Private/Elements/SMInstance/SMInstanceElementLevelEditorSelectionCustomization.cpp
+++ b/Engine/Source/Editor/LevelEditor/Private/Elements/SMInstance/SMInstanceElementLevelEditorSelectionCustomization.cpp
@@ -1,11 +1,14 @@
 // Copyright Epic Games, Inc. All Rights Reserved.
 
 #include "Elements/SMInstance/SMInstanceElementLevelEditorSelectionCustomization.h"
+#include "Elements/Component/ComponentElementLevelEditorSelectionCustomization.h"
 #include "Elements/SMInstance/SMInstanceElementData.h"
 #include "Components/InstancedStaticMeshComponent.h"
 
 #include "Elements/Framework/TypedElementList.h"
 #include "Elements/Framework/EngineElementsLibrary.h"
+#include "Elements/Framework/TypedElementInterfaceCustomization.h"
+#include "Elements/Interfaces/TypedElementHierarchyInterface.h"
 
 #include "LevelUtils.h"
 
@@ -78,3 +81,47 @@ bool FSMInstanceElementLevelEditorSelectionCustomization::DeselectElement(const
 
 	return true;
 }
+
+FTypedElementHandle FSMInstanceElementLevelEditorSelectionCustomization::GetSelectionElement(const TTypedElement<ITypedElementSelectionInterface>& InElementSelectionHandle, FTypedElementListConstRef InCurrentSelection, const ETypedElementSelectionMethod InSelectionMethod)
+{
+	const FSMInstanceManager SMInstance = SMInstanceElementDataUtil::GetSMInstanceFromHandleChecked(InElementSelectionHandle);
+	if (!SMInstance)
+	{
+		return InElementSelectionHandle;
+	}
+
+	if (const UInstancedStaticMeshComponent* Component = SMInstance.GetISMComponent())
+	{
+		const FTypedElementHandle OwningComponentHandle = UEngineElementsLibrary::AcquireEditorComponentElementHandle(Component);
+
+		const bool bWasDoubleClick = InSelectionMethod == ETypedElementSelectionMethod::Secondary;
+		const bool bComponentAlreadySelected = InCurrentSelection->Contains(OwningComponentHandle);
+		const bool bIsISMAlreadySelected = InCurrentSelection->Contains(InElementSelectionHandle);
+
+		bool bIsSiblingSelected = false;
+		if (InCurrentSelection->HasElementsOfType(InElementSelectionHandle.GetId().GetTypeId()))
+		{
+			const FSMInstanceManager SelectedSMInstance = SMInstanceElementDataUtil::GetSMInstanceFromHandleChecked(InCurrentSelection->GetTopElement<ITypedElementHierarchyInterface>());
+			if (SelectedSMInstance && (SelectedSMInstance.GetISMComponent() == Component))
+			{
+				bIsSiblingSelected = true;
+			}
+		}
+
+		// If a single click, and the ISM was already selected, defer to the component selection, so that we traverse the hierarchy one level instead of picking the top level selection again if we're hit on the same selection.
+		if (!bWasDoubleClick && (bIsSiblingSelected || bIsISMAlreadySelected))
+		{
+			return FComponentElementLevelEditorSelectionCustomization::GetSelectionElementStatic(OwningComponentHandle, InCurrentSelection, ETypedElementSelectionMethod::FromSecondary);
+		}
+		
+		// consider the component already selected if it is in the selection, or if the ISM handle that was hit is a sibling to the current handle
+		if (bWasDoubleClick && (bIsSiblingSelected || bComponentAlreadySelected))
+		{
+			return InElementSelectionHandle;
+		}
+
+		return FComponentElementLevelEditorSelectionCustomization::GetSelectionElementStatic(OwningComponentHandle, InCurrentSelection, InSelectionMethod);
+	}
+
+	return InElementSelectionHandle;
+}
diff --git a/Engine/Source/Editor/LevelEditor/Private/Elements/SMInstance/SMInstanceElementLevelEditorSelectionCustomization.h b/Engine/Source/Editor/LevelEditor/Private/Elements/SMInstance/SMInstanceElementLevelEditorSelectionCustomization.h
--- a/Engine/Source/Editor/LevelEditor/Private/Elements/SMInstance/SMInstanceElementLevelEditorSelectionCustomization.h
+++ b/Engine/Source/Editor/LevelEditor/Private/Elements/SMInstance/SMInstanceElementLevelEditorSelectionCustomization.h
@@ -12,4 +12,5 @@ public:
 	virtual bool CanDeselectElement(const TTypedElement<ITypedElementSelectionInterface>& InElementSelectionHandle, const FTypedElementSelectionOptions& InSelectionOptions) override;
 	virtual bool SelectElement(const TTypedElement<ITypedElementSelectionInterface>& InElementSelectionHandle, FTypedElementListRef InSelectionSet, const FTypedElementSelectionOptions& InSelectionOptions) override;
 	virtual bool DeselectElement(const TTypedElement<ITypedElementSelectionInterface>& InElementSelectionHandle, FTypedElementListRef InSelectionSet, const FTypedElementSelectionOptions& InSelectionOptions) override;
+	virtual FTypedElementHandle GetSelectionElement(const TTypedElement<ITypedElementSelectionInterface>& InElementSelectionHandle, FTypedElementListConstRef InCurrentSelection, const ETypedElementSelectionMethod InSelectionMethod) override;
 };
diff --git a/Engine/Source/Runtime/TypedElementRuntime/Public/Elements/Interfaces/TypedElementSelectionInterface.h b/Engine/Source/Runtime/TypedElementRuntime/Public/Elements/Interfaces/TypedElementSelectionInterface.h
--- a/Engine/Source/Runtime/TypedElementRuntime/Public/Elements/Interfaces/TypedElementSelectionInterface.h
+++ b/Engine/Source/Runtime/TypedElementRuntime/Public/Elements/Interfaces/TypedElementSelectionInterface.h
@@ -27,6 +27,8 @@ enum class ETypedElementSelectionMethod : uint8
 	Primary,
 	/** Select the "secondary" element (eg, a component would favor selecting itself) */
 	Secondary,
+	/** The "secondary" element is choosing to pass back up the chain (eg. an ISM was selected, but we want to traverse back to the parent element) */
+	FromSecondary,
 };
 
 UENUM()
