Improve the Performance & Scalability submenu of the new level editor viewport toolbar

* Add Icons.Scalability as a toolbar-specific editor version of the NiagaraEditor.Scalability icon.
* Use Icons.Scalability as the icon for the submenu.
* Remove the label in the toolbar case.

diff --git a/Engine/Source/Editor/EditorStyle/Private/StarshipStyle.cpp b/Engine/Source/Editor/EditorStyle/Private/StarshipStyle.cpp
--- a/Engine/Source/Editor/EditorStyle/Private/StarshipStyle.cpp
+++ b/Engine/Source/Editor/EditorStyle/Private/StarshipStyle.cpp
@@ -457,6 +457,7 @@ void FStarshipEditorStyle::FStyle::SetupGeneralStyles()
 		Set("Icons.HollowHeart", new IMAGE_BRUSH_SVG("Starship/Common/HollowHeart", Icon16x16));
 		Set("Icons.Heart", new IMAGE_BRUSH_SVG("Starship/Common/Heart", Icon16x16));
 		Set("Icons.RecentAssets", new IMAGE_BRUSH_SVG("Starship/Common/RecentAssets", Icon16x16));
+		Set("Icons.Scalability", new IMAGE_BRUSH_SVG("Starship/Common/Scalability_16", Icon20x20));
 
 		Set("Icons.Toolbar.Play", new IMAGE_BRUSH_SVG("Starship/Common/play", Icon20x20));
 		Set("Icons.Toolbar.Pause", new IMAGE_BRUSH_SVG("Starship/MainToolbar/pause", Icon20x20));
diff --git a/Engine/Source/Editor/LevelEditor/Private/ViewportToolbar/LevelEditorViewportToolbarSections.cpp b/Engine/Source/Editor/LevelEditor/Private/ViewportToolbar/LevelEditorViewportToolbarSections.cpp
--- a/Engine/Source/Editor/LevelEditor/Private/ViewportToolbar/LevelEditorViewportToolbarSections.cpp
+++ b/Engine/Source/Editor/LevelEditor/Private/ViewportToolbar/LevelEditorViewportToolbarSections.cpp
@@ -1361,7 +1361,7 @@ FToolMenuEntry CreateMaterialQualityLevelSubmenu()
 
 FToolMenuEntry CreateViewportToolbarPerformanceAndScalabilitySubmenu()
 {
-	return FToolMenuEntry::InitSubMenu(
+	FToolMenuEntry Entry = FToolMenuEntry::InitSubMenu(
 		"PerformanceAndScalability",
 		LOCTEXT("PerformanceAndScalabilityLabel", "Performance & Scalability"),
 		LOCTEXT("PerformanceAndScalabilityTooltip", "Performance and scalability tools tied to this viewport."),
@@ -1442,6 +1442,9 @@ FToolMenuEntry CreateViewportToolbarPerformanceAndScalabilitySubmenu()
 			}
 		)
 	);
+	Entry.Icon = FSlateIcon(FAppStyle::GetAppStyleSetName(), "Icons.Scalability");
+	Entry.ToolbarLabelOverride = FText();
+	return Entry;
 }
 
 void GenerateViewportLayoutsMenu(UToolMenu* InMenu)
