Improve Camera Speed submenu in the new level editor viewport toolbar

* Change icon of the Camera Speed submenu.
* Top-level flag the submenu so it always is visible in the new level editor viewport toolbar.
* Make the label of the Camera Speed submenu dynamic when it appears in a toolbar.

diff --git a/Engine/Source/Editor/EditorStyle/Private/StarshipStyle.cpp b/Engine/Source/Editor/EditorStyle/Private/StarshipStyle.cpp
--- a/Engine/Source/Editor/EditorStyle/Private/StarshipStyle.cpp
+++ b/Engine/Source/Editor/EditorStyle/Private/StarshipStyle.cpp
@@ -4237,6 +4237,7 @@ void FStarshipEditorStyle::FStyle::SetupLevelEditorStyle()
 		Set( "EditorViewport.TranslateMode", new IMAGE_BRUSH_SVG( "Starship/EditorViewport/translate", Icon16x16 ) );
 		Set( "EditorViewport.RotateMode", new IMAGE_BRUSH_SVG("Starship/EditorViewport/rotate", Icon16x16 ) );
 		Set( "EditorViewport.ScaleMode", new IMAGE_BRUSH_SVG( "Starship/EditorViewport/scale", Icon16x16 ) );
+		Set("EditorViewport.Speed", new IMAGE_BRUSH_SVG("Starship/EditorViewport/speed", Icon16x16));
 
 		Set( "EditorViewport.TranslateRotateMode", new IMAGE_BRUSH_SVG("Starship/EditorViewport/TranslateRotate3D", Icon16x16 ) );
 		Set( "EditorViewport.TranslateRotate2DMode", new IMAGE_BRUSH_SVG("Starship/EditorViewport/TranslateRotate2D", Icon16x16 ) );
diff --git a/Engine/Source/Editor/LevelEditor/Private/ViewportToolbar/LevelEditorViewportToolbarSections.cpp b/Engine/Source/Editor/LevelEditor/Private/ViewportToolbar/LevelEditorViewportToolbarSections.cpp
--- a/Engine/Source/Editor/LevelEditor/Private/ViewportToolbar/LevelEditorViewportToolbarSections.cpp
+++ b/Engine/Source/Editor/LevelEditor/Private/ViewportToolbar/LevelEditorViewportToolbarSections.cpp
@@ -2080,7 +2080,7 @@ void ExtendCameraSubmenu(FName InCameraOptionsSubmenuName)
 
 					// Camera Speed Submenu
 					{
-						PositioningSection.AddSubMenu(
+						FToolMenuEntry& Entry = PositioningSection.AddSubMenu(
 							"CameraSpeed",
 							LOCTEXT("CameraSpeedSubMenu", "Camera Speed"),
 							LOCTEXT("CameraSpeedSubMenu_ToolTip", "Camera Speed related actions"),
@@ -2091,8 +2091,16 @@ void ExtendCameraSubmenu(FName InCameraOptionsSubmenuName)
 								}
 							),
 							false,
-							FSlateIcon(FSlateIcon(FAppStyle::GetAppStyleSetName(), "LevelViewport.ToggleActorPilotCameraView"))
+							FSlateIcon(FSlateIcon(FAppStyle::GetAppStyleSetName(), "EditorViewport.Speed"))
+						);
+
+						Entry.ToolbarLabelOverride = TAttribute<FText>::CreateLambda(
+							[LevelViewportWeak]()
+							{
+								return UE::UnrealEd::GetCameraSpeedLabel(LevelViewportWeak);
+							}
 						);
+						Entry.SetShowInToolbarTopLevel(true);
 					}
 
 					PositioningSection.AddSeparator("PositioningSeparator_1");
diff --git a/Engine/Source/Editor/UnrealEd/Private/STransformViewportToolbar.cpp b/Engine/Source/Editor/UnrealEd/Private/STransformViewportToolbar.cpp
--- a/Engine/Source/Editor/UnrealEd/Private/STransformViewportToolbar.cpp
+++ b/Engine/Source/Editor/UnrealEd/Private/STransformViewportToolbar.cpp
@@ -585,17 +585,8 @@ FText STransformViewportToolBar::GetScaleGridLabel() const
 
 FText STransformViewportToolBar::GetCameraSpeedLabel() const
 {
-	auto ViewportPin = Viewport.Pin();
-	if (ViewportPin.IsValid() && ViewportPin->GetViewportClient().IsValid())
-	{
-		float CameraSpeed = ViewportPin->GetViewportClient()->GetCameraSpeed();
-		FNumberFormattingOptions FormattingOptions = FNumberFormattingOptions::DefaultNoGrouping();
-		FormattingOptions.MaximumFractionalDigits = CameraSpeed > 1 ? 1 : 3;
-		return FText::AsNumber(CameraSpeed, &FormattingOptions);
-	}
-
-	return FText();
-}	
+	return UE::UnrealEd::GetCameraSpeedLabel(Viewport);
+}
 
 float STransformViewportToolBar::GetCamSpeedSliderPosition() const
 {
diff --git a/Engine/Source/Editor/UnrealEd/Private/ViewportToolbar/UnrealEdViewportToolbar.cpp b/Engine/Source/Editor/UnrealEd/Private/ViewportToolbar/UnrealEdViewportToolbar.cpp
--- a/Engine/Source/Editor/UnrealEd/Private/ViewportToolbar/UnrealEdViewportToolbar.cpp
+++ b/Engine/Source/Editor/UnrealEd/Private/ViewportToolbar/UnrealEdViewportToolbar.cpp
@@ -1951,6 +1951,22 @@ TSharedRef<SWidget> CreateFarViewPlaneMenuWidget(const TSharedRef<SEditorViewpor
 	// clang-format on
 }
 
+FText GetCameraSpeedLabel(const TWeakPtr<SEditorViewport>& WeakViewport)
+{
+	if (const TSharedPtr<SEditorViewport> Viewport = WeakViewport.Pin())
+	{
+		if (Viewport->GetViewportClient())
+		{
+			const float CameraSpeed = Viewport->GetViewportClient()->GetCameraSpeed();
+			FNumberFormattingOptions FormattingOptions = FNumberFormattingOptions::DefaultNoGrouping();
+			FormattingOptions.MaximumFractionalDigits = CameraSpeed > 1 ? 1 : 3;
+			return FText::AsNumber(CameraSpeed, &FormattingOptions);
+		}
+	}
+
+	return FText();
+}
+
 FText GetCameraSubmenuLabelFromViewportType(const ELevelViewportType ViewportType)
 {
 	FText Label = LOCTEXT("CameraMenuTitle_Default", "Camera");
diff --git a/Engine/Source/Editor/UnrealEd/Public/ViewportToolbar/UnrealEdViewportToolbar.h b/Engine/Source/Editor/UnrealEd/Public/ViewportToolbar/UnrealEdViewportToolbar.h
--- a/Engine/Source/Editor/UnrealEd/Public/ViewportToolbar/UnrealEdViewportToolbar.h
+++ b/Engine/Source/Editor/UnrealEd/Public/ViewportToolbar/UnrealEdViewportToolbar.h
@@ -125,6 +125,7 @@ UNREALED_API FToolMenuEntry CreateNumericEntry(
 	int32 InMaxFractionalDigits = 2
 );
 
+UNREALED_API FText GetCameraSpeedLabel(const TWeakPtr<SEditorViewport>& WeakViewport);
 UNREALED_API FText GetCameraSubmenuLabelFromViewportType(const ELevelViewportType ViewportType);
 UNREALED_API FName GetCameraSubmenuIconFNameFromViewportType(const ELevelViewportType ViewportType);
 UNREALED_API FToolMenuEntry CreateViewportToolbarCameraSubmenu();
