[Viewport Toolbar] Create Show menu with default entries

Creates a default Show Menu with commonly used show flags.

diff --git a/Engine/Source/Editor/LevelEditor/Private/ViewportToolbar/LevelEditorViewportToolbarSections.cpp b/Engine/Source/Editor/LevelEditor/Private/ViewportToolbar/LevelEditorViewportToolbarSections.cpp
--- a/Engine/Source/Editor/LevelEditor/Private/ViewportToolbar/LevelEditorViewportToolbarSections.cpp
+++ b/Engine/Source/Editor/LevelEditor/Private/ViewportToolbar/LevelEditorViewportToolbarSections.cpp
@@ -1248,48 +1248,65 @@ FToolMenuEntry CreateViewportToolbarShowSubmenu()
 #endif
 				}
 
-				{
-					FToolMenuSection& CommonShowFlagsSection =
-						InMenu->FindOrAddSection("CommonShowFlags", LOCTEXT("CommonShowFlagsLabel", "Common Show Flags"));
-
-					FShowFlagMenuCommands::Get().PopulateCommonShowFlagsSection(CommonShowFlagsSection);
-				}
+				// Starting from commonly used flags
+				UE::UnrealEd::AddDefaultShowFlags(InMenu);
 
+				// Add Level Editor specific entries to the All Show Flags Section
 				{
 					FToolMenuSection& AllShowFlagsSection =
 						InMenu->FindOrAddSection("AllShowFlags", LOCTEXT("AllShowFlagsLabel", "All Show Flags"));
 
+					// Show Foliage
 					{
 						FToolMenuEntry ShowFoliageSubmenu = CreateShowFoliageSubmenu();
 						ShowFoliageSubmenu.Label = LOCTEXT("ShowFoliageLabel", "Foliage");
+						ShowFoliageSubmenu.InsertPosition.Position = EToolMenuInsertType::First;
 						AllShowFlagsSection.AddEntry(ShowFoliageSubmenu);
 					}
 
-					AllShowFlagsSection.AddEntry(CreateShowHLODsSubmenu());
-					AllShowFlagsSection.AddEntry(CreateShowLayersSubmenu());
-					AllShowFlagsSection.AddEntry(CreateShowSpritesSubmenu());
-					AllShowFlagsSection.AddEntry(CreateShowVolumesSubmenu());
+					// Show HLODs
+					{
+						FToolMenuEntry ShowHLODSubmenu = CreateShowHLODsSubmenu();
+						ShowHLODSubmenu.InsertPosition.Position = EToolMenuInsertType::First;
+						AllShowFlagsSection.AddEntry(ShowHLODSubmenu);
+					}
 
-					FShowFlagMenuCommands::Get().PopulateAllShowFlagsSection(AllShowFlagsSection);
-				}
+					// Show Layers
+					{
+						FToolMenuEntry ShowLayersSubmenu = CreateShowLayersSubmenu();
+						ShowLayersSubmenu.InsertPosition.Position = EToolMenuInsertType::First;
+						AllShowFlagsSection.AddEntry(ShowLayersSubmenu);
+					}
 
-				// Create these sections for backward compatibility with the old viewport toolbar.
-				{
-					// If your entries end up in this section, you should move it to the new "CommonShowFlags" section instead.
-					InMenu->FindOrAddSection(
-						"ShowFlagsMenuSectionCommon",
-						LOCTEXT("ShowFlagsMenuSectionCommonLabel", "Common Show Flags (Deprecated section)")
-					);
+					// Show Sprites
+					{
+						FToolMenuEntry ShowSpriteSubmenu = CreateShowSpritesSubmenu();
+						ShowSpriteSubmenu.InsertPosition.Position = EToolMenuInsertType::First;
+						AllShowFlagsSection.AddEntry(ShowSpriteSubmenu);
+					}
 
-					// If your entries end up in these sections, you should move them to the above "AllShowFlags" section instead.
-					InMenu->FindOrAddSection(
-						"LevelViewportShowFlags",
-						LOCTEXT("LevelViewportShowFlagsLabel", "All Show Flags (Deprecated section)")
-					);
-					InMenu->FindOrAddSection(
-						"LevelViewportEditorShow", LOCTEXT("LevelViewportEditorShowLabel", "Editor (Deprecated section)")
-					);
+					// Show Volumes
+					{
+						FToolMenuEntry ShowVolumesSubmenu = CreateShowVolumesSubmenu();
+						ShowVolumesSubmenu.InsertPosition.Position = EToolMenuInsertType::First;
+						AllShowFlagsSection.AddEntry(ShowVolumesSubmenu);
+					}
 				}
+
+				// Adds show flags sections for backward compatibility with the old viewport toolbar.
+				//  If your entries end up in this section, you should move it to the new "CommonShowFlags" section instead.
+				InMenu->FindOrAddSection(
+					"ShowFlagsMenuSectionCommon",
+					LOCTEXT("ShowFlagsMenuSectionCommonLabel", "Common Show Flags (Deprecated section)")
+				);
+
+				// If your entries end up in these sections, you should move them to the above "AllShowFlags" section instead.
+				InMenu->FindOrAddSection(
+					"LevelViewportShowFlags", LOCTEXT("LevelViewportShowFlagsLabel", "All Show Flags (Deprecated section)")
+				);
+				InMenu->FindOrAddSection(
+					"LevelViewportEditorShow", LOCTEXT("LevelViewportEditorShowLabel", "Editor (Deprecated section)")
+				);
 			}
 		)
 	);
diff --git a/Engine/Source/Editor/UnrealEd/Private/ViewportToolbar/UnrealEdViewportToolbar.cpp b/Engine/Source/Editor/UnrealEd/Private/ViewportToolbar/UnrealEdViewportToolbar.cpp
--- a/Engine/Source/Editor/UnrealEd/Private/ViewportToolbar/UnrealEdViewportToolbar.cpp
+++ b/Engine/Source/Editor/UnrealEd/Private/ViewportToolbar/UnrealEdViewportToolbar.cpp
@@ -12,11 +12,9 @@
 #include "GPUSkinCacheVisualizationMenuCommands.h"
 #include "IPreviewProfileController.h"
 #include "LevelEditorActions.h"
-#include "PreviewProfileController.h"
 #include "RayTracingDebugVisualizationMenuCommands.h"
-#include "SAssetEditorViewport.h"
-#include "SCommonEditorViewportToolbarBase.h"
 #include "SEditorViewport.h"
+#include "ShowFlagMenuCommands.h"
 #include "Settings/EditorProjectSettings.h"
 #include "Settings/LevelEditorViewportSettings.h"
 #include "Styling/SlateIconFinder.h"
@@ -30,7 +28,6 @@
 #include "Widgets/Input/SSpinBox.h"
 #include "Widgets/Layout/SUniformGridPanel.h"
 #include "Widgets/SToolTip.h"
-#include "Widgets/Input/STextComboBox.h"
 
 #define LOCTEXT_NAMESPACE "UnrealEdViewportToolbar"
 
@@ -2649,6 +2646,38 @@ UNREALED_API FToolMenuEntry CreatePerformanceAndScalabilitySubmenu()
 	);
 }
 
+FToolMenuEntry CreateDefaultShowSubmenu()
+{
+	return FToolMenuEntry::InitSubMenu(
+		"Show",
+		LOCTEXT("ShowSubmenuLabel", "Show"),
+		LOCTEXT("ShowSubmenuTooltip", "Show flags related to the current viewport"),
+		FNewToolMenuDelegate::CreateLambda(
+			[](UToolMenu* InMenu) -> void
+			{
+				AddDefaultShowFlags(InMenu);
+			}
+		)
+	);
+}
+
+void AddDefaultShowFlags(UToolMenu* InMenu)
+{
+	{
+		FToolMenuSection& CommonShowFlagsSection =
+			InMenu->FindOrAddSection("CommonShowFlags", LOCTEXT("CommonShowFlagsLabel", "Common Show Flags"));
+
+		FShowFlagMenuCommands::Get().PopulateCommonShowFlagsSection(CommonShowFlagsSection);
+	}
+
+	{
+		FToolMenuSection& AllShowFlagsSection =
+			InMenu->FindOrAddSection("AllShowFlags", LOCTEXT("AllShowFlagsLabel", "All Show Flags"));
+
+		FShowFlagMenuCommands::Get().PopulateAllShowFlagsSection(AllShowFlagsSection);
+	}
+}
+
 FToolMenuEntry CreateToggleRealtimeEntry()
 {
 	return FToolMenuEntry::InitDynamicEntry(
diff --git a/Engine/Source/Editor/UnrealEd/Public/ViewportToolbar/UnrealEdViewportToolbar.h b/Engine/Source/Editor/UnrealEd/Public/ViewportToolbar/UnrealEdViewportToolbar.h
--- a/Engine/Source/Editor/UnrealEd/Public/ViewportToolbar/UnrealEdViewportToolbar.h
+++ b/Engine/Source/Editor/UnrealEd/Public/ViewportToolbar/UnrealEdViewportToolbar.h
@@ -141,6 +141,16 @@ UNREALED_API void ExtendCameraSubmenu(FName InCameraOptionsSubmenuName);
 
 UNREALED_API FToolMenuEntry CreatePerformanceAndScalabilitySubmenu();
 
+/**
+ * Creates a Show submenu with commonly used show flags
+ */
+UNREALED_API FToolMenuEntry CreateDefaultShowSubmenu();
+
+/**
+ * Adds common flags sections to the specified menu
+ */
+UNREALED_API void AddDefaultShowFlags(UToolMenu* InMenu);
+
 UNREALED_API FToolMenuEntry CreateToggleRealtimeEntry();
 
 // Camera Menu Widgets
