[Actor Actions Usability] Replace Actor: change "Copy Properties" flag default to "false"

diff --git a/Engine/Source/Editor/LevelEditor/Private/LevelEditorCreateActorMenu.cpp b/Engine/Source/Editor/LevelEditor/Private/LevelEditorCreateActorMenu.cpp
--- a/Engine/Source/Editor/LevelEditor/Private/LevelEditorCreateActorMenu.cpp
+++ b/Engine/Source/Editor/LevelEditor/Private/LevelEditorCreateActorMenu.cpp
@@ -441,14 +441,14 @@ void LevelEditorCreateActorMenu::FillAddReplaceContextMenuSections(FToolMenuSect
 	}
 }
 
-bool GReplaceSelectedActorsWithSelectedClassCopyProperties = true;
+bool GReplaceSelectedActorsWithSelectedClassCopyProperties = false;
 void LevelEditorCreateActorMenu::FillAddReplaceActorMenu(UToolMenu* Menu, EActorCreateMode::Type CreateMode)
 {
 	if ( CreateMode == EActorCreateMode::Replace )
 	{
 		FToolMenuSection& Section = Menu->AddSection("Options", NSLOCTEXT("LevelViewportContextMenu", "Options", "Options"));
 
-		GReplaceSelectedActorsWithSelectedClassCopyProperties = true;
+		GReplaceSelectedActorsWithSelectedClassCopyProperties = false;
 
 		FToolMenuEntry ToolMenuEntry = FToolMenuEntry::InitMenuEntry(
 			"CopyProperties",
