Improve the Maxmize/Restore in the new level editor viewport toolbar

* Hide the Maximize/Restore button when the layout doesn't support maximizing viewports.
* Turn the Maximize/Restore button into a normal button instead of a toggle to make it not look highlighted when enabled.
* Set the label, tooltip, and icon dynamically to describe maximizing or restoring depending on the viewport state.

diff --git a/Engine/Source/Editor/LevelEditor/Private/SLevelViewport.cpp b/Engine/Source/Editor/LevelEditor/Private/SLevelViewport.cpp
--- a/Engine/Source/Editor/LevelEditor/Private/SLevelViewport.cpp
+++ b/Engine/Source/Editor/LevelEditor/Private/SLevelViewport.cpp
@@ -1887,6 +1887,15 @@ bool SLevelViewport::IsMaximized() const
 	return true;
 }
 
+bool SLevelViewport::CanMaximize() const
+{
+	if (TSharedPtr<FLevelViewportLayout> PinnedParentLayout = ParentLayout.Pin())
+	{
+		return PinnedParentLayout->IsMaximizeSupported();
+	}
+	return false;
+}
+
 TSharedRef<FEditorViewportClient> SLevelViewport::MakeEditorViewportClient() 
 {
 	return LevelViewportClient.ToSharedRef();
diff --git a/Engine/Source/Editor/LevelEditor/Private/ViewportToolbar/LevelEditorViewportToolbarSections.cpp b/Engine/Source/Editor/LevelEditor/Private/ViewportToolbar/LevelEditorViewportToolbarSections.cpp
--- a/Engine/Source/Editor/LevelEditor/Private/ViewportToolbar/LevelEditorViewportToolbarSections.cpp
+++ b/Engine/Source/Editor/LevelEditor/Private/ViewportToolbar/LevelEditorViewportToolbarSections.cpp
@@ -1681,57 +1681,116 @@ FToolMenuEntry CreateLevelEditorViewportToolbarSettingsSubmenu()
 
 					SettingsSection.AddMenuEntry(FLevelViewportCommands::Get().ToggleImmersive);
 
-					{
-						FToolUIAction MaximizeRestoreAction;
-						MaximizeRestoreAction.ExecuteAction = FToolMenuExecuteAction::CreateLambda(
-							[](const FToolMenuContext& Context)
+					SettingsSection.AddDynamicEntry(
+						"MaximizeRestoreDynamicEntry",
+						FNewToolMenuSectionDelegate::CreateLambda(
+							[](FToolMenuSection& InnerSection) -> void
 							{
 								ULevelViewportContext* const LevelViewportContext =
-									Context.FindContext<ULevelViewportContext>();
+									InnerSection.FindContext<ULevelViewportContext>();
 								if (!LevelViewportContext)
 								{
 									return;
 								}
 
-								if (const TSharedPtr<::SLevelViewport> LevelViewport =
-										LevelViewportContext->LevelViewport.Pin())
-								{
-									LevelViewport->OnToggleMaximize();
-								}
-							}
-						);
-						MaximizeRestoreAction.GetActionCheckState = FToolMenuGetActionCheckState::CreateLambda(
-							[](const FToolMenuContext& Context) -> ECheckBoxState
-							{
-								ULevelViewportContext* const LevelViewportContext =
-									Context.FindContext<ULevelViewportContext>();
-								if (!LevelViewportContext)
-								{
-									return ECheckBoxState::Undetermined;
-								}
+								const TAttribute<FText> Label = TAttribute<FText>::CreateLambda(
+									[WeakLevelViewport = LevelViewportContext->LevelViewport]() -> FText
+									{
+										if (const TSharedPtr<::SLevelViewport> LevelViewport = WeakLevelViewport.Pin())
+										{
+											if (!LevelViewport->IsMaximized())
+											{
+												return LOCTEXT("MaximizeRestoreLabel_Maximize", "Maximize Viewport");
+											}
+										}
+										return LOCTEXT("MaximizeRestoreLabel_Restore", "Restore All Viewports");
+									}
+								);
 
-								if (const TSharedPtr<::SLevelViewport> LevelViewport =
-										LevelViewportContext->LevelViewport.Pin())
-								{
-									return LevelViewport->IsMaximized() ? ECheckBoxState::Checked
-																		: ECheckBoxState::Unchecked;
-								}
+								const TAttribute<FText> Tooltip = TAttribute<FText>::CreateLambda(
+									[WeakLevelViewport = LevelViewportContext->LevelViewport]() -> FText
+									{
+										if (const TSharedPtr<::SLevelViewport> LevelViewport = WeakLevelViewport.Pin())
+										{
+											if (!LevelViewport->IsMaximized())
+											{
+												return LOCTEXT("MaximizeRestoreTooltip_Maximize", "Maximizes this viewport");
+											}
+										}
+										return LOCTEXT(
+											"MaximizeRestoreLabel_Restore", "Restores the layout to show all viewports"
+										);
+									}
+								);
 
-								return ECheckBoxState::Undetermined;
-							}
-						);
+								const TAttribute<FSlateIcon> Icon = TAttribute<FSlateIcon>::CreateLambda(
+									[WeakLevelViewport = LevelViewportContext->LevelViewport]() -> FSlateIcon
+									{
+										if (const TSharedPtr<::SLevelViewport> LevelViewport = WeakLevelViewport.Pin())
+										{
+											if (!LevelViewport->IsMaximized())
+											{
+												return FSlateIcon(
+													FAppStyle::GetAppStyleSetName(), "EditorViewportToolBar.Maximize.Normal"
+												);
+											}
+										}
 
-						SettingsSection
-							.AddMenuEntry(
-								"MaximizeRestore",
-								LOCTEXT("MaximizeRestoreLabel", "Maximize Viewport"),
-								LOCTEXT("MaximizeRestoreTooltip", "Maximizes or restores this viewport"),
-								FSlateIcon(FAppStyle::GetAppStyleSetName(), "EditorViewportToolBar.Maximize.Normal"),
-								MaximizeRestoreAction,
-								EUserInterfaceActionType::ToggleButton
-							)
-							.SetShowInToolbarTopLevel(true);
-					}
+										return FSlateIcon(
+											FAppStyle::GetAppStyleSetName(), "EditorViewportToolBar.Maximize.Checked"
+										);
+									}
+								);
+
+								FToolUIAction Action;
+								Action.ExecuteAction = FToolMenuExecuteAction::CreateLambda(
+									[](const FToolMenuContext& Context) -> void
+									{
+										if (ULevelViewportContext* const LevelViewportContext =
+												Context.FindContext<ULevelViewportContext>())
+										{
+											if (const TSharedPtr<::SLevelViewport> LevelViewport =
+													LevelViewportContext->LevelViewport.Pin())
+											{
+												LevelViewport->OnToggleMaximize();
+											}
+										}
+									}
+								);
+								Action.CanExecuteAction = FToolMenuCanExecuteAction::CreateLambda(
+									[](const FToolMenuContext& Context) -> bool
+									{
+										if (ULevelViewportContext* const LevelViewportContext =
+												Context.FindContext<ULevelViewportContext>())
+										{
+											if (const TSharedPtr<::SLevelViewport> LevelViewport =
+													LevelViewportContext->LevelViewport.Pin())
+											{
+												return LevelViewport->CanMaximize();
+											}
+										}
+
+										return false;
+									}
+								);
+
+								// Only top-level flag the button if it is actually clickable.
+								const TAttribute<bool> TopLevel = TAttribute<bool>::CreateLambda(
+									[WeakLevelViewport = LevelViewportContext->LevelViewport]() -> bool
+									{
+										if (const TSharedPtr<::SLevelViewport> LevelViewport = WeakLevelViewport.Pin())
+										{
+											return LevelViewport->CanMaximize();
+										}
+
+										return false;
+									}
+								);
+								InnerSection.AddMenuEntry("MaximizeRestore", Label, Tooltip, Icon, Action)
+									.SetShowInToolbarTopLevel(TopLevel);
+							}
+						)
+					);
 
 					SettingsSection.AddSeparator("AdvancedSeparator");
 
diff --git a/Engine/Source/Editor/LevelEditor/Public/SLevelViewport.h b/Engine/Source/Editor/LevelEditor/Public/SLevelViewport.h
--- a/Engine/Source/Editor/LevelEditor/Public/SLevelViewport.h
+++ b/Engine/Source/Editor/LevelEditor/Public/SLevelViewport.h
@@ -144,6 +144,11 @@ public:
 	 */
 	bool IsMaximized() const;
 
+	/**
+	 * @return true if this viewport can be maximized, false otherwise
+	 */
+	bool CanMaximize() const;
+
 	/**
 	 * Attempts to switch this viewport into immersive mode
 	 *
