Json Preview Platform for Editor
- New menu options for preview platforms to generate/preview devices from json files
- Currently only enabled for Android preview platform

diff --git a/Engine/Plugins/Runtime/AndroidDeviceProfileSelector/Source/AndroidDeviceProfileSelector/AndroidDeviceProfileSelector.Build.cs b/Engine/Plugins/Runtime/AndroidDeviceProfileSelector/Source/AndroidDeviceProfileSelector/AndroidDeviceProfileSelector.Build.cs
--- a/Engine/Plugins/Runtime/AndroidDeviceProfileSelector/Source/AndroidDeviceProfileSelector/AndroidDeviceProfileSelector.Build.cs
+++ b/Engine/Plugins/Runtime/AndroidDeviceProfileSelector/Source/AndroidDeviceProfileSelector/AndroidDeviceProfileSelector.Build.cs
@@ -1,6 +1,7 @@
 // Copyright Epic Games, Inc. All Rights Reserved.
 using EpicGames.Core;
 using System.IO;
+using System.Linq;
 using UnrealBuildTool;
 
 public class AndroidDeviceProfileSelector : ModuleRules
@@ -32,5 +33,18 @@ public class AndroidDeviceProfileSelector : ModuleRules
 				"Engine",
 			}
 			);
+
+		if(Target.Type == TargetType.Editor)
+		{
+			if (!Target.GlobalDefinitions.Contains("WITH_ANDROID_DEVICE_DETECTION=0"))
+			{
+				PrivateDefinitions.Add("WITH_ANDROID_DEVICE_DETECTION=1");
+				PrivateDependencyModuleNames.Add("AndroidDeviceDetection");
+			}
+
+			PrivateDependencyModuleNames.Add("PIEPreviewDeviceSpecification");
+			PublicDependencyModuleNames.Add("Json");
+			PublicDependencyModuleNames.Add("JsonUtilities");
+		}
 	}
 }
diff --git a/Engine/Plugins/Runtime/AndroidDeviceProfileSelector/Source/AndroidDeviceProfileSelector/Private/AndroidDeviceProfileSelectorModule.cpp b/Engine/Plugins/Runtime/AndroidDeviceProfileSelector/Source/AndroidDeviceProfileSelector/Private/AndroidDeviceProfileSelectorModule.cpp
--- a/Engine/Plugins/Runtime/AndroidDeviceProfileSelector/Source/AndroidDeviceProfileSelector/Private/AndroidDeviceProfileSelectorModule.cpp
+++ b/Engine/Plugins/Runtime/AndroidDeviceProfileSelector/Source/AndroidDeviceProfileSelector/Private/AndroidDeviceProfileSelectorModule.cpp
@@ -3,6 +3,17 @@
 #include "AndroidDeviceProfileSelectorModule.h"
 #include "AndroidDeviceProfileSelector.h"
 #include "Modules/ModuleManager.h"
+#if WITH_EDITOR
+
+#if WITH_ANDROID_DEVICE_DETECTION
+#include "Interfaces/IAndroidDeviceDetection.h"
+#include "Interfaces/IAndroidDeviceDetectionModule.h"
+#endif
+#include "PIEPreviewDeviceSpecification.h"
+#include "JsonObjectConverter.h"
+#include "Misc/FileHelper.h"
+
+#endif
 
 IMPLEMENT_MODULE(FAndroidDeviceProfileSelectorModule, AndroidDeviceProfileSelector);
 
@@ -25,6 +36,84 @@ const FString FAndroidDeviceProfileSelectorModule::GetRuntimeDeviceProfileName()
 	return FString();
 }
 
+#if WITH_EDITOR
+void FAndroidDeviceProfileSelectorModule::ExportDeviceParametersToJson(FString& FolderLocation)
+{
+#if WITH_ANDROID_DEVICE_DETECTION
+	IAndroidDeviceDetection* DeviceDetection;
+	DeviceDetection = FModuleManager::LoadModuleChecked<IAndroidDeviceDetectionModule>("AndroidDeviceDetection").GetAndroidDeviceDetection();
+	DeviceDetection->Initialize(TEXT("ANDROID_HOME"),
+#if PLATFORM_WINDOWS
+		TEXT("platform-tools\\adb.exe"),
+#else
+		TEXT("platform-tools/adb"),
+#endif
+		TEXT("shell getprop"), true);
+
+	TSet<FString> AlreadyExported;
+
+	{
+		FScopeLock ExportLock(DeviceDetection->GetDeviceMapLock());
+
+		const TMap<FString, FAndroidDeviceInfo>& Devices = DeviceDetection->GetDeviceMap();
+		for (auto DeviceTuple : Devices)
+		{
+			const FAndroidDeviceInfo& DeviceInfo = DeviceTuple.Value;
+			FString DeviceName = FString::Printf(TEXT("%s_%s(OS%s)"), *DeviceInfo.DeviceBrand, *DeviceInfo.Model, *DeviceInfo.HumanAndroidVersion);
+			if (!AlreadyExported.Find(DeviceName))
+			{
+				FString ExportPath = FolderLocation / (DeviceName + TEXT(".json"));
+				DeviceDetection->ExportDeviceProfile(ExportPath, DeviceTuple.Key);
+				AlreadyExported.Add(DeviceName);
+			}
+		}
+	}
+	FPlatformProcess::Sleep(1.0f);
+#endif
+}
+
+bool FAndroidDeviceProfileSelectorModule::CanExportDeviceParametersToJson()
+{
+#if WITH_ANDROID_DEVICE_DETECTION
+	return true;
+#else
+	return false;
+#endif
+}
+
+void FAndroidDeviceProfileSelectorModule::GetDeviceParametersFromJson(FString& JsonLocation, TMap<FName, FString>& OutDeviceParameters)
+{
+	TSharedPtr<FJsonObject> JsonRootObject;
+	FString Json;
+
+	if (FFileHelper::LoadFileToString(Json, *JsonLocation))
+	{
+		TSharedRef<TJsonReader<> > JsonReader = TJsonReaderFactory<>::Create(Json);
+		FJsonSerializer::Deserialize(JsonReader, JsonRootObject);
+	}
+
+	FPIEPreviewDeviceSpecifications DeviceSpecs;
+	if (JsonRootObject.IsValid() && FJsonObjectConverter::JsonAttributesToUStruct(JsonRootObject->Values, FPIEPreviewDeviceSpecifications::StaticStruct(), &DeviceSpecs, 0, 0))
+	{
+		FPIEAndroidDeviceProperties& AndroidProperties = DeviceSpecs.AndroidProperties;
+		OutDeviceParameters.Add(FName(TEXT("SRC_GPUFamily")), AndroidProperties.GPUFamily);
+		OutDeviceParameters.Add(FName(TEXT("SRC_GLVersion")), AndroidProperties.GLVersion);
+		OutDeviceParameters.Add(FName(TEXT("SRC_VulkanAvailable")), AndroidProperties.VulkanAvailable ? "true" : "false");
+		OutDeviceParameters.Add(FName(TEXT("SRC_VulkanVersion")), AndroidProperties.VulkanVersion);
+		OutDeviceParameters.Add(FName(TEXT("SRC_AndroidVersion")), AndroidProperties.AndroidVersion);
+		OutDeviceParameters.Add(FName(TEXT("SRC_DeviceMake")), AndroidProperties.DeviceMake);
+		OutDeviceParameters.Add(FName(TEXT("SRC_DeviceModel")), AndroidProperties.DeviceModel);
+		OutDeviceParameters.Add(FName(TEXT("SRC_DeviceBuildNumber")), AndroidProperties.DeviceBuildNumber);
+		OutDeviceParameters.Add(FName(TEXT("SRC_UsingHoudini")), AndroidProperties.UsingHoudini ? "true" : "false");
+		OutDeviceParameters.Add(FName(TEXT("SRC_Hardware")), AndroidProperties.Hardware);
+		OutDeviceParameters.Add(FName(TEXT("SRC_Chipset")), AndroidProperties.Chipset);
+		OutDeviceParameters.Add(FName(TEXT("SRC_TotalPhysicalGB")), AndroidProperties.TotalPhysicalGB);
+		OutDeviceParameters.Add(FName(TEXT("SRC_HMDSystemName")), TEXT(""));
+		OutDeviceParameters.Add(FName(TEXT("SRC_SM5Available")), AndroidProperties.SM5Available ? "true" : "false");
+	}
+}
+#endif
+
 const FString FAndroidDeviceProfileSelectorModule::GetDeviceProfileName()
 {
 	FString ProfileName; 
diff --git a/Engine/Plugins/Runtime/AndroidDeviceProfileSelector/Source/AndroidDeviceProfileSelector/Private/AndroidDeviceProfileSelectorModule.h b/Engine/Plugins/Runtime/AndroidDeviceProfileSelector/Source/AndroidDeviceProfileSelector/Private/AndroidDeviceProfileSelectorModule.h
--- a/Engine/Plugins/Runtime/AndroidDeviceProfileSelector/Source/AndroidDeviceProfileSelector/Private/AndroidDeviceProfileSelectorModule.h
+++ b/Engine/Plugins/Runtime/AndroidDeviceProfileSelector/Source/AndroidDeviceProfileSelector/Private/AndroidDeviceProfileSelectorModule.h
@@ -16,7 +16,12 @@ public:
 	//~ Begin IDeviceProfileSelectorModule Interface
 	virtual const FString GetDeviceProfileName() override;
 	virtual const FString GetRuntimeDeviceProfileName() override;
-
+#if WITH_EDITOR
+	virtual void ExportDeviceParametersToJson(FString& FolderLocation) override;
+	virtual bool CanExportDeviceParametersToJson() override;
+	virtual void GetDeviceParametersFromJson(FString& JsonLocation, TMap<FName, FString>& OutDeviceParameters) override;
+	virtual bool CanGetDeviceParametersFromJson() override { return true; }
+#endif
 	// Set the device parameters this selector will use.
 	virtual void SetSelectorProperties(const TMap<FName, FString>& SelectorProperties) override;
 	virtual bool GetSelectorPropertyValue(const FName& PropertyType, FString& PropertyValueOUT) override;
diff --git a/Engine/Source/Editor/DeviceProfileServices/Private/CheckAndroidDeviceProfileCommandlet.cpp b/Engine/Source/Editor/DeviceProfileServices/Private/CheckAndroidDeviceProfileCommandlet.cpp
--- a/Engine/Source/Editor/DeviceProfileServices/Private/CheckAndroidDeviceProfileCommandlet.cpp
+++ b/Engine/Source/Editor/DeviceProfileServices/Private/CheckAndroidDeviceProfileCommandlet.cpp
@@ -124,33 +124,9 @@ int32 UCheckAndroidDeviceProfileCommandlet::Main(const FString& RawCommandLine)
 					DeviceName = DeviceSpecsFile.Left(ExtensionPos);
 					OutputFilename = DeviceName;
 				}
-				TSharedPtr<FJsonObject> JsonRootObject;
-				FString Json;
 
-				if (FFileHelper::LoadFileToString(Json, *(ParamDeviceSpecsFolder / DeviceSpecsFile)))
-				{
-					TSharedRef<TJsonReader<> > JsonReader = TJsonReaderFactory<>::Create(Json);
-					FJsonSerializer::Deserialize(JsonReader, JsonRootObject);
-				}
-				FPIEPreviewDeviceSpecifications DeviceSpecs;
-				if (JsonRootObject.IsValid() && FJsonObjectConverter::JsonAttributesToUStruct(JsonRootObject->Values, FPIEPreviewDeviceSpecifications::StaticStruct(), &DeviceSpecs, 0, 0))
-				{
-					FPIEAndroidDeviceProperties& AndroidProperties = DeviceSpecs.AndroidProperties;
-					DeviceParameters.Add(FName(TEXT("SRC_GPUFamily")), AndroidProperties.GPUFamily);
-					DeviceParameters.Add(FName(TEXT("SRC_GLVersion")), AndroidProperties.GLVersion);
-					DeviceParameters.Add(FName(TEXT("SRC_VulkanAvailable")), AndroidProperties.VulkanAvailable ? "true" : "false");
-					DeviceParameters.Add(FName(TEXT("SRC_VulkanVersion")), AndroidProperties.VulkanVersion);
-					DeviceParameters.Add(FName(TEXT("SRC_AndroidVersion")), AndroidProperties.AndroidVersion);
-					DeviceParameters.Add(FName(TEXT("SRC_DeviceMake")), AndroidProperties.DeviceMake);
-					DeviceParameters.Add(FName(TEXT("SRC_DeviceModel")), AndroidProperties.DeviceModel);
-					DeviceParameters.Add(FName(TEXT("SRC_DeviceBuildNumber")), AndroidProperties.DeviceBuildNumber);
-					DeviceParameters.Add(FName(TEXT("SRC_UsingHoudini")), AndroidProperties.UsingHoudini ? "true" : "false");
-					DeviceParameters.Add(FName(TEXT("SRC_Hardware")), AndroidProperties.Hardware);
-					DeviceParameters.Add(FName(TEXT("SRC_Chipset")), AndroidProperties.Chipset);
-					DeviceParameters.Add(FName(TEXT("SRC_TotalPhysicalGB")), AndroidProperties.TotalPhysicalGB);
-					DeviceParameters.Add(FName(TEXT("SRC_HMDSystemName")), TEXT(""));
-					DeviceParameters.Add(FName(TEXT("SRC_SM5Available")), AndroidProperties.SM5Available ? "true" : "false");
-				}
+				FString JsonLocation = ParamDeviceSpecsFolder / DeviceSpecsFile;
+				AndroidDeviceProfileSelector->GetDeviceParametersFromJson(JsonLocation, DeviceParameters);
 			}
 			else
 			{
diff --git a/Engine/Source/Editor/LevelEditor/LevelEditor.Build.cs b/Engine/Source/Editor/LevelEditor/LevelEditor.Build.cs
--- a/Engine/Source/Editor/LevelEditor/LevelEditor.Build.cs
+++ b/Engine/Source/Editor/LevelEditor/LevelEditor.Build.cs
@@ -84,7 +84,9 @@ public class LevelEditor : ModuleRules
 				"UncontrolledChangelists",
 				"RenderCore",
 				"DeveloperSettings",
-				"ActionableMessage"
+				"ActionableMessage",
+				"Json",
+				"JsonUtilities"
 			}
 		);
 
diff --git a/Engine/Source/Editor/LevelEditor/Private/LevelEditor.cpp b/Engine/Source/Editor/LevelEditor/Private/LevelEditor.cpp
--- a/Engine/Source/Editor/LevelEditor/Private/LevelEditor.cpp
+++ b/Engine/Source/Editor/LevelEditor/Private/LevelEditor.cpp
@@ -1797,6 +1797,8 @@ void FLevelEditorModule::BindGlobalLevelEditorCommands()
 	// We need one extra slot for the Disable Preview option
 	check(MenuItems.Num() == Commands.PreviewPlatformOverrides.Num());
 
+	TSet<FName> PreviewShaderPlatformNames;
+	TSet<FName> PlatformNames;
 	for (int32 Index=0; Index < MenuItems.Num(); Index++)
 	{
 		const FPreviewPlatformMenuItem& Item = MenuItems[Index];
@@ -1824,6 +1826,31 @@ void FLevelEditorModule::BindGlobalLevelEditorCommands()
 				FExecuteAction::CreateStatic(&FLevelEditorActionCallbacks::SetPreviewPlatform, PreviewFeatureLevelInfo),
 				bIsDefaultShaderPlatform ? FCanExecuteAction() : FCanExecuteAction::CreateStatic(&FLevelEditorActionCallbacks::CanExecutePreviewPlatform, PreviewFeatureLevelInfo),
 				FIsActionChecked::CreateStatic(&FLevelEditorActionCallbacks::IsPreviewPlatformChecked, PreviewFeatureLevelInfo));
+
+			FConfigCacheIni* PlatformEngineIni = FConfigCacheIni::ForPlatform(*Item.PlatformName.ToString());
+			FString DeviceProfileSelectionModule;
+			if (PlatformEngineIni && PlatformEngineIni->GetString(TEXT("DeviceProfileManager"), TEXT("PreviewDeviceProfileSelectionModule"), DeviceProfileSelectionModule, GEngineIni))
+			{
+				if (!PlatformNames.Find(Item.PlatformName))
+				{
+					FUIAction Action;
+					Action.ExecuteAction = FExecuteAction::CreateStatic(&FLevelEditorActionCallbacks::GeneratePreviewJson_Clicked, Item.PlatformName.ToString());
+					Action.IsActionVisibleDelegate = FIsActionButtonVisible::CreateStatic(&FLevelEditorActionCallbacks::IsGeneratePreviewJsonVisible, Item.PlatformName);
+
+					ActionList.MapAction(Commands.GeneratePlatformJson[PlatformNames.Num()], Action);
+					PlatformNames.Add(Item.PlatformName);
+				}
+
+				if (!PreviewShaderPlatformNames.Find(Item.PreviewShaderPlatformName))
+				{
+					FUIAction Action;
+					Action.ExecuteAction = FExecuteAction::CreateStatic(&FLevelEditorActionCallbacks::PreviewJson_Clicked, Item.PlatformName, Item.PreviewShaderPlatformName);
+					Action.IsActionVisibleDelegate = FIsActionButtonVisible::CreateStatic(&FLevelEditorActionCallbacks::IsPreviewJsonVisible, Item.PlatformName);
+
+					ActionList.MapAction(Commands.PreviewPlatformFromJson[PreviewShaderPlatformNames.Num()], Action);
+					PreviewShaderPlatformNames.Add(Item.PreviewShaderPlatformName);
+				}
+			}
 		}
 	}
 
diff --git a/Engine/Source/Editor/LevelEditor/Private/LevelEditorActions.cpp b/Engine/Source/Editor/LevelEditor/Private/LevelEditorActions.cpp
--- a/Engine/Source/Editor/LevelEditor/Private/LevelEditorActions.cpp
+++ b/Engine/Source/Editor/LevelEditor/Private/LevelEditorActions.cpp
@@ -118,6 +118,15 @@
 #include "DataDrivenShaderPlatformInfo.h"
 
 #include "Internationalization/Culture.h"
+#include "Misc/FileHelper.h"
+#include "EditorDirectories.h"
+#include "IDesktopPlatform.h"
+#include "DesktopPlatformModule.h"
+#include "Serialization/JsonReader.h"
+#include "Serialization/JsonSerializer.h"
+#include "IDeviceProfileSelectorModule.h"
+#include "DeviceProfiles/DeviceProfile.h"
+#include "DeviceProfiles/DeviceProfileManager.h"
 
 #if WITH_LIVE_CODING
 #include "ILiveCodingModule.h"
@@ -537,6 +546,119 @@ void FLevelEditorActionCallbacks::ImportScene_Clicked()
 	FEditorFileUtils::Import();
 }
 
+void FLevelEditorActionCallbacks::PreviewJson_Clicked(FName PlatformName, FName PreviewShaderPlatformName)
+{
+	TArray<FString> OpenedFiles;
+	FString DefaultLocation(FEditorDirectories::Get().GetLastDirectory(ELastDirectory::PROJECT));
+
+	IDesktopPlatform* DesktopPlatform = FDesktopPlatformModule::Get();
+	bool bOpened = false;
+	if (DesktopPlatform)
+	{
+		bOpened = DesktopPlatform->OpenFileDialog(
+			FSlateApplication::Get().FindBestParentWindowHandleForDialogs(nullptr),
+			NSLOCTEXT("UnrealEd", "PreviewJson", "Preview Json").ToString(),
+			DefaultLocation,
+			TEXT(""),
+			TEXT("*.json"),
+			EFileDialogFlags::None,
+			OpenedFiles
+		);
+	}
+
+	if (bOpened && OpenedFiles.Num() > 0 && OpenedFiles[0].IsEmpty() == false)
+	{
+		FConfigCacheIni* PlatformEngineIni = FConfigCacheIni::ForPlatform(*PlatformName.ToString());
+		FString DeviceProfileSelectionModule;
+		if (PlatformEngineIni && PlatformEngineIni->GetString(TEXT("DeviceProfileManager"), TEXT("PreviewDeviceProfileSelectionModule"), DeviceProfileSelectionModule, GEngineIni))
+		{
+			if (IDeviceProfileSelectorModule* DPSelectorModule = FModuleManager::LoadModulePtr<IDeviceProfileSelectorModule>(*DeviceProfileSelectionModule))
+			{
+				FString DevileProfileName;
+
+				TMap<FName, FString> DeviceParameters;
+				DPSelectorModule->GetDeviceParametersFromJson(OpenedFiles[0], DeviceParameters);
+				DPSelectorModule->SetSelectorProperties(DeviceParameters);
+
+				DevileProfileName = DPSelectorModule->GetDeviceProfileName();
+				
+				UDeviceProfile* DevileProfile = UDeviceProfileManager::Get().FindProfile(DevileProfileName, false);
+				if (DevileProfile)
+				{
+					EShaderPlatform ShaderPlatform = FDataDrivenShaderPlatformInfo::GetShaderPlatformFromName(PreviewShaderPlatformName);
+
+					auto GetPreviewFeatureLevelInfo = [&]()
+						{
+							const ERHIFeatureLevel::Type FeatureLevel = GetMaxSupportedFeatureLevel(ShaderPlatform);
+							return FPreviewPlatformInfo(FeatureLevel, ShaderPlatform, PlatformName, FDataDrivenShaderPlatformInfo::GetShaderFormat(ShaderPlatform), DevileProfile->GetFName(), true, PreviewShaderPlatformName, FText::FromName(DevileProfile->GetFName()));
+						};
+
+					FPreviewPlatformInfo PreviewFeatureLevelInfo = GetPreviewFeatureLevelInfo();
+					GEditor->SetPreviewPlatform(PreviewFeatureLevelInfo, false);
+				}
+			}
+		}
+	}
+}
+
+bool FLevelEditorActionCallbacks::IsPreviewJsonVisible(FName PlatformName)
+{
+	FConfigCacheIni* PlatformEngineIni = FConfigCacheIni::ForPlatform(*PlatformName.ToString());
+	FString DeviceProfileSelectionModule;
+	if (PlatformEngineIni && PlatformEngineIni->GetString(TEXT("DeviceProfileManager"), TEXT("PreviewDeviceProfileSelectionModule"), DeviceProfileSelectionModule, GEngineIni))
+	{
+		if (IDeviceProfileSelectorModule* DPSelectorModule = FModuleManager::LoadModulePtr<IDeviceProfileSelectorModule>(*DeviceProfileSelectionModule))
+		{
+			return DPSelectorModule->CanGetDeviceParametersFromJson();
+		}
+	}
+	return false;
+}
+
+bool FLevelEditorActionCallbacks::IsGeneratePreviewJsonVisible(FName PlatformName)
+{
+	FConfigCacheIni* PlatformEngineIni = FConfigCacheIni::ForPlatform(*PlatformName.ToString());
+	FString DeviceProfileSelectionModule;
+	if (PlatformEngineIni && PlatformEngineIni->GetString(TEXT("DeviceProfileManager"), TEXT("PreviewDeviceProfileSelectionModule"), DeviceProfileSelectionModule, GEngineIni))
+	{
+		if (IDeviceProfileSelectorModule* DPSelectorModule = FModuleManager::LoadModulePtr<IDeviceProfileSelectorModule>(*DeviceProfileSelectionModule))
+		{
+			return DPSelectorModule->CanExportDeviceParametersToJson();
+		}
+	}
+	return false;
+}
+
+void FLevelEditorActionCallbacks::GeneratePreviewJson_Clicked(FString PlatformName)
+{
+	FString DefaultLocation(FEditorDirectories::Get().GetLastDirectory(ELastDirectory::PROJECT));
+	FString OutputDirectory;
+	IDesktopPlatform* DesktopPlatform = FDesktopPlatformModule::Get();
+
+	if (DesktopPlatform &&
+		DesktopPlatform->OpenDirectoryDialog(FSlateApplication::Get().FindBestParentWindowHandleForDialogs(nullptr), 
+		LOCTEXT("LocationToOutputJsons", "Location to Output Json Files").ToString(), 
+		DefaultLocation,
+		OutputDirectory)
+		)
+	{
+		FScopedSlowTask SlowTask(100.f, NSLOCTEXT("Engine", "GeneratePlatformJson", "Generate Platform Json"), true);
+		SlowTask.Visibility = ESlowTaskVisibility::ForceVisible;
+		SlowTask.MakeDialog();
+
+		SlowTask.EnterProgressFrame(35.0f);
+
+		FConfigCacheIni* PlatformEngineIni = FConfigCacheIni::ForPlatform(*PlatformName);
+		FString DeviceProfileSelectionModule;
+		if (PlatformEngineIni && PlatformEngineIni->GetString(TEXT("DeviceProfileManager"), TEXT("PreviewDeviceProfileSelectionModule"), DeviceProfileSelectionModule, GEngineIni))
+		{
+			if (IDeviceProfileSelectorModule* DPSelectorModule = FModuleManager::LoadModulePtr<IDeviceProfileSelectorModule>(*DeviceProfileSelectionModule))
+			{
+				DPSelectorModule->ExportDeviceParametersToJson(OutputDirectory);
+			}
+		}
+	}
+}
 
 void FLevelEditorActionCallbacks::ExportAll_Clicked()
 {
@@ -3992,6 +4114,8 @@ void FLevelEditorCommands::RegisterCommands()
 	UI_COMMAND(ToggleFeatureLevelPreview, "Preview Mode Toggle", "Toggles the Preview Mode on or off for the currently selected Preview target", EUserInterfaceActionType::ToggleButton, FInputChord());
 
 	// Add preview platforms
+	TSet<FName> PreviewShaderPlatformNames;
+	TSet<FName> PlatformName;
 	for (const FPreviewPlatformMenuItem& Item : FDataDrivenPlatformInfoRegistry::GetAllPreviewPlatformMenuItems())
 	{
 		FTextBuilder FriendlyNameBuilder;
@@ -4029,6 +4153,47 @@ void FLevelEditorCommands::RegisterCommands()
 		);
 
 		PlatformToPreviewPlatformOverrides.FindOrAdd(Item.PlatformName).Add(PreviewPlatformOverrides.Last());
+
+		FConfigCacheIni* PlatformEngineIni = FConfigCacheIni::ForPlatform(*Item.PlatformName.ToString());
+		FString DeviceProfileSelectionModule;
+
+		if (PlatformEngineIni && PlatformEngineIni->GetString(TEXT("DeviceProfileManager"), TEXT("PreviewDeviceProfileSelectionModule"), DeviceProfileSelectionModule, GEngineIni))
+		{
+			EShaderPlatform ShaderPlatform = FDataDrivenShaderPlatformInfo::GetShaderPlatformFromName(Item.ShaderPlatformToPreview);
+			FText PlatformFriendlyName = FDataDrivenShaderPlatformInfo::GetFriendlyName(ShaderPlatform);
+
+			if (!PlatformName.Find(Item.PlatformName))
+			{
+				GeneratePlatformJson.Add(
+					FUICommandInfoDecl(
+						this->AsShared(),
+						FName(*FString::Printf(TEXT("Generate Platform Json for %s"), *Item.PlatformName.ToString())),
+						NSLOCTEXT("GeneratePlatformJson", "Generate Platform Json", "Generate Platform Json"),
+						NSLOCTEXT("GeneratePlatformJsonDesc", "Generate Platform Json Desc", "Generate Platform Json Desc"))
+					.UserInterfaceType(EUserInterfaceActionType::Button)
+					.DefaultChord(FInputChord())
+				);
+
+				PlatformToPreviewJsonPlatformOverrides.FindOrAdd(Item.PlatformName).Add(GeneratePlatformJson.Last());
+				PlatformName.Add(Item.PlatformName);
+			}
+
+			if (!PreviewShaderPlatformNames.Find(Item.PreviewShaderPlatformName))
+			{
+				PreviewPlatformFromJson.Add(
+					FUICommandInfoDecl(
+						this->AsShared(),
+						FName(*FString::Printf(TEXT("Preview %s with Json"), *Item.PreviewShaderPlatformName.ToString())),
+						FText::Format(NSLOCTEXT("PreviewJson", "Preview Json", "Preview {0} with Json"), PlatformFriendlyName),
+						NSLOCTEXT("PreviewJsonDesc", "Preview Json Desc", "Preview Json Desc"))
+					.UserInterfaceType(EUserInterfaceActionType::Button)
+					.DefaultChord(FInputChord())
+				);
+
+				PlatformToPreviewJsonPlatformOverrides.FindOrAdd(Item.PlatformName).Add(PreviewPlatformFromJson.Last());
+				PreviewShaderPlatformNames.Add(Item.PreviewShaderPlatformName);
+			}
+		}
 	}
 
 	PlatformToPreviewPlatformOverrides.KeyStableSort([](FName lhs, FName rhs) {return lhs.Compare(rhs) < 0; });
diff --git a/Engine/Source/Editor/LevelEditor/Private/LevelEditorToolBar.cpp b/Engine/Source/Editor/LevelEditor/Private/LevelEditorToolBar.cpp
--- a/Engine/Source/Editor/LevelEditor/Private/LevelEditorToolBar.cpp
+++ b/Engine/Source/Editor/LevelEditor/Private/LevelEditorToolBar.cpp
@@ -81,15 +81,13 @@ namespace PreviewModeFunctionality
 	FText GetPreviewModeTooltip()
 	{
 #define LOCTEXT_NAMESPACE "LevelEditorToolBar"
-		EShaderPlatform PreviewShaderPlatform = GEditor->PreviewPlatform.PreviewShaderPlatformName != NAME_None ?
-			FDataDrivenShaderPlatformInfo::GetShaderPlatformFromName(GEditor->PreviewPlatform.PreviewShaderPlatformName) :
-			GetFeatureLevelShaderPlatform(GEditor->PreviewPlatform.PreviewFeatureLevel);
+		EShaderPlatform PreviewShaderPlatform = GEditor->PreviewPlatform.GetShaderPlatform();
 
 		EShaderPlatform MaxRHIFeatureLevelPlatform = GetFeatureLevelShaderPlatform(GMaxRHIFeatureLevel);
 
 		{
-			const FText& RenderingAsPlatformName = FDataDrivenShaderPlatformInfo::GetFriendlyName(GEditor->PreviewPlatform.bPreviewFeatureLevelActive ? PreviewShaderPlatform : MaxRHIFeatureLevelPlatform);
-			const FText& SwitchToPlatformName = FDataDrivenShaderPlatformInfo::GetFriendlyName(GEditor->PreviewPlatform.bPreviewFeatureLevelActive ? MaxRHIFeatureLevelPlatform : PreviewShaderPlatform);
+			const FText& RenderingAsPlatformName = GEditor->PreviewPlatform.bPreviewFeatureLevelActive ? GEditor->PreviewPlatform.GetFriendlyName() : FDataDrivenShaderPlatformInfo::GetFriendlyName(MaxRHIFeatureLevelPlatform);
+			const FText& SwitchToPlatformName = GEditor->PreviewPlatform.bPreviewFeatureLevelActive ? FDataDrivenShaderPlatformInfo::GetFriendlyName(MaxRHIFeatureLevelPlatform) :GEditor->PreviewPlatform.GetFriendlyName();
 			if (PreviewShaderPlatform == MaxRHIFeatureLevelPlatform)
 			{
 				return FText::Format(LOCTEXT("PreviewModeViewingAs", "Viewing {0}."), RenderingAsPlatformName);
@@ -110,7 +108,7 @@ namespace PreviewModeFunctionality
 	{
 		const FPreviewPlatformMenuItem* Item = FDataDrivenPlatformInfoRegistry::GetAllPreviewPlatformMenuItems().FindByPredicate([](const FPreviewPlatformMenuItem& TestItem)
 			{
-				return GEditor->PreviewPlatform.PreviewPlatformName == TestItem.PlatformName && GEditor->PreviewPlatform.PreviewShaderFormatName == TestItem.ShaderFormat && GEditor->PreviewPlatform.PreviewShaderPlatformName == TestItem.PreviewShaderPlatformName;
+				return GEditor->PreviewPlatform.PreviewPlatformName == TestItem.PlatformName;
 			});
 		if (Item)
 		{
diff --git a/Engine/Source/Editor/LevelEditor/Private/ViewportToolbar/LevelEditorViewportToolbarSections.cpp b/Engine/Source/Editor/LevelEditor/Private/ViewportToolbar/LevelEditorViewportToolbarSections.cpp
--- a/Engine/Source/Editor/LevelEditor/Private/ViewportToolbar/LevelEditorViewportToolbarSections.cpp
+++ b/Engine/Source/Editor/LevelEditor/Private/ViewportToolbar/LevelEditorViewportToolbarSections.cpp
@@ -1331,16 +1331,38 @@ FToolMenuEntry CreateFeatureLevelPreviewSubmenu()
 				{
 					FName PlatformName = Iter.Key();
 					const TArray<TSharedPtr<FUICommandInfo>>& CommandList = Iter.Value();
+					const TArray<TSharedPtr<FUICommandInfo>>* CommandListJson = FLevelEditorCommands::Get().PlatformToPreviewJsonPlatformOverrides.Find(PlatformName);
 
-					Section.AddSubMenu(FName(PlatformName), FText::FromString(PlatformName.ToString()), FText(),
-						FNewToolMenuDelegate::CreateLambda([CommandList](UToolMenu* InSubMenu)
+					Section.AddSubMenu(
+						FName(PlatformName), 
+						FText::FromString(PlatformName.ToString()), 
+						FText(),
+						FNewToolMenuDelegate::CreateLambda(
+							[CommandList, CommandListJson](UToolMenu* InSubMenu)
 							{
-								FToolMenuSection& Section =
-									InSubMenu->AddSection("", LOCTEXT("EditorPreviewModeDevices", "Preview Devices"));
+								FToolMenuSection& Section = InSubMenu->AddSection("", LOCTEXT("EditorPreviewModeDevices", "Preview Devices"));
 								for (const TSharedPtr<FUICommandInfo>& Command : CommandList)
 								{
 									Section.AddMenuEntry(Command);
 								}
+
+								if (CommandListJson != nullptr)
+								{
+									Section.AddSubMenu(
+										"PreviewJson",
+										LOCTEXT("PreviewJson", "Preview Json"),
+										LOCTEXT("PreviewJson_ToolTip", "Preview Json"),
+										FNewToolMenuDelegate::CreateLambda(
+											[CommandListJson](UToolMenu* InSubMenu)
+											{
+												FToolMenuSection& Section = InSubMenu->AddSection("", LOCTEXT("EditorPreviewModeDevicesJson", "Preview Devices Json"));
+												for (const TSharedPtr<FUICommandInfo>& Command : *CommandListJson)
+												{
+													Section.AddMenuEntry(Command);
+												}
+											})
+									);
+								}
 							})
 					);
 				}
diff --git a/Engine/Source/Editor/LevelEditor/Public/LevelEditorActions.h b/Engine/Source/Editor/LevelEditor/Public/LevelEditorActions.h
--- a/Engine/Source/Editor/LevelEditor/Public/LevelEditorActions.h
+++ b/Engine/Source/Editor/LevelEditor/Public/LevelEditorActions.h
@@ -610,7 +610,10 @@ public:
 	TSharedPtr< FUICommandInfo > ToggleFeatureLevelPreview;
 
 	TArray<TSharedPtr<FUICommandInfo>> PreviewPlatformOverrides;
+	TArray<TSharedPtr<FUICommandInfo>> PreviewPlatformFromJson;
+	TArray<TSharedPtr<FUICommandInfo>> GeneratePlatformJson;
 	TMap<FName, TArray<TSharedPtr<FUICommandInfo>>> PlatformToPreviewPlatformOverrides;
+	TMap<FName, TArray<TSharedPtr<FUICommandInfo>>> PlatformToPreviewJsonPlatformOverrides;
 
 	///**
 	// * Mode Commands                   
@@ -730,7 +733,26 @@ public:
 	 * Called when import scene is selected
 	 */
 	static void ImportScene_Clicked();
+	
+	/**
+	* Called When Preview Json is selected in the Platforms Preview Sub Menu
+	*/
+	static void PreviewJson_Clicked(FName PlatformName, FName PreviewShaderPlatformName);
+
+	/**
+	* Is Preview Json visible in the Platforms Preview Sub Menu
+	*/
+	static bool IsPreviewJsonVisible(FName PlatformName);
 
+	/**
+	* Called When Generate Preview Json is selected in the Platforms Preview Sub Menu
+	*/
+	static void GeneratePreviewJson_Clicked(FString PlatformName);
+	
+	/**
+	* Is Generate Preview Json visible in the Platforms Preview Sub Menu
+	*/
+	static bool IsGeneratePreviewJsonVisible(FName PlatformName);
 
 	/**
 	 * Called when export all is selected
diff --git a/Engine/Source/Editor/UnrealEd/Classes/Editor/EditorEngine.h b/Engine/Source/Editor/UnrealEd/Classes/Editor/EditorEngine.h
--- a/Engine/Source/Editor/UnrealEd/Classes/Editor/EditorEngine.h
+++ b/Engine/Source/Editor/UnrealEd/Classes/Editor/EditorEngine.h
@@ -227,9 +227,12 @@ struct FPreviewPlatformInfo
 	,	PreviewShaderFormatName(NAME_None)
 	,	bPreviewFeatureLevelActive(false)
 	,	PreviewShaderPlatformName(NAME_None)
-	{}
+	,	PreviewShaderPlatformFriendlyName()
+	{
+		InternalSetFriendlyName();
+	}
 
-	FPreviewPlatformInfo(ERHIFeatureLevel::Type InFeatureLevel, EShaderPlatform InShaderPlatform = EShaderPlatform::SP_NumPlatforms, FName InPreviewPlatformName = NAME_None, FName InPreviewShaderFormatName = NAME_None, FName InDeviceProfileName = NAME_None, bool InbPreviewFeatureLevelActive = false, FName InShaderPlatformName = NAME_None)
+	FPreviewPlatformInfo(ERHIFeatureLevel::Type InFeatureLevel, EShaderPlatform InShaderPlatform = EShaderPlatform::SP_NumPlatforms, FName InPreviewPlatformName = NAME_None, FName InPreviewShaderFormatName = NAME_None, FName InDeviceProfileName = NAME_None, bool InbPreviewFeatureLevelActive = false, FName InShaderPlatformName = NAME_None, FText InPreviewShaderPlatformFriendlyName = FText())
 	:	PreviewFeatureLevel(InFeatureLevel)
 	,	ShaderPlatform(InShaderPlatform)
 	,	PreviewPlatformName(InPreviewPlatformName)
@@ -237,7 +240,10 @@ struct FPreviewPlatformInfo
 	,	DeviceProfileName(InDeviceProfileName)
 	,	bPreviewFeatureLevelActive(InbPreviewFeatureLevelActive)
 	,	PreviewShaderPlatformName(InShaderPlatformName)
-	{}
+	,	PreviewShaderPlatformFriendlyName(InPreviewShaderPlatformFriendlyName)
+	{
+		InternalSetFriendlyName();
+	}
 
 	/** The feature level we should use when loading or creating a new world */
 	ERHIFeatureLevel::Type PreviewFeatureLevel;
@@ -274,6 +280,17 @@ struct FPreviewPlatformInfo
 
 	/** returns the preview feature level if active, or GMaxRHIFeatureLevel otherwise */
 	UNREALED_API ERHIFeatureLevel::Type GetEffectivePreviewFeatureLevel() const;
+	
+	/** Returns the Friendly Name of the Preview*/
+	UNREALED_API FText GetFriendlyName() const { return PreviewShaderPlatformFriendlyName; }
+	
+	/** Returns the Shader Platform of the Preview*/
+	UNREALED_API EShaderPlatform GetShaderPlatform() const;
+private:
+	/**Cache Preview ShaderPlatform Name*/
+	FText PreviewShaderPlatformFriendlyName;
+	
+	UNREALED_API void InternalSetFriendlyName();
 };
 
 enum class EAssetReferenceFilterProperties : uint8
diff --git a/Engine/Source/Editor/UnrealEd/Private/EditorEngine.cpp b/Engine/Source/Editor/UnrealEd/Private/EditorEngine.cpp
--- a/Engine/Source/Editor/UnrealEd/Private/EditorEngine.cpp
+++ b/Engine/Source/Editor/UnrealEd/Private/EditorEngine.cpp
@@ -448,6 +448,23 @@ ERHIFeatureLevel::Type FPreviewPlatformInfo::GetEffectivePreviewFeatureLevel() c
 	return bPreviewFeatureLevelActive ? PreviewFeatureLevel : GMaxRHIFeatureLevel;
 }
 
+EShaderPlatform FPreviewPlatformInfo::GetShaderPlatform() const
+{
+	return PreviewShaderPlatformName != NAME_None ?
+		FDataDrivenShaderPlatformInfo::GetShaderPlatformFromName(PreviewShaderPlatformName) :
+		GetFeatureLevelShaderPlatform(PreviewFeatureLevel);
+}
+
+void FPreviewPlatformInfo::InternalSetFriendlyName()
+{
+	if (PreviewShaderPlatformFriendlyName.IsEmpty())
+	{
+		EShaderPlatform PreviewShaderPlatform = GetShaderPlatform();
+		EShaderPlatform MaxRHIFeatureLevelPlatform = GetFeatureLevelShaderPlatform(GMaxRHIFeatureLevel);
+		PreviewShaderPlatformFriendlyName = FDataDrivenShaderPlatformInfo::GetFriendlyName(bPreviewFeatureLevelActive ? PreviewShaderPlatform : MaxRHIFeatureLevelPlatform);
+	}
+}
+
 void FAssetReferenceFilterContext::AddReferencingAsset(const FAssetData& InReferencingAsset, EAssetReferenceFilterProperties InProperties)
 {
 	if (InReferencingAsset.IsValid())
diff --git a/Engine/Source/Editor/UnrealEd/Private/SEditorViewport.cpp b/Engine/Source/Editor/UnrealEd/Private/SEditorViewport.cpp
--- a/Engine/Source/Editor/UnrealEd/Private/SEditorViewport.cpp
+++ b/Engine/Source/Editor/UnrealEd/Private/SEditorViewport.cpp
@@ -1010,8 +1010,7 @@ FText SEditorViewport::GetCurrentFeatureLevelPreviewText(bool bDrawOnlyLabel) co
 		if (World != nullptr)
 		{
 			ERHIFeatureLevel::Type TargetFeatureLevel = World->GetFeatureLevel();
-			EShaderPlatform ShaderPlatform = GetShaderPlatformHelper(TargetFeatureLevel);
-			const FText& PlatformText = FDataDrivenShaderPlatformInfo::GetFriendlyName(ShaderPlatform);
+			const FText& PlatformText = GEditor->PreviewPlatform.GetFriendlyName();
 			LabelName = FText::Format(LOCTEXT("WorldFeatureLevel", "{0}"), PlatformText);
 		}
 	}
diff --git a/Engine/Source/Runtime/Engine/Public/IDeviceProfileSelectorModule.h b/Engine/Source/Runtime/Engine/Public/IDeviceProfileSelectorModule.h
--- a/Engine/Source/Runtime/Engine/Public/IDeviceProfileSelectorModule.h
+++ b/Engine/Source/Runtime/Engine/Public/IDeviceProfileSelectorModule.h
@@ -8,7 +8,9 @@
 
 #include "CoreMinimal.h"
 #include "Modules/ModuleInterface.h"
+#include "DeviceProfiles/DeviceProfileMatching.h"
 
+class UDeviceProfile;
 
 /**
  * Device Profile Selector module
@@ -25,6 +27,33 @@ public:
 	 */
 	virtual const FString GetRuntimeDeviceProfileName() = 0;
 
+#if WITH_EDITOR
+	/**
+	 * Save Device Jsons to the location provided
+	 * @param FolderLocation Folder where to store the Jsons
+	 */
+	virtual void ExportDeviceParametersToJson(FString& FolderLocation){}
+	
+	/**
+	 * Can we export device parameters to Jsons
+	 *  @return Wheter we can export device parameters to Jsonsn
+	 */
+	virtual bool CanExportDeviceParametersToJson() { return false; }
+
+	/**
+	 * Get the Device Parameters from the Json
+	* @Param JsonLocation Json location
+	* @Param OutDeviceParameters the device parameters read from Json
+	 */
+	virtual void GetDeviceParametersFromJson(FString& JsonLocation, TMap<FName, FString>& OutDeviceParameters) {}
+	
+	/**
+	 * Can we get device parameters from Json
+	 *  @return Wheter we can get device parameters from Json
+	 */
+	virtual bool CanGetDeviceParametersFromJson() { return false; }
+#endif
+
 	/**
 	* Run the logic to choose an appropriate device profile for this session.
 	* @param DeviceParameters	A map of parameters to be used by device profile logic.
