Fix LOCTEXT bug

diff --git a/Engine/Source/Editor/LevelEditor/Private/ViewportToolbar/LevelEditorViewportToolbarSections.cpp b/Engine/Source/Editor/LevelEditor/Private/ViewportToolbar/LevelEditorViewportToolbarSections.cpp
--- a/Engine/Source/Editor/LevelEditor/Private/ViewportToolbar/LevelEditorViewportToolbarSections.cpp
+++ b/Engine/Source/Editor/LevelEditor/Private/ViewportToolbar/LevelEditorViewportToolbarSections.cpp
@@ -1740,7 +1740,7 @@ FToolMenuEntry CreateLevelEditorViewportToolbarSettingsSubmenu()
 											}
 										}
 										return LOCTEXT(
-											"MaximizeRestoreLabel_Restore", "Restores the layout to show all viewports"
+											"MaximizeRestoreTooltip_Restore", "Restores the layout to show all viewports"
 										);
 									}
 								);
