[Viewport Toolbar] Add Actor camera selection and Pilot entries to the "Camera" submenu for the new level editor viewport toolbar

diff --git a/Engine/Source/Editor/LevelEditor/Private/SLevelViewport.cpp b/Engine/Source/Editor/LevelEditor/Private/SLevelViewport.cpp
--- a/Engine/Source/Editor/LevelEditor/Private/SLevelViewport.cpp
+++ b/Engine/Source/Editor/LevelEditor/Private/SLevelViewport.cpp
@@ -2048,7 +2048,7 @@ TSharedPtr<SWidget> SLevelViewport::BuildViewportToolbar()
 					UE::LevelEditor::ExtendCameraSubmenu(SubmenuName);
 					UE::UnrealEd::ExtendCameraSubmenu(SubmenuName);
 
-					FToolMenuEntry CameraSubmenu = UE::UnrealEd::CreateViewportToolbarCameraSubmenu();
+					FToolMenuEntry CameraSubmenu = UE::LevelEditor::CreateLevelViewportToolbarCameraSubmenu();
 					CameraSubmenu.InsertPosition.Position = EToolMenuInsertType::First;
 					RightSection.AddEntry(CameraSubmenu);
 				}
@@ -2775,6 +2775,11 @@ void SLevelViewport::OnActorLockToggleFromMenu(AActor* Actor)
 	}
 }
 
+void SLevelViewport::OnActorLockToggleFromMenu()
+{
+	OnActorUnlock();
+}
+
 bool SLevelViewport::IsActorLocked(const TWeakObjectPtr<AActor> Actor) const
 {
 	return LevelViewportClient->IsActorLocked(Actor);
diff --git a/Engine/Source/Editor/LevelEditor/Private/ViewportToolbar/LevelEditorViewportToolbarSections.cpp b/Engine/Source/Editor/LevelEditor/Private/ViewportToolbar/LevelEditorViewportToolbarSections.cpp
--- a/Engine/Source/Editor/LevelEditor/Private/ViewportToolbar/LevelEditorViewportToolbarSections.cpp
+++ b/Engine/Source/Editor/LevelEditor/Private/ViewportToolbar/LevelEditorViewportToolbarSections.cpp
@@ -6,6 +6,8 @@
 #include "Bookmarks/IBookmarkTypeTools.h"
 #include "Camera/CameraActor.h"
 #include "EditorViewportCommands.h"
+#include "Engine/SceneCapture.h"
+#include "EngineUtils.h"
 #include "FoliageType.h"
 #include "GameFramework/ActorPrimitiveColorHandler.h"
 #include "GameFramework/WorldSettings.h"
@@ -17,14 +19,15 @@
 #include "SCommonEditorViewportToolbarBase.h"
 #include "SLevelViewport.h"
 #include "SScalabilitySettings.h"
-#include "ShowFlagMenuCommands.h"
+#include "Selection.h"
+#include "SortHelper.h"
 #include "Stats/StatsData.h"
+#include "Styling/SlateIconFinder.h"
 #include "Templates/SharedPointer.h"
 #include "ToolMenu.h"
 #include "ToolMenus.h"
 #include "ViewportToolbar/LevelViewportContext.h"
 #include "ViewportToolbar/UnrealEdViewportToolbar.h"
-#include "Widgets/Input/SNumericEntryBox.h"
 #include "Widgets/Input/SSpinBox.h"
 #include "Widgets/Input/SVolumeControl.h"
 #include "Widgets/SBoxPanel.h"
@@ -282,6 +285,178 @@ void AddClearBookmarkMenu(UToolMenu* InMenu, const TWeakPtr<::SLevelViewport>& I
 	}
 }
 
+void GeneratePlacedCameraMenuEntries(
+	FToolMenuSection& InSection, TArray<AActor*> InLookThroughActors, const TSharedPtr<::SLevelViewport>& InLevelViewport
+)
+{
+	// Sort the cameras to make the ordering predictable for users.
+	InLookThroughActors.StableSort(
+		[](const AActor& Left, const AActor& Right)
+		{
+			// Do "natural sorting" via SceneOutliner::FNumericStringWrapper to make more sense to humans (also matches
+			// the Scene Outliner). This sorts "Camera2" before "Camera10" which a normal lexicographical sort wouldn't.
+			SceneOutliner::FNumericStringWrapper LeftWrapper(FString(Left.GetActorLabel()));
+			SceneOutliner::FNumericStringWrapper RightWrapper(FString(Right.GetActorLabel()));
+
+			return LeftWrapper < RightWrapper;
+		}
+	);
+
+	for (AActor* LookThroughActor : InLookThroughActors)
+	{
+		// Needed for the delegate hookup to work below
+		AActor* GenericActor = LookThroughActor;
+
+		FText ActorDisplayName = FText::FromString(LookThroughActor->GetActorLabel());
+		FUIAction LookThroughCameraAction(
+			FExecuteAction::CreateSP(InLevelViewport.ToSharedRef(), &::SLevelViewport::OnActorLockToggleFromMenu, GenericActor),
+			FCanExecuteAction(),
+			FIsActionChecked::CreateSP(
+				InLevelViewport.ToSharedRef(), &::SLevelViewport::IsActorLocked, MakeWeakObjectPtr(GenericActor)
+			)
+		);
+
+		FSlateIcon ActorIcon;
+
+		if (LookThroughActor->IsA<ACameraActor>() || LookThroughActor->IsA<ASceneCapture>())
+		{
+			ActorIcon = FSlateIcon(FAppStyle::GetAppStyleSetName(), "ClassIcon.CameraComponent");
+		}
+		else
+		{
+			ActorIcon = FSlateIconFinder::FindIconForClass(LookThroughActor->GetClass());
+		}
+
+		InSection.AddMenuEntry(
+			NAME_None,
+			ActorDisplayName,
+			FText::Format(LOCTEXT("LookThroughCameraActor_ToolTip", "Look through and pilot {0}"), ActorDisplayName),
+			ActorIcon,
+			LookThroughCameraAction,
+			EUserInterfaceActionType::RadioButton
+		);
+	}
+}
+
+FToolMenuEntry CreateEjectActorPilotEntry()
+{
+	return FToolMenuEntry::InitDynamicEntry(
+		"EjectActorPilotDynamicSection",
+		FNewToolMenuSectionDelegate::CreateLambda(
+			[](FToolMenuSection& InnerSection) -> void
+			{
+				ULevelViewportContext* const LevelViewportContext = InnerSection.FindContext<ULevelViewportContext>();
+				if (!LevelViewportContext)
+				{
+					return;
+				}
+
+				FToolUIAction EjectActorPilotAction;
+
+				EjectActorPilotAction.ExecuteAction = FToolMenuExecuteAction::CreateLambda(
+					[LevelViewportWeak = LevelViewportContext->LevelViewport](const FToolMenuContext& Context) -> void
+					{
+						if (TSharedPtr<::SLevelViewport> LevelViewport = LevelViewportWeak.Pin())
+						{
+							LevelViewport->OnActorLockToggleFromMenu();
+						}
+					}
+				);
+
+				EjectActorPilotAction.CanExecuteAction = FToolMenuCanExecuteAction::CreateLambda(
+					[LevelViewportWeak = LevelViewportContext->LevelViewport](const FToolMenuContext& Context)
+					{
+						if (TSharedPtr<::SLevelViewport> EditorViewport = LevelViewportWeak.Pin())
+						{
+							return EditorViewport->IsAnyActorLocked();
+						}
+						return false;
+					}
+				);
+
+				// We use this entry to gather its Name, Tooltip and Icon. See comment below as to why we cannot directly use this entry.
+				FToolMenuEntry SourceEjectPilotEntry =
+					FToolMenuEntry::InitMenuEntry(FLevelViewportCommands::Get().EjectActorPilot);
+
+				// We want to use SetShowInToolbarTopLevel to show the Eject entry in the Top Level only when piloting is active.
+				// Currently, this will not work with Commands, e.g. AddMenuEntry(FLevelViewportCommands::Get().EjectActorPilot).
+				// So, we create the entry using FToolMenuEntry::InitMenuEntry, and we create our own Action to handle it.
+				FToolMenuEntry EjectPilotActor = FToolMenuEntry::InitMenuEntry(
+					"EjectActorPilot",
+					LOCTEXT("EjectActorPilotLabel", "Stop Piloting Actor"),
+					LOCTEXT(
+						"EjectActorPilotTooltip", "Stop piloting an actor with the current viewport. Unlocks the viewport's position and orientation from the actor the viewport is currently piloting."
+					),
+					FSlateIcon(FAppStyle::GetAppStyleSetName(), "LevelViewport.EjectActorPilot"),
+					EjectActorPilotAction,
+					EUserInterfaceActionType::Button
+				);
+
+				const TAttribute<bool> bShownInTopLevel = TAttribute<bool>::CreateLambda(
+					[ViewportContextWeak = TWeakObjectPtr<ULevelViewportContext>(LevelViewportContext)]() -> bool
+					{
+						if (TStrongObjectPtr<ULevelViewportContext> LevelViewportContext = ViewportContextWeak.Pin())
+						{
+							if (const TSharedPtr<::SLevelViewport> LevelViewport = LevelViewportContext->LevelViewport.Pin())
+							{
+								return LevelViewport->GetLevelViewportClient().IsAnyActorLocked();
+							}
+						}
+
+						return true;
+					}
+				);
+
+				EjectPilotActor.SetShowInToolbarTopLevel(bShownInTopLevel);
+
+				InnerSection.AddEntry(EjectPilotActor);
+			}
+		)
+	);
+}
+
+FText GetCameraSubmenuLabelFromLevelViewport(const TWeakPtr<::SLevelViewport>& InLevelEditorViewportClientWeak)
+{
+	if (TSharedPtr<::SLevelViewport> LevelViewport = InLevelEditorViewportClientWeak.Pin())
+	{
+		const FLevelEditorViewportClient& LevelViewportClient = LevelViewport->GetLevelViewportClient();
+
+		if (!LevelViewportClient.IsAnyActorLocked())
+		{
+			return UnrealEd::GetCameraSubmenuLabelFromViewportType(LevelViewportClient.GetViewportType());
+		}
+		else if (TStrongObjectPtr<AActor> ActorLock = LevelViewportClient.GetActiveActorLock().Pin())
+		{
+			return FText::FromString(ActorLock->GetActorNameOrLabel());
+		}
+	}
+
+	return LOCTEXT("MissingActiveCameraLabel", "No Active Camera");
+}
+
+FSlateIcon GetCameraSubmenuIconFromLevelViewport(const TWeakPtr<::SLevelViewport>& InLevelEditorViewportClientWeak)
+{
+	if (TSharedPtr<::SLevelViewport> LevelViewport = InLevelEditorViewportClientWeak.Pin())
+	{
+		const FLevelEditorViewportClient& LevelViewportClient = LevelViewport->GetLevelViewportClient();
+		if (!LevelViewportClient.IsAnyActorLocked())
+		{
+			const FName IconName =
+				UnrealEd::GetCameraSubmenuIconFNameFromViewportType(LevelViewportClient.GetViewportType());
+			return FSlateIcon(FAppStyle::GetAppStyleSetName(), IconName);
+		}
+		else if (TStrongObjectPtr<AActor> LockedActor = LevelViewportClient.GetActorLock().LockedActor.Pin())
+		{
+			if (!LockedActor->IsA<ACameraActor>() && !LockedActor->IsA<ASceneCapture>())
+			{
+				return FSlateIconFinder::FindIconForClass(LockedActor->GetClass());
+			}
+		}
+	}
+
+	return FSlateIcon(FAppStyle::GetAppStyleSetName(), "ClassIcon.CameraComponent");
+}
+
 } // namespace UE::LevelEditor::Private
 
 namespace UE::LevelEditor
@@ -1440,7 +1615,6 @@ FToolMenuEntry CreateViewportToolbarPerformanceAndScalabilitySubmenu()
 
 					PerformanceAndScalabilitySection.AddEntry(CreateMaterialQualityLevelSubmenu());
 
-					// FEditorViewportClient& ViewportClient = Viewport.Pin()->GetLevelViewportClient();
 					PerformanceAndScalabilitySection.AddSubMenu(
 						"ScreenPercentageSubMenu",
 						LOCTEXT("ScreenPercentageSubMenu", "Screen Percentage"),
@@ -2124,6 +2298,72 @@ void CreateCameraSpeedMenu(UToolMenu* InMenu)
 	Section.AddEntry(CreateCameraSpeedScalarSlider(LevelViewport));
 }
 
+void AddCameraActorSelectSection(UToolMenu* InMenu)
+{
+	ULevelViewportContext* const LevelViewportContext = InMenu->FindContext<ULevelViewportContext>();
+	if (!LevelViewportContext)
+	{
+		return;
+	}
+
+	TSharedPtr<::SLevelViewport> LevelViewport = LevelViewportContext->LevelViewport.Pin();
+	if (!LevelViewport)
+	{
+		return;
+	}
+
+	TArray<AActor*> LookThroughActors;
+	for (TActorIterator<ACameraActor> It(LevelViewport->GetWorld()); It; ++It)
+	{
+		LookThroughActors.Add(Cast<AActor>(*It));
+	}
+
+	for (TActorIterator<ASceneCapture> It(LevelViewport->GetWorld()); It; ++It)
+	{
+		LookThroughActors.Add(Cast<AActor>(*It));
+	}
+
+	FText CameraActorsHeading = LOCTEXT("CameraActorsHeading", "Placed Cameras and Scene Capture Actors");
+
+	FToolMenuInsert InsertPosition("LevelViewportCameraType_Perspective", EToolMenuInsertType::After);
+
+	FToolMenuSection& Section = InMenu->AddSection("CameraActors");
+	Section.InsertPosition = InsertPosition;
+
+	// Don't add too many cameras to the top level menu or else it becomes too large
+	constexpr uint32 MaxCamerasInTopLevelMenu = 10;
+	if (LookThroughActors.Num() > MaxCamerasInTopLevelMenu)
+	{
+		Section.AddSubMenu(
+			"CameraActors",
+			CameraActorsHeading,
+			LOCTEXT("LookThroughPlacedCameras_ToolTip", "Look through and pilot placed cameras"),
+			FNewToolMenuDelegate::CreateLambda(
+				[LookThroughActors, LevelViewport](UToolMenu* InMenu)
+				{
+					FToolMenuSection& Section = InMenu->FindOrAddSection(NAME_None);
+					UE::LevelEditor::Private::GeneratePlacedCameraMenuEntries(Section, LookThroughActors, LevelViewport);
+				}
+			)
+		);
+	}
+	else
+	{
+		Section.AddSeparator(NAME_None);
+		UE::LevelEditor::Private::GeneratePlacedCameraMenuEntries(Section, LookThroughActors, LevelViewport);
+	}
+
+	TWeakObjectPtr<AActor> LockedActorWeak = LevelViewport->GetLevelViewportClient().GetActorLock().LockedActor;
+
+	if (TStrongObjectPtr<AActor> LockedActor = LockedActorWeak.Pin())
+	{
+		if (!LockedActor->IsA<ACameraActor>() && !LockedActor->IsA<ASceneCapture>())
+		{
+			UE::LevelEditor::Private::GeneratePlacedCameraMenuEntries(Section, { LockedActor.Get() }, LevelViewport);
+		}
+	}
+}
+
 void ExtendCameraSubmenu(FName InCameraOptionsSubmenuName)
 {
 	UToolMenu* const Submenu = UToolMenus::Get()->ExtendMenu(InCameraOptionsSubmenuName);
@@ -2133,10 +2373,59 @@ void ExtendCameraSubmenu(FName InCameraOptionsSubmenuName)
 		FNewToolMenuDelegate::CreateLambda(
 			[](UToolMenu* InDynamicMenu)
 			{
-				TWeakPtr<::SLevelViewport> LevelViewportWeak = InDynamicMenu->FindContext<ULevelViewportContext>()->LevelViewport;
+				ULevelViewportContext* LevelViewportContext = InDynamicMenu->FindContext<ULevelViewportContext>();
+				TWeakPtr<::SLevelViewport> LevelViewportWeak = LevelViewportContext->LevelViewport;
+
+				// Camera Selection elements
+				{
+					AddCameraActorSelectSection(InDynamicMenu);
+				}
+
+				// Pilot Section
+				{
+					FToolMenuSection& PilotSection = InDynamicMenu->FindOrAddSection("Pilot");
+
+					PilotSection.InsertPosition =
+						FToolMenuInsert("LevelViewportCameraType_Ortho", EToolMenuInsertType::After);
+					PilotSection.AddSeparator("PilotSectionSeparator");
 
-				// TODO:
-				// add Select Active Camera
+					bool bShowPilotSelectedActorEntry = false;
+
+					AActor* SelectedActor = nullptr;
+					if (TSharedPtr<::SLevelViewport> LevelViewport = LevelViewportWeak.Pin())
+					{
+						TArray<AActor*> SelectedActors;
+						GEditor->GetSelectedActors()->GetSelectedObjects<AActor>(SelectedActors);
+
+						if (!SelectedActors.IsEmpty() && !LevelViewport->IsSelectedActorLocked())
+						{
+							SelectedActor = SelectedActors[0];
+							const FLevelEditorViewportClient& ViewportClient = LevelViewport->GetLevelViewportClient();
+
+							bShowPilotSelectedActorEntry = SelectedActor && ViewportClient.IsPerspective()
+														&& !ViewportClient.IsLockedToCinematic();
+						}
+					}
+
+					if (bShowPilotSelectedActorEntry)
+					{
+						// Pilot Selected Actor Entry
+						PilotSection.AddMenuEntry(
+							FLevelViewportCommands::Get().PilotSelectedActor,
+							FText::Format(
+								LOCTEXT("PilotActor", "Pilot '{0}'"), FText::FromString(SelectedActor->GetActorLabel())
+							)
+						);
+					}
+
+					// Stop Piloting Entry
+					PilotSection.AddEntry(UE::LevelEditor::Private::CreateEjectActorPilotEntry());
+
+					// Exact Camera View Entry
+					FToolMenuEntry& ToggleCameraView =
+						PilotSection.AddMenuEntry(FLevelViewportCommands::Get().ToggleActorPilotCameraView);
+					ToggleCameraView.Label = LOCTEXT("ToggleCameraViewLabel", "Exact Camera View");
+				}
 
 				// Create Section
 				{
@@ -2212,23 +2501,6 @@ void ExtendCameraSubmenu(FName InCameraOptionsSubmenuName)
 						PositioningSection.AddEntry(FocusViewportToSelection);
 					}
 
-					// Pilot Submenu
-					{
-						PositioningSection.AddSubMenu(
-							"Pilot",
-							LOCTEXT("PilotSubMenu", "Pilot"),
-							LOCTEXT("PilotSubMenu_ToolTip", "Pilot related actions"),
-							FNewToolMenuDelegate::CreateLambda(
-								[](UToolMenu* InMenu)
-								{
-									// TODO add create Pilot menu
-								}
-							),
-							false,
-							FSlateIcon(FSlateIcon(FAppStyle::GetAppStyleSetName(), "LevelViewport.PilotSelectedActor"))
-						);
-					}
-
 					PositioningSection.AddSeparator("PositioningSeparator_2");
 
 					// Move Camera/Object
@@ -2282,6 +2554,49 @@ void ExtendCameraSubmenu(FName InCameraOptionsSubmenuName)
 	);
 }
 
+FToolMenuEntry CreateLevelViewportToolbarCameraSubmenu()
+{
+	return FToolMenuEntry::InitDynamicEntry(
+		"DynamicCameraOptions",
+		FNewToolMenuSectionDelegate::CreateLambda(
+			[](FToolMenuSection& InDynamicSection) -> void
+			{
+				if (ULevelViewportContext* const LevelViewportContext =
+						InDynamicSection.FindContext<ULevelViewportContext>())
+				{
+					const TAttribute<FText> Label = TAttribute<FText>::CreateLambda(
+						[LevelViewportWeak = LevelViewportContext->LevelViewport]()
+						{
+							return UE::LevelEditor::Private::GetCameraSubmenuLabelFromLevelViewport(LevelViewportWeak);
+						}
+					);
+
+					const TAttribute<FSlateIcon> Icon = TAttribute<FSlateIcon>::CreateLambda(
+						[LevelViewportWeak = LevelViewportContext->LevelViewport]()
+						{
+							return UE::LevelEditor::Private::GetCameraSubmenuIconFromLevelViewport(LevelViewportWeak);
+						}
+					);
+
+					InDynamicSection.AddSubMenu(
+						"CameraOptions",
+						Label,
+						LOCTEXT("CameraSubmenuTooltip", "Camera options"),
+						FNewToolMenuDelegate::CreateLambda(
+							[](UToolMenu* Submenu) -> void
+							{
+								UnrealEd::PopulateCameraMenu(Submenu);
+							}
+						),
+						false,
+						Icon
+					);
+				}
+			}
+		)
+	);
+}
+
 } // namespace UE::LevelEditor
 
 #undef LOCTEXT_NAMESPACE
diff --git a/Engine/Source/Editor/LevelEditor/Private/ViewportToolbar/LevelEditorViewportToolbarSections.h b/Engine/Source/Editor/LevelEditor/Private/ViewportToolbar/LevelEditorViewportToolbarSections.h
--- a/Engine/Source/Editor/LevelEditor/Private/ViewportToolbar/LevelEditorViewportToolbarSections.h
+++ b/Engine/Source/Editor/LevelEditor/Private/ViewportToolbar/LevelEditorViewportToolbarSections.h
@@ -17,11 +17,13 @@ namespace UE::LevelEditor
 void CreateCameraSpawnMenu(UToolMenu* InMenu);
 void CreateBookmarksMenu(UToolMenu* InMenu);
 void CreateCameraSpeedMenu(UToolMenu* InMenu);
+void AddCameraActorSelectSection(UToolMenu* InMenu);
 
 FToolMenuEntry CreateFOVMenu(TWeakPtr<SLevelViewport> InLevelViewportWeak);
 FToolMenuEntry CreateFarViewPlaneMenu(TWeakPtr<SLevelViewport> InLevelViewportWeak);
 FToolMenuEntry CreateCameraSpeedSlider(TWeakPtr<SLevelViewport> InLevelViewportWeak);
 FToolMenuEntry CreateCameraSpeedScalarSlider(TWeakPtr<SLevelViewport> InLevelViewportWeak);
+FToolMenuEntry CreateLevelViewportToolbarCameraSubmenu();
 
 bool ShowViewportRealtimeWarning(FLevelEditorViewportClient& ViewportClient);
 
diff --git a/Engine/Source/Editor/LevelEditor/Public/SLevelViewport.h b/Engine/Source/Editor/LevelEditor/Public/SLevelViewport.h
--- a/Engine/Source/Editor/LevelEditor/Public/SLevelViewport.h
+++ b/Engine/Source/Editor/LevelEditor/Public/SLevelViewport.h
@@ -246,6 +246,9 @@ public:
 	/** Called to lock/unlock the actor from the viewport's context menu */
 	void OnActorLockToggleFromMenu(AActor* Actor);
 
+	/** Called to unlock the actor from the viewport's context menu */
+	void OnActorLockToggleFromMenu();
+
 	/**
 	 * @return true if the actor is locked to the viewport
 	 */
diff --git a/Engine/Source/Editor/UnrealEd/Private/ViewportToolbar/UnrealEdViewportToolbar.cpp b/Engine/Source/Editor/UnrealEd/Private/ViewportToolbar/UnrealEdViewportToolbar.cpp
--- a/Engine/Source/Editor/UnrealEd/Private/ViewportToolbar/UnrealEdViewportToolbar.cpp
+++ b/Engine/Source/Editor/UnrealEd/Private/ViewportToolbar/UnrealEdViewportToolbar.cpp
@@ -2214,11 +2214,18 @@ void PopulateCameraMenu(UToolMenu* InMenu)
 		return;
 	}
 
-	FToolMenuSection& UnnamedSection = InMenu->FindOrAddSection("", LOCTEXT("UnnamedLabel", ""));
-
-	UnnamedSection.AddEntry(FToolMenuEntry::InitWidget(
-		"CameraMenuItems", UE::UnrealEd::CreateCameraMenuWidget(EditorViewport.ToSharedRef()), FText(), true
-	));
+	FToolMenuSection& PerspectiveCameraSection =
+		InMenu->FindOrAddSection("LevelViewportCameraType_Perspective", LOCTEXT("UnnamedLabel", ""));
+	PerspectiveCameraSection.AddMenuEntry(FEditorViewportCommands::Get().Perspective);
+
+	FToolMenuSection& OrthographicCameraSection =
+		InMenu->FindOrAddSection("LevelViewportCameraType_Ortho", LOCTEXT("CameraTypeHeader_Ortho", "Orthographic"));
+	OrthographicCameraSection.AddMenuEntry(FEditorViewportCommands::Get().Top);
+	OrthographicCameraSection.AddMenuEntry(FEditorViewportCommands::Get().Bottom);
+	OrthographicCameraSection.AddMenuEntry(FEditorViewportCommands::Get().Left);
+	OrthographicCameraSection.AddMenuEntry(FEditorViewportCommands::Get().Right);
+	OrthographicCameraSection.AddMenuEntry(FEditorViewportCommands::Get().Front);
+	OrthographicCameraSection.AddMenuEntry(FEditorViewportCommands::Get().Back);
 }
 
 void ExtendCameraSubmenu(FName InCameraOptionsSubmenuName)
@@ -2243,7 +2250,7 @@ void ExtendCameraSubmenu(FName InCameraOptionsSubmenuName)
 					return;
 				}
 
-				FToolMenuInsert InsertPosition("CameraOptions", EToolMenuInsertType::After);
+				FToolMenuInsert InsertPosition("LevelViewportCameraType_Ortho", EToolMenuInsertType::After);
 
 				FToolMenuSection& UnnamedSection =
 					InDynamicMenu->FindOrAddSection("", LOCTEXT("UnnamedLabel", ""), InsertPosition);
