Make viewport toolbar camera list prettier

* Rename the new viewport toolbar "Camera" submenu entry "Create Camera Here" to "Create Camera".
* Set UCLASS display name of VPCameraBase to "VP Camera Base" (affects old and new viewport toolbar).
* Set blueprint display name of VPCineCamera to "VP Cine Camera" (affects old and new viewport toolbar).

#virtualized

diff --git a/Engine/Plugins/Experimental/VirtualProductionUtilities/Source/VPUtilitiesEditor/Private/VPCameraBase.h b/Engine/Plugins/Experimental/VirtualProductionUtilities/Source/VPUtilitiesEditor/Private/VPCameraBase.h
--- a/Engine/Plugins/Experimental/VirtualProductionUtilities/Source/VPUtilitiesEditor/Private/VPCameraBase.h
+++ b/Engine/Plugins/Experimental/VirtualProductionUtilities/Source/VPUtilitiesEditor/Private/VPCameraBase.h
@@ -9,8 +9,7 @@
 class ACameraActor;
 class UCameraComponent;
 
-
-UCLASS()
+UCLASS(meta = (DisplayName = "VP Camera Base"))
 class AVPCameraBase : public ACineCameraActor
 {
 	GENERATED_BODY()
diff --git a/Engine/Source/Editor/LevelEditor/Private/ViewportToolbar/LevelEditorViewportToolbarSections.cpp b/Engine/Source/Editor/LevelEditor/Private/ViewportToolbar/LevelEditorViewportToolbarSections.cpp
--- a/Engine/Source/Editor/LevelEditor/Private/ViewportToolbar/LevelEditorViewportToolbarSections.cpp
+++ b/Engine/Source/Editor/LevelEditor/Private/ViewportToolbar/LevelEditorViewportToolbarSections.cpp
@@ -2428,7 +2428,7 @@ void ExtendCameraSubmenu(FName InCameraOptionsSubmenuName)
 
 					CreateSection.AddSubMenu(
 						"CreateCamera",
-						LOCTEXT("CameraSubMenu", "Create Camera Here"),
+						LOCTEXT("CameraSubMenu", "Create Camera"),
 						LOCTEXT("CameraSubMenu_ToolTip", "Select a camera type to create at current viewport's location"),
 						FNewToolMenuDelegate::CreateLambda(
 							[](UToolMenu* InMenu)
