Improve the Favorite Levels menu

* Fix issue where the level editor "Favorite Levels" menu would be hidden if "Recent Levels" was cleared.
* Do not show an empty "Favorite Levels" menu in the case where there aren't any favorites and the current level cannot be made a favorite.

diff --git a/Engine/Source/Editor/LevelEditor/Private/LevelEditorMenu.cpp b/Engine/Source/Editor/LevelEditor/Private/LevelEditorMenu.cpp
--- a/Engine/Source/Editor/LevelEditor/Private/LevelEditorMenu.cpp
+++ b/Engine/Source/Editor/LevelEditor/Private/LevelEditorMenu.cpp
@@ -139,7 +139,11 @@ void FLevelEditorMenu::RegisterLevelEditorMenus()
 				{
 					IMainFrameModule& MainFrameModule = FModuleManager::LoadModuleChecked<IMainFrameModule>("MainFrame");
 					const FMainMRUFavoritesList& RecentsAndFavorites = *MainFrameModule.GetMRUFavoritesList();
-					if (RecentsAndFavorites.GetNumItems() > 0)
+
+					// Only show the Favorite Levels menu if either 1) the current level could be favorited (it's saved)
+					// or 2) there are 1 or more favorite levels.
+					if (FLevelEditorActionCallbacks::ToggleFavorite_CanExecute()
+						|| RecentsAndFavorites.GetNumFavorites() > 0)
 					{
 						InSection.AddSubMenu(
 							"FavoriteLevelsSubMenu",
