Fix empty label on the level editor Actor main menu

The label is set on selection change, so at startup the label can be empty. With this change we call the label logic on startup with an empty selection set, resulting in the default label "Actor".

diff --git a/Engine/Source/Editor/LevelEditor/Private/SLevelEditor.cpp b/Engine/Source/Editor/LevelEditor/Private/SLevelEditor.cpp
--- a/Engine/Source/Editor/LevelEditor/Private/SLevelEditor.cpp
+++ b/Engine/Source/Editor/LevelEditor/Private/SLevelEditor.cpp
@@ -90,6 +90,8 @@ static const FName WorldBrowserCompositionTab("WorldBrowserComposition");
 SLevelEditor::SLevelEditor()
 	: World(nullptr)
 	, bNeedsRefresh(false)
+	// Ensure the Actor main menu has a default label before the first selection change.
+	, CachedViewportContextMenuTitle(FLevelEditorContextMenu::GetContextMenuTitle(ELevelEditorMenuContext::MainMenu, nullptr))
 {
 }
 
