Raise "Exact Camera View" to the top-level viewport toolbar when piloting

Also don't add separator between Perspective and the list of camera actors if that list is empty.

diff --git a/Engine/Source/Editor/LevelEditor/Private/ViewportToolbar/LevelEditorViewportToolbarSections.cpp b/Engine/Source/Editor/LevelEditor/Private/ViewportToolbar/LevelEditorViewportToolbarSections.cpp
--- a/Engine/Source/Editor/LevelEditor/Private/ViewportToolbar/LevelEditorViewportToolbarSections.cpp
+++ b/Engine/Source/Editor/LevelEditor/Private/ViewportToolbar/LevelEditorViewportToolbarSections.cpp
@@ -2342,7 +2342,7 @@ void AddCameraActorSelectSection(UToolMenu* InMenu)
 			)
 		);
 	}
-	else
+	else if (!LookThroughActors.IsEmpty())
 	{
 		Section.AddSeparator(NAME_None);
 		UE::LevelEditor::Private::GeneratePlacedCameraMenuEntries(Section, LookThroughActors, LevelViewport);
@@ -2417,9 +2417,21 @@ void ExtendCameraSubmenu(FName InCameraOptionsSubmenuName)
 					PilotSection.AddEntry(UE::LevelEditor::Private::CreateEjectActorPilotEntry());
 
 					// Exact Camera View Entry
-					FToolMenuEntry& ToggleCameraView =
-						PilotSection.AddMenuEntry(FLevelViewportCommands::Get().ToggleActorPilotCameraView);
-					ToggleCameraView.Label = LOCTEXT("ToggleCameraViewLabel", "Exact Camera View");
+					{
+						FToolMenuEntry& ToggleCameraView =
+							PilotSection.AddMenuEntry(FLevelViewportCommands::Get().ToggleActorPilotCameraView);
+						ToggleCameraView.Label = LOCTEXT("ToggleCameraViewLabel", "Exact Camera View");
+						ToggleCameraView.SetShowInToolbarTopLevel(TAttribute<bool>::CreateLambda(
+							[LevelViewportWeak]()
+							{
+								if (TSharedPtr<::SLevelViewport> EditorViewport = LevelViewportWeak.Pin())
+								{
+									return EditorViewport->IsAnyActorLocked();
+								}
+								return false;
+							}
+						));
+					}
 				}
 
 				// Create Section
