Shorten name of the new level editor viewport toolbar camera overflow menu to "Cameras"

diff --git a/Engine/Source/Editor/LevelEditor/Private/ViewportToolbar/LevelEditorViewportToolbarSections.cpp b/Engine/Source/Editor/LevelEditor/Private/ViewportToolbar/LevelEditorViewportToolbarSections.cpp
--- a/Engine/Source/Editor/LevelEditor/Private/ViewportToolbar/LevelEditorViewportToolbarSections.cpp
+++ b/Engine/Source/Editor/LevelEditor/Private/ViewportToolbar/LevelEditorViewportToolbarSections.cpp
@@ -2318,7 +2318,7 @@ void AddCameraActorSelectSection(UToolMenu* InMenu)
 		LookThroughActors.Add(Cast<AActor>(*It));
 	}
 
-	FText CameraActorsHeading = LOCTEXT("CameraActorsHeading", "Placed Cameras and Scene Capture Actors");
+	FText CameraActorsHeading = LOCTEXT("CameraActorsHeading", "Cameras");
 
 	FToolMenuInsert InsertPosition("LevelViewportCameraType_Perspective", EToolMenuInsertType::After);
 
