Rename "Performance & Scalability > Engine Scalabilty" to "Performance & Scalability > Viewport Scalability" in the new level editor viewport toolbar

diff --git a/Engine/Source/Editor/LevelEditor/Private/ViewportToolbar/LevelEditorViewportToolbarSections.cpp b/Engine/Source/Editor/LevelEditor/Private/ViewportToolbar/LevelEditorViewportToolbarSections.cpp
--- a/Engine/Source/Editor/LevelEditor/Private/ViewportToolbar/LevelEditorViewportToolbarSections.cpp
+++ b/Engine/Source/Editor/LevelEditor/Private/ViewportToolbar/LevelEditorViewportToolbarSections.cpp
@@ -1595,8 +1595,8 @@ FToolMenuEntry CreateViewportToolbarPerformanceAndScalabilitySubmenu()
 
 					PerformanceAndScalabilitySection.AddSubMenu(
 						"Scalability",
-						LOCTEXT("ScalabilitySubMenu", "Engine Scalability"),
-						LOCTEXT("ScalabilitySubMenu_ToolTip", "Open the engine scalability settings"),
+						LOCTEXT("ScalabilitySubMenu", "Viewport Scalability"),
+						LOCTEXT("ScalabilitySubMenu_ToolTip", "Open the viewport scalability settings"),
 						FNewToolMenuDelegate::CreateLambda(
 							[](UToolMenu* InMenu) -> void
 							{
