Outliner: Fix the uncontrolled filter not working when the file was in any non-default uncontrolled CL

diff --git a/Engine/Source/Editor/LevelEditor/Private/LevelEditorOutlinerSettings.cpp b/Engine/Source/Editor/LevelEditor/Private/LevelEditorOutlinerSettings.cpp
--- a/Engine/Source/Editor/LevelEditor/Private/LevelEditorOutlinerSettings.cpp
+++ b/Engine/Source/Editor/LevelEditor/Private/LevelEditorOutlinerSettings.cpp
@@ -309,7 +309,10 @@ bool FLevelEditorOutlinerSettings::DoesActorPassUncontrolledFilter(const ISceneO
 	
 	for (const TSharedRef<FUncontrolledChangelistState>& UncontrolledChangelistState : UncontrolledChangelistStates)
 	{
-		return UncontrolledChangelistState->GetFilenames().Contains(ExternalPackageFilename);
+		if (UncontrolledChangelistState->GetFilenames().Contains(ExternalPackageFilename))
+		{
+			return true;
+		}
 	}
 	
 	return false;
