Fixed Place Actor menu for Custom Actors (e.g. Runtime Virtual Texture Volume)

diff --git a/Engine/Source/Editor/LevelEditor/Private/LevelEditorCreateActorMenu.cpp b/Engine/Source/Editor/LevelEditor/Private/LevelEditorCreateActorMenu.cpp
--- a/Engine/Source/Editor/LevelEditor/Private/LevelEditorCreateActorMenu.cpp
+++ b/Engine/Source/Editor/LevelEditor/Private/LevelEditorCreateActorMenu.cpp
@@ -610,14 +610,19 @@ void LevelEditorCreateActorMenu::FillAddReplaceActorMenu(UToolMenu* Menu, EActor
 		{
 			AssetMenuOptions.Empty();
 			UActorFactory* Factory = GEditor->ActorFactories[FactoryIdx];
-			FAssetData AssetData = FAssetData( Factory->GetDefaultActorClass( FAssetData() ) );
-
-			const bool FactoryWorksWithoutAsset = Factory->CanCreateActorFrom( NoAssetData, UnusedErrorMessage );
 
-			if ( FactoryWorksWithoutAsset && Factory->bShowInEditorQuickMenu )
+			if (Factory->bShowInEditorQuickMenu)
 			{
-				AssetMenuOptions.Add( FActorFactoryAssetProxy::FMenuItem( Factory, NoAssetData ) );
-				BuildSingleAssetAddReplaceActorMenu(Section, AssetData, AssetMenuOptions, CreateMode);
+				FAssetData AssetData(Factory->GetDefaultActorClass(FAssetData()));
+				const bool bFactoryWorksWithoutAsset = Factory->CanCreateActorFrom(NoAssetData, UnusedErrorMessage);
+				const bool bFactoryWorksWithAsset = AssetData.IsValid() && Factory->CanCreateActorFrom(AssetData, UnusedErrorMessage);
+				const bool bFactoryWorks = bFactoryWorksWithAsset || bFactoryWorksWithoutAsset;
+
+				if (bFactoryWorks)
+				{
+					AssetMenuOptions.Add(FActorFactoryAssetProxy::FMenuItem(Factory, bFactoryWorksWithAsset ? AssetData : NoAssetData));
+					BuildSingleAssetAddReplaceActorMenu(Section, AssetData, AssetMenuOptions, CreateMode);
+				}
 			}
 		}
 	}
