Addressing feedback Preview Json
- Generated Preview Jsons will be saved in <Project>\Saved\PreviewJsonDevices\Android
- Preview menu will have a list of Preview Json options. The Jsons will be gathered from <Project>\Saved\PreviewJsonDevices\Android and <Project>\Content\Editor\PreviewJsons
- Preview menu will also allow the ability to import Jsons from other locations
- Disable Preview button moved at the top of the menu
- Cleaned up Preview menu, added separators for Shader Format

diff --git a/Engine/Source/Editor/LevelEditor/Private/LevelEditor.cpp b/Engine/Source/Editor/LevelEditor/Private/LevelEditor.cpp
--- a/Engine/Source/Editor/LevelEditor/Private/LevelEditor.cpp
+++ b/Engine/Source/Editor/LevelEditor/Private/LevelEditor.cpp
@@ -1797,8 +1797,6 @@ void FLevelEditorModule::BindGlobalLevelEditorCommands()
 	// We need one extra slot for the Disable Preview option
 	check(MenuItems.Num() == Commands.PreviewPlatformOverrides.Num());
 
-	TSet<FName> PreviewShaderPlatformNames;
-	TSet<FName> PlatformNames;
 	for (int32 Index=0; Index < MenuItems.Num(); Index++)
 	{
 		const FPreviewPlatformMenuItem& Item = MenuItems[Index];
@@ -1816,7 +1814,7 @@ void FLevelEditorModule::BindGlobalLevelEditorCommands()
 				}
 
 				const ERHIFeatureLevel::Type FeatureLevel = GetMaxSupportedFeatureLevel(ShaderPlatform);
-				return FPreviewPlatformInfo(FeatureLevel, ShaderPlatform, Item.PlatformName, Item.ShaderFormat, Item.DeviceProfileName, true, Item.PreviewShaderPlatformName);
+				return FPreviewPlatformInfo(FeatureLevel, ShaderPlatform, Item.PlatformName, Item.ShaderFormat, Item.DeviceProfileName, true, Item.PreviewShaderPlatformName, Item.OptionalFriendlyNameOverride);
 			};
 
 			FPreviewPlatformInfo PreviewFeatureLevelInfo = GetPreviewFeatureLevelInfo();
@@ -1831,24 +1829,27 @@ void FLevelEditorModule::BindGlobalLevelEditorCommands()
 			FString DeviceProfileSelectionModule;
 			if (PlatformEngineIni && PlatformEngineIni->GetString(TEXT("DeviceProfileManager"), TEXT("PreviewDeviceProfileSelectionModule"), DeviceProfileSelectionModule, GEngineIni))
 			{
-				if (!PlatformNames.Find(Item.PlatformName))
-				{
-					FUIAction Action;
-					Action.ExecuteAction = FExecuteAction::CreateStatic(&FLevelEditorActionCallbacks::GeneratePreviewJson_Clicked, Item.PlatformName.ToString());
-					Action.IsActionVisibleDelegate = FIsActionButtonVisible::CreateStatic(&FLevelEditorActionCallbacks::IsGeneratePreviewJsonVisible, Item.PlatformName);
-
-					ActionList.MapAction(Commands.GeneratePlatformJson[PlatformNames.Num()], Action);
-					PlatformNames.Add(Item.PlatformName);
-				}
+				const TArray<FLevelEditorCommands::PreviewPlatformCommand>* CommandListJson = Commands.PlatformToPreviewJsonPlatformOverrides.Find(Item.PlatformName);
 
-				if (!PreviewShaderPlatformNames.Find(Item.PreviewShaderPlatformName))
+				if (CommandListJson != nullptr)
 				{
-					FUIAction Action;
-					Action.ExecuteAction = FExecuteAction::CreateStatic(&FLevelEditorActionCallbacks::PreviewJson_Clicked, Item.PlatformName, Item.PreviewShaderPlatformName);
-					Action.IsActionVisibleDelegate = FIsActionButtonVisible::CreateStatic(&FLevelEditorActionCallbacks::IsPreviewJsonVisible, Item.PlatformName);
-
-					ActionList.MapAction(Commands.PreviewPlatformFromJson[PreviewShaderPlatformNames.Num()], Action);
-					PreviewShaderPlatformNames.Add(Item.PreviewShaderPlatformName);
+					for (const FLevelEditorCommands::PreviewPlatformCommand& PreviewJsonPlatform : *CommandListJson)
+					{
+						if (PreviewJsonPlatform.bIsGeneratingJsonCommand)
+						{
+							FUIAction Action;
+							Action.ExecuteAction = FExecuteAction::CreateStatic(&FLevelEditorActionCallbacks::GeneratePreviewJson_Clicked, Item.PlatformName.ToString());
+							Action.IsActionVisibleDelegate = FIsActionButtonVisible::CreateStatic(&FLevelEditorActionCallbacks::IsGeneratePreviewJsonVisible, Item.PlatformName);
+							ActionList.MapAction(PreviewJsonPlatform.CommandInfo, Action);
+						}
+						else
+						{
+							FUIAction ActionJsonFile;
+							ActionJsonFile.ExecuteAction = FExecuteAction::CreateStatic(&FLevelEditorActionCallbacks::PreviewJson_Clicked, Item.PlatformName, Item.PreviewShaderPlatformName, PreviewJsonPlatform.FilePath);
+							ActionJsonFile.IsActionVisibleDelegate = FIsActionButtonVisible::CreateStatic(&FLevelEditorActionCallbacks::IsPreviewJsonVisible, Item.PlatformName);
+							ActionList.MapAction(PreviewJsonPlatform.CommandInfo, ActionJsonFile);
+						}
+					}
 				}
 			}
 		}
diff --git a/Engine/Source/Editor/LevelEditor/Private/LevelEditorActions.cpp b/Engine/Source/Editor/LevelEditor/Private/LevelEditorActions.cpp
--- a/Engine/Source/Editor/LevelEditor/Private/LevelEditorActions.cpp
+++ b/Engine/Source/Editor/LevelEditor/Private/LevelEditorActions.cpp
@@ -547,27 +547,34 @@ void FLevelEditorActionCallbacks::ImportScene_Clicked()
 	FEditorFileUtils::Import();
 }
 
-void FLevelEditorActionCallbacks::PreviewJson_Clicked(FName PlatformName, FName PreviewShaderPlatformName)
-{
-	TArray<FString> OpenedFiles;
-	FString DefaultLocation(FEditorDirectories::Get().GetLastDirectory(ELastDirectory::PROJECT));
-
-	IDesktopPlatform* DesktopPlatform = FDesktopPlatformModule::Get();
-	bool bOpened = false;
-	if (DesktopPlatform)
-	{
-		bOpened = DesktopPlatform->OpenFileDialog(
-			FSlateApplication::Get().FindBestParentWindowHandleForDialogs(nullptr),
-			NSLOCTEXT("UnrealEd", "PreviewJson", "Preview Json").ToString(),
-			DefaultLocation,
-			TEXT(""),
-			TEXT("*.json"),
-			EFileDialogFlags::None,
-			OpenedFiles
-		);
+void FLevelEditorActionCallbacks::PreviewJson_Clicked(FName PlatformName, FName PreviewShaderPlatformName, FString JsonFile)
+{
+	if (JsonFile.IsEmpty())
+	{
+		TArray<FString> OpenedFiles;
+		FString DefaultLocation = IFileManager::Get().ConvertToAbsolutePathForExternalAppForWrite(*(FPaths::ProjectSavedDir() / TEXT("PreviewJsonDevices") / PlatformName.ToString()));
+
+		IDesktopPlatform* DesktopPlatform = FDesktopPlatformModule::Get();
+		bool bOpened = false;
+		if (DesktopPlatform)
+		{
+			bOpened = DesktopPlatform->OpenFileDialog(
+				FSlateApplication::Get().FindBestParentWindowHandleForDialogs(nullptr),
+				NSLOCTEXT("UnrealEd", "PreviewJson", "Preview Json").ToString(),
+				DefaultLocation,
+				TEXT(""),
+				TEXT("*.json"),
+				EFileDialogFlags::None,
+				OpenedFiles
+			);
+		}
+		if (bOpened && OpenedFiles.Num() > 0 && OpenedFiles[0].IsEmpty() == false)
+		{
+			JsonFile = OpenedFiles[0];
+		}
 	}
 
-	if (bOpened && OpenedFiles.Num() > 0 && OpenedFiles[0].IsEmpty() == false)
+	if (!JsonFile.IsEmpty())
 	{
 		FConfigCacheIni* PlatformEngineIni = FConfigCacheIni::ForPlatform(*PlatformName.ToString());
 		FString DeviceProfileSelectionModule;
@@ -578,7 +585,7 @@ void FLevelEditorActionCallbacks::PreviewJson_Clicked(FName PlatformName, FName
 				FString DevileProfileName;
 
 				TMap<FName, FString> DeviceParameters;
-				DPSelectorModule->GetDeviceParametersFromJson(OpenedFiles[0], DeviceParameters);
+				DPSelectorModule->GetDeviceParametersFromJson(JsonFile, DeviceParameters);
 				DPSelectorModule->SetSelectorProperties(DeviceParameters);
 
 				DevileProfileName = DPSelectorModule->GetDeviceProfileName();
@@ -632,16 +639,9 @@ bool FLevelEditorActionCallbacks::IsGeneratePreviewJsonVisible(FName PlatformNam
 
 void FLevelEditorActionCallbacks::GeneratePreviewJson_Clicked(FString PlatformName)
 {
-	FString DefaultLocation(FEditorDirectories::Get().GetLastDirectory(ELastDirectory::PROJECT));
-	FString OutputDirectory;
-	IDesktopPlatform* DesktopPlatform = FDesktopPlatformModule::Get();
 
-	if (DesktopPlatform &&
-		DesktopPlatform->OpenDirectoryDialog(FSlateApplication::Get().FindBestParentWindowHandleForDialogs(nullptr), 
-		LOCTEXT("LocationToOutputJsons", "Location to Output Json Files").ToString(), 
-		DefaultLocation,
-		OutputDirectory)
-		)
+	FString AbsoluteDebugInfoDirectory = IFileManager::Get().ConvertToAbsolutePathForExternalAppForWrite(*(FPaths::ProjectSavedDir() / TEXT("PreviewJsonDevices") / PlatformName));
+
 	{
 		FScopedSlowTask SlowTask(100.f, NSLOCTEXT("Engine", "GeneratePlatformJson", "Generate Platform Json"), true);
 		SlowTask.Visibility = ESlowTaskVisibility::ForceVisible;
@@ -655,7 +655,7 @@ void FLevelEditorActionCallbacks::GeneratePreviewJson_Clicked(FString PlatformNa
 		{
 			if (IDeviceProfileSelectorModule* DPSelectorModule = FModuleManager::LoadModulePtr<IDeviceProfileSelectorModule>(*DeviceProfileSelectionModule))
 			{
-				DPSelectorModule->ExportDeviceParametersToJson(OutputDirectory);
+				DPSelectorModule->ExportDeviceParametersToJson(AbsoluteDebugInfoDirectory);
 			}
 		}
 	}
@@ -4116,30 +4116,34 @@ void FLevelEditorCommands::RegisterCommands()
 
 	// Add preview platforms
 	TSet<FName> PreviewShaderPlatformNames;
-	TSet<FName> PlatformName;
 	for (const FPreviewPlatformMenuItem& Item : FDataDrivenPlatformInfoRegistry::GetAllPreviewPlatformMenuItems())
 	{
 		FTextBuilder FriendlyNameBuilder;
+		bool bIsDisablePreview = false;
 		if (!IsRunningCommandlet() && !GUsingNullRHI)
 		{
-			EShaderPlatform ShaderPlatform = FDataDrivenShaderPlatformInfo::GetShaderPlatformFromName(Item.PreviewShaderPlatformName);
-			if (ShaderPlatform == SP_NumPlatforms)
-			{
-				// if the shader platform isn't compiled in, we don't have a friendly name available, so use ugly name
-				FriendlyNameBuilder.AppendLine(FText::FromName(Item.PreviewShaderPlatformName));
-			}
-			else if (!Item.OptionalFriendlyNameOverride.IsEmpty())
+			if (FDataDrivenShaderPlatformInfo::GetShaderPlatformFromName(Item.ShaderPlatformToPreview) == GMaxRHIShaderPlatform)
 			{
-				FriendlyNameBuilder.AppendLine(Item.OptionalFriendlyNameOverride);
+				bIsDisablePreview = true;
+				FriendlyNameBuilder.AppendLine(NSLOCTEXT("PreviewPlatform", "PreviewMenuText_DisablePreview", "Disable Preview"));
 			}
 			else
 			{
-				FriendlyNameBuilder.AppendLine(FDataDrivenShaderPlatformInfo::GetFriendlyName(ShaderPlatform));
-				FPlatformMisc::LowLevelOutputDebugStringf(TEXT("MENU friendly name %s\n"), *FriendlyNameBuilder.ToText().ToString());
-			}
-			if (FDataDrivenShaderPlatformInfo::GetShaderPlatformFromName(Item.ShaderPlatformToPreview) == GMaxRHIShaderPlatform)
-			{
-				FriendlyNameBuilder.AppendLine(NSLOCTEXT("PreviewPlatform", "PreviewMenuText_DisablePreview", "(Disable Preview)"));
+				EShaderPlatform ShaderPlatform = FDataDrivenShaderPlatformInfo::GetShaderPlatformFromName(Item.PreviewShaderPlatformName);
+				if (ShaderPlatform == SP_NumPlatforms)
+				{
+					// if the shader platform isn't compiled in, we don't have a friendly name available, so use ugly name
+					FriendlyNameBuilder.AppendLine(FText::FromName(Item.PreviewShaderPlatformName));
+				}
+				else if (!Item.OptionalFriendlyNameOverride.IsEmpty())
+				{
+					FriendlyNameBuilder.AppendLine(Item.OptionalFriendlyNameOverride);
+				}
+				else
+				{
+					FriendlyNameBuilder.AppendLine(FDataDrivenShaderPlatformInfo::GetFriendlyName(ShaderPlatform));
+					FPlatformMisc::LowLevelOutputDebugStringf(TEXT("MENU friendly name %s\n"), *FriendlyNameBuilder.ToText().ToString());
+				}
 			}
 		}
 
@@ -4152,8 +4156,19 @@ void FLevelEditorCommands::RegisterCommands()
 			.UserInterfaceType(EUserInterfaceActionType::Check)
 			.DefaultChord(FInputChord())
 		);
-
-		PlatformToPreviewPlatformOverrides.FindOrAdd(Item.PlatformName).Add(PreviewPlatformOverrides.Last());
+		
+		FName SectionName = FDataDrivenShaderPlatformInfo::GetLanguage(FDataDrivenShaderPlatformInfo::GetShaderPlatformFromName(Item.ShaderPlatformToPreview));
+		if (bIsDisablePreview)
+		{
+			DisablePlatformPreview = PreviewPlatformOverrides.Last();
+		}
+		else
+		{
+			FLevelEditorCommands::PreviewPlatformCommand PreviewPlatform;
+			PreviewPlatform.CommandInfo = PreviewPlatformOverrides.Last();
+			PreviewPlatform.SectionName = SectionName;
+			PlatformToPreviewPlatformOverrides.FindOrAdd(Item.PlatformName).Add(PreviewPlatform);
+		}
 
 		FConfigCacheIni* PlatformEngineIni = FConfigCacheIni::ForPlatform(*Item.PlatformName.ToString());
 		FString DeviceProfileSelectionModule;
@@ -4163,35 +4178,68 @@ void FLevelEditorCommands::RegisterCommands()
 			EShaderPlatform ShaderPlatform = FDataDrivenShaderPlatformInfo::GetShaderPlatformFromName(Item.ShaderPlatformToPreview);
 			FText PlatformFriendlyName = FDataDrivenShaderPlatformInfo::GetFriendlyName(ShaderPlatform);
 
-			if (!PlatformName.Find(Item.PlatformName))
+			TArray<PreviewPlatformCommand>* PlatformToPreviewJsonPlatformOverridesValue = PlatformToPreviewJsonPlatformOverrides.Find(Item.PlatformName);
+			if (!PlatformToPreviewJsonPlatformOverridesValue)
 			{
-				GeneratePlatformJson.Add(
-					FUICommandInfoDecl(
-						this->AsShared(),
-						FName(*FString::Printf(TEXT("Generate Platform Json for %s"), *Item.PlatformName.ToString())),
-						NSLOCTEXT("GeneratePlatformJson", "Generate Platform Json", "Generate Platform Json"),
-						NSLOCTEXT("GeneratePlatformJsonDesc", "Generate Platform Json Desc", "Generate Platform Json Desc"))
+				FLevelEditorCommands::PreviewPlatformCommand GenerateJsonPlatform;
+				GenerateJsonPlatform.CommandInfo = FUICommandInfoDecl(
+					this->AsShared(),
+					FName(*FString::Printf(TEXT("Generate Platform Json for %s"), *Item.PlatformName.ToString())),
+					NSLOCTEXT("GeneratePlatformJson", "Generate Platform Json", "Generate Platform Json..."),
+					NSLOCTEXT("GeneratePlatformJsonDesc", "Generate Platform Json From Connected Devices", "Generate Platform Json From Connected Devices"))
 					.UserInterfaceType(EUserInterfaceActionType::Button)
-					.DefaultChord(FInputChord())
-				);
+					.DefaultChord(FInputChord());
+				GenerateJsonPlatform.bIsGeneratingJsonCommand = true;
 
-				PlatformToPreviewJsonPlatformOverrides.FindOrAdd(Item.PlatformName).Add(GeneratePlatformJson.Last());
-				PlatformName.Add(Item.PlatformName);
+				PlatformToPreviewJsonPlatformOverridesValue = &PlatformToPreviewJsonPlatformOverrides.Add(Item.PlatformName);
+				PlatformToPreviewJsonPlatformOverridesValue->Add(GenerateJsonPlatform);
 			}
 
+			check(PlatformToPreviewJsonPlatformOverridesValue != nullptr);
 			if (!PreviewShaderPlatformNames.Find(Item.PreviewShaderPlatformName))
 			{
-				PreviewPlatformFromJson.Add(
-					FUICommandInfoDecl(
-						this->AsShared(),
-						FName(*FString::Printf(TEXT("Preview %s with Json"), *Item.PreviewShaderPlatformName.ToString())),
-						FText::Format(NSLOCTEXT("PreviewJson", "Preview Json", "Preview {0} with Json"), PlatformFriendlyName),
-						NSLOCTEXT("PreviewJsonDesc", "Preview Json Desc", "Preview Json Desc"))
+				FLevelEditorCommands::PreviewPlatformCommand PreviewJsonPlatform;
+				PreviewJsonPlatform.CommandInfo = FUICommandInfoDecl(
+					this->AsShared(),
+					FName(*FString::Printf(TEXT("Preview via Json with %s"), *Item.PreviewShaderPlatformName.ToString())),
+					NSLOCTEXT("PreviewviaJson", "Preview via Json", "Preview via Json..."),
+					NSLOCTEXT("PreviewviaJsonDesc", "Preview via Json", "Preview via Json"))
 					.UserInterfaceType(EUserInterfaceActionType::Button)
-					.DefaultChord(FInputChord())
-				);
-
-				PlatformToPreviewJsonPlatformOverrides.FindOrAdd(Item.PlatformName).Add(PreviewPlatformFromJson.Last());
+					.DefaultChord(FInputChord());
+				PreviewJsonPlatform.bIsGeneratingJsonCommand = false;
+				PreviewJsonPlatform.SectionName = SectionName;
+				PlatformToPreviewJsonPlatformOverridesValue->Add(PreviewJsonPlatform);
+
+				TMap<FString, TArray<FString>> DirectoryToJsonFiles;
+				FString AbsoluteDebugInfoDirectory = IFileManager::Get().ConvertToAbsolutePathForExternalAppForWrite(*(FPaths::ProjectSavedDir() / TEXT("PreviewJsonDevices") / Item.PlatformName.ToString()));
+				IFileManager::Get().FindFiles(DirectoryToJsonFiles.Add(AbsoluteDebugInfoDirectory), *AbsoluteDebugInfoDirectory, TEXT(".json"));
+				FString ProjectEditorJsonDir = FPaths::ProjectContentDir() / TEXT("Editor") / TEXT("PreviewJsonDevices") / Item.PlatformName.ToString();
+				IFileManager::Get().FindFiles(DirectoryToJsonFiles.Add(ProjectEditorJsonDir), *ProjectEditorJsonDir, TEXT(".json"));
+
+				TSet<FString> UniqueJsons;
+				for (auto Iter = DirectoryToJsonFiles.CreateConstIterator(); Iter; ++Iter)
+				{
+					FString DirectoryName = Iter.Key();
+					for (const FString& JsonFile : Iter.Value())
+					{
+						if (!UniqueJsons.Find(JsonFile))
+						{
+							FLevelEditorCommands::PreviewPlatformCommand PreviewJsonFilePlatform;
+							PreviewJsonFilePlatform.CommandInfo = FUICommandInfoDecl(
+								this->AsShared(),
+								FName(*FString::Printf(TEXT("Preview %s with Json %s"), *JsonFile, *Item.PreviewShaderPlatformName.ToString())),
+								FText::Format(NSLOCTEXT("PreviewJson", "Preview Json", "Preview {0}"), FText::FromString(FPaths::GetBaseFilename(JsonFile))),
+								FText::Format(NSLOCTEXT("PreviewJsonDesc", "Preview using Platform Json", "Preview {0}"), FText::FromString(FPaths::GetBaseFilename(JsonFile))))
+								.UserInterfaceType(EUserInterfaceActionType::Check)
+								.DefaultChord(FInputChord());
+							PreviewJsonFilePlatform.bIsGeneratingJsonCommand = false;
+							PreviewJsonFilePlatform.FilePath = DirectoryName / JsonFile;
+							PreviewJsonFilePlatform.SectionName = SectionName;
+							PlatformToPreviewJsonPlatformOverridesValue->Add(PreviewJsonFilePlatform);
+							UniqueJsons.Add(JsonFile);
+						}
+					}
+				}	
 				PreviewShaderPlatformNames.Add(Item.PreviewShaderPlatformName);
 			}
 		}
diff --git a/Engine/Source/Editor/LevelEditor/Private/ViewportToolbar/LevelEditorViewportToolbarSections.cpp b/Engine/Source/Editor/LevelEditor/Private/ViewportToolbar/LevelEditorViewportToolbarSections.cpp
--- a/Engine/Source/Editor/LevelEditor/Private/ViewportToolbar/LevelEditorViewportToolbarSections.cpp
+++ b/Engine/Source/Editor/LevelEditor/Private/ViewportToolbar/LevelEditorViewportToolbarSections.cpp
@@ -1496,12 +1496,15 @@ FToolMenuEntry CreateFeatureLevelPreviewSubmenu()
 			{
 				FToolMenuSection& Section =
 					InMenu->AddSection("EditorPreviewMode", LOCTEXT("EditorPreviewModePlatforms", "Preview Platforms"));
+				
+				Section.AddMenuEntry(FLevelEditorCommands::Get().DisablePlatformPreview);
+				Section.AddSeparator("DisablePlatformPreviewSeparator");
 
 				for (auto Iter = FLevelEditorCommands::Get().PlatformToPreviewPlatformOverrides.CreateConstIterator(); Iter; ++Iter)
 				{
 					FName PlatformName = Iter.Key();
-					const TArray<TSharedPtr<FUICommandInfo>>& CommandList = Iter.Value();
-					const TArray<TSharedPtr<FUICommandInfo>>* CommandListJson = FLevelEditorCommands::Get().PlatformToPreviewJsonPlatformOverrides.Find(PlatformName);
+					const TArray<FLevelEditorCommands::PreviewPlatformCommand>& CommandList = Iter.Value();
+					const TArray<FLevelEditorCommands::PreviewPlatformCommand>* CommandListJson = FLevelEditorCommands::Get().PlatformToPreviewJsonPlatformOverrides.Find(PlatformName);
 
 					Section.AddSubMenu(
 						FName(PlatformName), 
@@ -1510,28 +1513,47 @@ FToolMenuEntry CreateFeatureLevelPreviewSubmenu()
 						FNewToolMenuDelegate::CreateLambda(
 							[CommandList, CommandListJson](UToolMenu* InSubMenu)
 							{
-								FToolMenuSection& Section = InSubMenu->AddSection(NAME_None);
-								for (const TSharedPtr<FUICommandInfo>& Command : CommandList)
+								for (const FLevelEditorCommands::PreviewPlatformCommand& Command : CommandList)
 								{
-									Section.AddMenuEntry(Command);
+									FToolMenuSection& Section = InSubMenu->FindOrAddSection(Command.SectionName, FText::Format(LOCTEXT("PreviewJson", "{0}"), FText::FromName(Command.SectionName)));
+									Section.AddMenuEntry(Command.CommandInfo);
 								}
 
 								if (CommandListJson != nullptr)
 								{
-									Section.AddSubMenu(
-										"PreviewJson",
-										LOCTEXT("PreviewJson", "Preview Json"),
-										LOCTEXT("PreviewJson_ToolTip", "Preview Json"),
-										FNewToolMenuDelegate::CreateLambda(
-											[CommandListJson](UToolMenu* InSubMenu)
-											{
-												FToolMenuSection& Section = InSubMenu->AddSection(NAME_None);
-												for (const TSharedPtr<FUICommandInfo>& Command : *CommandListJson)
+									FToolMenuSection& SectionJson = InSubMenu->FindOrAddSection("PreviewWithJson", LOCTEXT("PreviewJson", "Preview With Json"));
+									TMap<FName, TArray<TSharedPtr<FUICommandInfo>>> SectionNameToCommandList;
+									for (const FLevelEditorCommands::PreviewPlatformCommand& PreviewJsonPlatform : *CommandListJson)
+									{
+										if (PreviewJsonPlatform.bIsGeneratingJsonCommand)
+										{
+											SectionJson.AddMenuEntry(PreviewJsonPlatform.CommandInfo);
+										}
+										else
+										{
+											SectionNameToCommandList.FindOrAdd(PreviewJsonPlatform.SectionName).Add(PreviewJsonPlatform.CommandInfo);
+										}
+									}
+
+									for (auto Iter = SectionNameToCommandList.CreateConstIterator(); Iter; ++Iter)
+									{
+										FName SectionName = Iter.Key();
+										const TArray<TSharedPtr<FUICommandInfo>>& CommandListValue = Iter.Value();
+										SectionJson.AddSubMenu(
+											SectionName,
+											FText::Format(LOCTEXT("PreviewJson", "Preview {0}"), FText::FromName(SectionName)),
+											FText::Format(LOCTEXT("PreviewJson", "Preview {0}"), FText::FromName(SectionName)),
+											FNewToolMenuDelegate::CreateLambda(
+												[CommandListValue](UToolMenu* InSubMenu)
 												{
-													Section.AddMenuEntry(Command);
-												}
-											})
-									);
+													FToolMenuSection& Section = InSubMenu->AddSection(NAME_None);
+													for (const TSharedPtr<FUICommandInfo>& Command : CommandListValue)
+													{
+														Section.AddMenuEntry(Command);
+													}
+												})
+										);
+									}
 								}
 							})
 					);
diff --git a/Engine/Source/Editor/LevelEditor/Public/LevelEditorActions.h b/Engine/Source/Editor/LevelEditor/Public/LevelEditorActions.h
--- a/Engine/Source/Editor/LevelEditor/Public/LevelEditorActions.h
+++ b/Engine/Source/Editor/LevelEditor/Public/LevelEditorActions.h
@@ -610,10 +610,24 @@ public:
 	TSharedPtr< FUICommandInfo > ToggleFeatureLevelPreview;
 
 	TArray<TSharedPtr<FUICommandInfo>> PreviewPlatformOverrides;
-	TArray<TSharedPtr<FUICommandInfo>> PreviewPlatformFromJson;
-	TArray<TSharedPtr<FUICommandInfo>> GeneratePlatformJson;
-	TMap<FName, TArray<TSharedPtr<FUICommandInfo>>> PlatformToPreviewPlatformOverrides;
-	TMap<FName, TArray<TSharedPtr<FUICommandInfo>>> PlatformToPreviewJsonPlatformOverrides;
+
+	struct PreviewPlatformCommand
+	{
+		PreviewPlatformCommand()
+			: bIsGeneratingJsonCommand(false)
+			, SectionName(NAME_None)
+		{
+		}
+
+		bool bIsGeneratingJsonCommand;
+		FName SectionName;
+		TSharedPtr<FUICommandInfo> CommandInfo;
+		FString FilePath;
+	};
+	
+	TSharedPtr< FUICommandInfo > DisablePlatformPreview;
+	TMap<FName, TArray<PreviewPlatformCommand>> PlatformToPreviewPlatformOverrides;
+	TMap<FName, TArray<PreviewPlatformCommand>> PlatformToPreviewJsonPlatformOverrides;
 
 	///**
 	// * Mode Commands                   
@@ -737,7 +751,7 @@ public:
 	/**
 	* Called When Preview Json is selected in the Platforms Preview Sub Menu
 	*/
-	static void PreviewJson_Clicked(FName PlatformName, FName PreviewShaderPlatformName);
+	static void PreviewJson_Clicked(FName PlatformName, FName PreviewShaderPlatformName, FString JsonFile);
 
 	/**
 	* Is Preview Json visible in the Platforms Preview Sub Menu
diff --git a/Engine/Source/Runtime/RHI/Public/DataDrivenShaderPlatformInfo.h b/Engine/Source/Runtime/RHI/Public/DataDrivenShaderPlatformInfo.h
--- a/Engine/Source/Runtime/RHI/Public/DataDrivenShaderPlatformInfo.h
+++ b/Engine/Source/Runtime/RHI/Public/DataDrivenShaderPlatformInfo.h
@@ -198,6 +198,12 @@ public:
 		return Infos[Platform].Language == LANGUAGE_Nintendo;
 	}
 
+	static FORCEINLINE_DEBUGGABLE const FName GetLanguage(const FStaticShaderPlatform Platform)
+	{
+		check(IsValid(Platform));
+		return Infos[Platform].Language;
+	}
+
 	static FORCEINLINE_DEBUGGABLE const ERHIFeatureLevel::Type GetMaxFeatureLevel(const FStaticShaderPlatform Platform)
 	{
 		check(IsValid(Platform));
