Hook up the "Show > Viewport Stats" submenu to the ToggleStats command so the submenu itself gets a checkbox to toggle all viewport stats on and off

diff --git a/Engine/Source/Editor/LevelEditor/Private/ViewportToolbar/LevelEditorViewportToolbarSections.cpp b/Engine/Source/Editor/LevelEditor/Private/ViewportToolbar/LevelEditorViewportToolbarSections.cpp
--- a/Engine/Source/Editor/LevelEditor/Private/ViewportToolbar/LevelEditorViewportToolbarSections.cpp
+++ b/Engine/Source/Editor/LevelEditor/Private/ViewportToolbar/LevelEditorViewportToolbarSections.cpp
@@ -1315,79 +1315,115 @@ FToolMenuEntry CreateShowVolumesSubmenu()
 }
 
 #if STATS
-FToolMenuEntry CreateShowStatsSubmenu()
+FToolMenuEntry CreateShowStatsSubmenu(bool bInAddToggleStatsCheckbox, TAttribute<FText> InLabelOverride)
 {
-	return FToolMenuEntry::InitSubMenu(
-		"ShowStatsMenu",
-		LOCTEXT("ShowStatsMenu", "Stat"),
-		LOCTEXT("ShowStatsMenu_ToolTip", "Show Stat commands"),
-		FNewToolMenuDelegate::CreateLambda(
-			[](UToolMenu* InMenu) -> void
+	return FToolMenuEntry::InitDynamicEntry(
+		"DynamicShowStatsEntry",
+		FNewToolMenuSectionDelegate::CreateLambda(
+			[bInAddToggleStatsCheckbox, InLabelOverride](FToolMenuSection& InDynamicSection)
 			{
-				TArray<FLevelViewportCommands::FShowMenuCommand> HideStatsMenu;
-				HideStatsMenu.Add(FLevelViewportCommands::FShowMenuCommand(
-					FLevelViewportCommands::Get().HideAllStats, LOCTEXT("HideAllLabel", "Hide All")
-				));
-
-				UE::LevelEditor::Private::PopulateMenuWithCommands(InMenu, HideStatsMenu, 1);
-
-				FToolMenuSection& Section = InMenu->FindOrAddSection("Section");
-
-				// Separate out stats into two list, those with and without submenus
-				TArray<FLevelViewportCommands::FShowMenuCommand> SingleStatCommands;
-				TMap<FString, TArray<FLevelViewportCommands::FShowMenuCommand>> SubbedStatCommands;
-				for (auto StatCatIt = FLevelViewportCommands::Get().ShowStatCatCommands.CreateConstIterator(); StatCatIt;
-					 ++StatCatIt)
+				TWeakPtr<::SLevelViewport> WeakLevelViewport;
+				if (ULevelViewportContext* const LevelViewportContext =
+						InDynamicSection.FindContext<ULevelViewportContext>())
 				{
-					const TArray<FLevelViewportCommands::FShowMenuCommand>& ShowStatCommands = StatCatIt.Value();
-					const FString& CategoryName = StatCatIt.Key();
+					WeakLevelViewport = LevelViewportContext->LevelViewport;
+				}
 
-					// If no category is specified, or there's only one category, don't use submenus
-					FString NoCategory = FStatConstants::NAME_NoCategory.ToString();
-					NoCategory.RemoveFromStart(TEXT("STATCAT_"));
-					if (CategoryName == NoCategory || FLevelViewportCommands::Get().ShowStatCatCommands.Num() == 1)
-					{
-						for (int32 StatIndex = 0; StatIndex < ShowStatCommands.Num(); ++StatIndex)
-						{
-							const FLevelViewportCommands::FShowMenuCommand& StatCommand = ShowStatCommands[StatIndex];
-							SingleStatCommands.Add(StatCommand);
-						}
-					}
-					else
+				FToolUIActionChoice CommandAction;
+				if (bInAddToggleStatsCheckbox)
+				{
+					if (TSharedPtr<::SLevelViewport> Viewport = WeakLevelViewport.Pin())
 					{
-						SubbedStatCommands.Add(CategoryName, ShowStatCommands);
+						CommandAction =
+							FToolUIActionChoice(FEditorViewportCommands::Get().ToggleStats, *Viewport->GetCommandList());
 					}
 				}
 
-				// First add all the stats that don't have a sub menu
-				for (auto StatCatIt = SingleStatCommands.CreateConstIterator(); StatCatIt; ++StatCatIt)
-				{
-					const FLevelViewportCommands::FShowMenuCommand& StatCommand = *StatCatIt;
-					Section.AddMenuEntry(NAME_None, StatCommand.ShowMenuItem, StatCommand.LabelOverride);
-				}
+				const TAttribute<FText> Label = InLabelOverride.IsSet() ? InLabelOverride
+																		: LOCTEXT("ShowStatsMenu", "Stat");
 
-				// Now add all the stats that have sub menus
-				for (auto StatCatIt = SubbedStatCommands.CreateConstIterator(); StatCatIt; ++StatCatIt)
-				{
-					const TArray<FLevelViewportCommands::FShowMenuCommand>& StatCommands = StatCatIt.Value();
-					const FText CategoryName = FText::FromString(StatCatIt.Key());
+				InDynamicSection.AddSubMenu(
+					"ShowStatsMenu",
+					Label,
+					LOCTEXT("ShowStatsMenu_ToolTip", "Show Stat commands"),
+					FNewToolMenuDelegate::CreateLambda(
+						[](UToolMenu* InMenu) -> void
+						{
+							TArray<FLevelViewportCommands::FShowMenuCommand> HideStatsMenu;
+							HideStatsMenu.Add(FLevelViewportCommands::FShowMenuCommand(
+								FLevelViewportCommands::Get().HideAllStats, LOCTEXT("HideAllLabel", "Hide All")
+							));
 
-					FFormatNamedArguments Args;
-					Args.Add(TEXT("StatCat"), CategoryName);
-					const FText CategoryDescription =
-						FText::Format(NSLOCTEXT("UICommands", "StatShowCatName", "Show {StatCat} stats"), Args);
+							UE::LevelEditor::Private::PopulateMenuWithCommands(InMenu, HideStatsMenu, 1);
 
-					Section.AddSubMenu(
-						NAME_None,
-						CategoryName,
-						CategoryDescription,
-						FNewToolMenuDelegate::CreateStatic(&UE::LevelEditor::Private::PopulateMenuWithCommands, StatCommands, 0)
-					);
-				}
+							FToolMenuSection& Section = InMenu->FindOrAddSection("Section");
+
+							// Separate out stats into two list, those with and without submenus
+							TArray<FLevelViewportCommands::FShowMenuCommand> SingleStatCommands;
+							TMap<FString, TArray<FLevelViewportCommands::FShowMenuCommand>> SubbedStatCommands;
+							for (auto StatCatIt = FLevelViewportCommands::Get().ShowStatCatCommands.CreateConstIterator();
+								 StatCatIt;
+								 ++StatCatIt)
+							{
+								const TArray<FLevelViewportCommands::FShowMenuCommand>& ShowStatCommands =
+									StatCatIt.Value();
+								const FString& CategoryName = StatCatIt.Key();
+
+								// If no category is specified, or there's only one category, don't use submenus
+								FString NoCategory = FStatConstants::NAME_NoCategory.ToString();
+								NoCategory.RemoveFromStart(TEXT("STATCAT_"));
+								if (CategoryName == NoCategory
+									|| FLevelViewportCommands::Get().ShowStatCatCommands.Num() == 1)
+								{
+									for (int32 StatIndex = 0; StatIndex < ShowStatCommands.Num(); ++StatIndex)
+									{
+										const FLevelViewportCommands::FShowMenuCommand& StatCommand =
+											ShowStatCommands[StatIndex];
+										SingleStatCommands.Add(StatCommand);
+									}
+								}
+								else
+								{
+									SubbedStatCommands.Add(CategoryName, ShowStatCommands);
+								}
+							}
+
+							// First add all the stats that don't have a sub menu
+							for (auto StatCatIt = SingleStatCommands.CreateConstIterator(); StatCatIt; ++StatCatIt)
+							{
+								const FLevelViewportCommands::FShowMenuCommand& StatCommand = *StatCatIt;
+								Section.AddMenuEntry(NAME_None, StatCommand.ShowMenuItem, StatCommand.LabelOverride);
+							}
+
+							// Now add all the stats that have sub menus
+							for (auto StatCatIt = SubbedStatCommands.CreateConstIterator(); StatCatIt; ++StatCatIt)
+							{
+								const TArray<FLevelViewportCommands::FShowMenuCommand>& StatCommands = StatCatIt.Value();
+								const FText CategoryName = FText::FromString(StatCatIt.Key());
+
+								FFormatNamedArguments Args;
+								Args.Add(TEXT("StatCat"), CategoryName);
+								const FText CategoryDescription =
+									FText::Format(NSLOCTEXT("UICommands", "StatShowCatName", "Show {StatCat} stats"), Args);
+
+								Section.AddSubMenu(
+									NAME_None,
+									CategoryName,
+									CategoryDescription,
+									FNewToolMenuDelegate::CreateStatic(
+										&UE::LevelEditor::Private::PopulateMenuWithCommands, StatCommands, 0
+									)
+								);
+							}
+						}
+					),
+					CommandAction,
+					bInAddToggleStatsCheckbox ? EUserInterfaceActionType::ToggleButton : EUserInterfaceActionType::Button,
+					false,
+					FSlateIcon(FAppStyle::Get().GetStyleSetName(), "EditorViewport.SubMenu.Stats")
+				);
 			}
-		),
-		false,
-		FSlateIcon(FAppStyle::Get().GetStyleSetName(), "EditorViewport.SubMenu.Stats")
+		)
 	);
 }
 #endif
@@ -1410,11 +1446,9 @@ FToolMenuEntry CreateViewportToolbarShowSubmenu()
 
 #if STATS
 					// Override the label of the stats submenu for the new viewport toolbar.
-					{
-						FToolMenuEntry StatsSubmenu = UE::LevelEditor::CreateShowStatsSubmenu();
-						StatsSubmenu.Label = LOCTEXT("ViewportStatsLabel", "Viewport Stats");
-						UnnamedSection.AddEntry(StatsSubmenu);
-					}
+					UnnamedSection.AddEntry(
+						UE::LevelEditor::CreateShowStatsSubmenu(true, LOCTEXT("ViewportStatsLabel", "Viewport Stats"))
+					);
 #endif
 				}
 
diff --git a/Engine/Source/Editor/LevelEditor/Private/ViewportToolbar/LevelEditorViewportToolbarSections.h b/Engine/Source/Editor/LevelEditor/Private/ViewportToolbar/LevelEditorViewportToolbarSections.h
--- a/Engine/Source/Editor/LevelEditor/Private/ViewportToolbar/LevelEditorViewportToolbarSections.h
+++ b/Engine/Source/Editor/LevelEditor/Private/ViewportToolbar/LevelEditorViewportToolbarSections.h
@@ -34,7 +34,9 @@ FToolMenuEntry CreateShowHLODsSubmenu();
 FToolMenuEntry CreateShowLayersSubmenu();
 FToolMenuEntry CreateShowSpritesSubmenu();
 #if STATS
-FToolMenuEntry CreateShowStatsSubmenu();
+FToolMenuEntry CreateShowStatsSubmenu(
+	bool bInAddToggleStatsCheckbox = false, TAttribute<FText> InLabelOverride = TAttribute<FText>()
+);
 #endif
 FToolMenuEntry CreateShowVolumesSubmenu();
 FToolMenuEntry CreateViewportToolbarShowSubmenu();
