LevelEditorViewportToolbarSections : In AddCameraActorSelectSection guard against an invalid LevelViewport->GetWorld()

This fixes a crash when performing Window -> Load Layout -> Default Editor Layout

diff --git a/Engine/Source/Editor/LevelEditor/Private/ViewportToolbar/LevelEditorViewportToolbarSections.cpp b/Engine/Source/Editor/LevelEditor/Private/ViewportToolbar/LevelEditorViewportToolbarSections.cpp
--- a/Engine/Source/Editor/LevelEditor/Private/ViewportToolbar/LevelEditorViewportToolbarSections.cpp
+++ b/Engine/Source/Editor/LevelEditor/Private/ViewportToolbar/LevelEditorViewportToolbarSections.cpp
@@ -2367,14 +2367,18 @@ void AddCameraActorSelectSection(UToolMenu* InMenu)
 	}
 
 	TArray<AActor*> LookThroughActors;
-	for (TActorIterator<ACameraActor> It(LevelViewport->GetWorld()); It; ++It)
-	{
-		LookThroughActors.Add(Cast<AActor>(*It));
-	}
 
-	for (TActorIterator<ASceneCapture> It(LevelViewport->GetWorld()); It; ++It)
+	if (UWorld* World = LevelViewport->GetWorld())
 	{
-		LookThroughActors.Add(Cast<AActor>(*It));
+		for (TActorIterator<ACameraActor> It(World); It; ++It)
+		{
+			LookThroughActors.Add(Cast<AActor>(*It));
+		}
+
+		for (TActorIterator<ASceneCapture> It(World); It; ++It)
+		{
+			LookThroughActors.Add(Cast<AActor>(*It));
+		}
 	}
 
 	FText CameraActorsHeading = LOCTEXT("CameraActorsHeading", "Cameras");
