[LevelEditor] Assign name to external build type menu entries that matches the external build type name

#rnx

diff --git a/Engine/Source/Editor/LevelEditor/Private/LevelEditorMenu.cpp b/Engine/Source/Editor/LevelEditor/Private/LevelEditorMenu.cpp
--- a/Engine/Source/Editor/LevelEditor/Private/LevelEditorMenu.cpp
+++ b/Engine/Source/Editor/LevelEditor/Private/LevelEditorMenu.cpp
@@ -683,7 +683,7 @@ void FLevelEditorMenu::RegisterBuildMenu()
 						: BuildTypeLocalizedNames[Index];
 
 					check(BuildTypeSection != nullptr);
-					BuildTypeSection->AddMenuEntry(CommandInfo, Label, ToolTip).Name = NAME_None;
+					BuildTypeSection->AddMenuEntry(CommandInfo, Label, ToolTip).Name = BuildTypeNames[Index];
 				}
 			}
 		}));
