LevelEditorViewportToolbarSections : fix conflicting LOCTEXT keys for "PreviewJson"

diff --git a/Engine/Source/Editor/LevelEditor/Private/ViewportToolbar/LevelEditorViewportToolbarSections.cpp b/Engine/Source/Editor/LevelEditor/Private/ViewportToolbar/LevelEditorViewportToolbarSections.cpp
--- a/Engine/Source/Editor/LevelEditor/Private/ViewportToolbar/LevelEditorViewportToolbarSections.cpp
+++ b/Engine/Source/Editor/LevelEditor/Private/ViewportToolbar/LevelEditorViewportToolbarSections.cpp
@@ -1558,7 +1558,7 @@ FToolMenuEntry CreateFeatureLevelPreviewSubmenu()
 
 								if (CommandListJson != nullptr)
 								{
-									FToolMenuSection& SectionJson = InSubMenu->FindOrAddSection("PreviewWithJson", LOCTEXT("PreviewJson", "Preview With Json"));
+									FToolMenuSection& SectionJson = InSubMenu->FindOrAddSection("PreviewWithJson", LOCTEXT("PreviewWithJsonLabel", "Preview With Json"));
 									TMap<FName, TArray<TSharedPtr<FUICommandInfo>>> SectionNameToCommandList;
 									for (const FLevelEditorCommands::PreviewPlatformCommand& PreviewJsonPlatform : *CommandListJson)
 									{
@@ -1578,8 +1578,8 @@ FToolMenuEntry CreateFeatureLevelPreviewSubmenu()
 										const TArray<TSharedPtr<FUICommandInfo>>& CommandListValue = Iter.Value();
 										SectionJson.AddSubMenu(
 											SectionName,
-											FText::Format(LOCTEXT("PreviewJson", "Preview {0}"), FText::FromName(SectionName)),
-											FText::Format(LOCTEXT("PreviewJson", "Preview {0}"), FText::FromName(SectionName)),
+											FText::Format(LOCTEXT("PreviewJsonLabel", "Preview {0}"), FText::FromName(SectionName)),
+											FText::Format(LOCTEXT("PreviewJsonTooltip", "Preview {0}"), FText::FromName(SectionName)),
 											FNewToolMenuDelegate::CreateLambda(
 												[CommandListValue](UToolMenu* InSubMenu)
 												{
