Fix regression where ProjectBadgeBackgroundColor (read from ini) was not respected anymore

diff --git a/Engine/Source/Editor/LevelEditor/Private/LevelEditor.cpp b/Engine/Source/Editor/LevelEditor/Private/LevelEditor.cpp
--- a/Engine/Source/Editor/LevelEditor/Private/LevelEditor.cpp
+++ b/Engine/Source/Editor/LevelEditor/Private/LevelEditor.cpp
@@ -105,8 +105,25 @@ public:
 		FString ProjectNameWatermarkPrefix;
 		GConfig->GetString(TEXT("LevelEditor"), TEXT("ProjectNameWatermarkPrefix"), /*out*/ ProjectNameWatermarkPrefix, GEditorPerProjectIni);
 
-		FColor BadgeTextColor = FColor(128, 128, 128, 255);
-		GConfig->GetColor(TEXT("LevelEditor"), TEXT("ProjectBadgeTextColor"), /*out*/ BadgeTextColor, GEditorPerProjectIni);
+		FSlateColor BadgeBackgroundColor = FAppStyle::Get().GetSlateColor("Colors.Title");
+		{
+			FColor ConfigColor;
+			if (GConfig->GetColor(
+					TEXT("LevelEditor"), TEXT("ProjectBadgeBackgroundColor"), /*out*/ ConfigColor, GEditorPerProjectIni
+				))
+			{
+				BadgeBackgroundColor = FLinearColor(ConfigColor);
+			}
+		}
+
+		FSlateColor BadgeTextColor = FStyleColors::Foreground;
+		{
+			FColor ConfigColor;
+			if (GConfig->GetColor(TEXT("LevelEditor"), TEXT("ProjectBadgeTextColor"), /*out*/ ConfigColor, GEditorPerProjectIni))
+			{
+				BadgeTextColor = FLinearColor(ConfigColor);
+			}
+		}
 
 		const FString EngineVersionString = FEngineVersion::Current().ToString(FEngineVersion::Current().HasChangelist() ? EVersionComponent::Changelist : EVersionComponent::Patch);
 
@@ -142,16 +159,25 @@ public:
 			.Margin(FAppStyle::Get().GetMargin("SProjectBadge.BadgePadding"))
 			.ColorAndOpacity(BadgeTextColor);
 
+		// clang-format off
 		SBox::Construct(SBox::FArguments()
 			.HAlign(HAlign_Right)
 			.VAlign(VAlign_Top)
 			.Padding(FMargin(0.0f, 0.0f, 2.0f, 0.0f))
 			[
-				SNew(SExtensionPanel)
-				.ExtensionPanelID("LevelEditorProjectNamePlate")
-				.DefaultWidget(DefaultNamePlate)
-				.WindowZoneOverride(EWindowZone::TitleBar)
+				SNew(SBorder)
+				.BorderBackgroundColor(BadgeBackgroundColor)
+				.BorderImage(FAppStyle::GetBrush("SProjectBadge.BadgeShape"))
+				.Padding(FMargin(0))
+				.VAlign(VAlign_Top)
+				[
+					SNew(SExtensionPanel)
+					.ExtensionPanelID("LevelEditorProjectNamePlate")
+					.DefaultWidget(DefaultNamePlate)
+					.WindowZoneOverride(EWindowZone::TitleBar)
+				]
 			]);
+		// clang-format on
 	}
 
 	FVector2D GetSizeLastFrame() const
