Added early out for empty package names to DoesActorPassUnsavedFilter and DoesActorPassUncontrolledFilter to avoid warning spam in DoesPackageExist
Confirmed that we are no longer hitting this warning in Editor

diff --git a/Engine/Source/Editor/LevelEditor/Private/LevelEditorOutlinerSettings.cpp b/Engine/Source/Editor/LevelEditor/Private/LevelEditorOutlinerSettings.cpp
--- a/Engine/Source/Editor/LevelEditor/Private/LevelEditorOutlinerSettings.cpp
+++ b/Engine/Source/Editor/LevelEditor/Private/LevelEditorOutlinerSettings.cpp
@@ -300,13 +300,23 @@ void FLevelEditorOutlinerSettings::CreateDefaultFilters()
 
 bool FLevelEditorOutlinerSettings::DoesActorPassUnsavedFilter(const ISceneOutlinerTreeItem& InItem)
 {
+	if (InItem.GetPackageName().IsEmpty())
+	{
+		return false;
+	}
+
 	return UnsavedAssets.Contains(USourceControlHelpers::PackageFilename(InItem.GetPackageName()));
 }
 
 bool FLevelEditorOutlinerSettings::DoesActorPassUncontrolledFilter(const ISceneOutlinerTreeItem& InItem)
 {
+	if (InItem.GetPackageName().IsEmpty())
+	{
+		return false;
+	}
+
 	FString ExternalPackageFilename = USourceControlHelpers::PackageFilename(InItem.GetPackageName());
-	
+		
 	for (const TSharedRef<FUncontrolledChangelistState>& UncontrolledChangelistState : UncontrolledChangelistStates)
 	{
 		if (UncontrolledChangelistState->GetFilenames().Contains(ExternalPackageFilename))
