OFPA - When performing the "Copy Selected Actor(s) File Path" action, avoid appending a newline at the end of the string
* Allows paths copied like this to be pasted into Perforce directly, in order to expand the Depot or Workspace trees

diff --git a/Engine/Source/Editor/LevelEditor/Private/LevelEditorActions.cpp b/Engine/Source/Editor/LevelEditor/Private/LevelEditorActions.cpp
--- a/Engine/Source/Editor/LevelEditor/Private/LevelEditorActions.cpp
+++ b/Engine/Source/Editor/LevelEditor/Private/LevelEditorActions.cpp
@@ -3103,6 +3103,11 @@ void FLevelEditorActionCallbacks::CopyActorFilePathtoClipboard_Clicked()
 
 	for (AActor* Actor : SelectedActors)
 	{
+		if (Result.Len() != 0)
+		{
+			Result.Append(LINE_TERMINATOR);
+		}
+
 		const UPackage* Package = Actor->GetPackage();
 		const FString LocalFullPath(Package->GetLoadedPath().GetLocalFullPath());
 		if (SelectedActors.Num() > 1)
@@ -3117,7 +3122,6 @@ void FLevelEditorActionCallbacks::CopyActorFilePathtoClipboard_Clicked()
 			Result.Append(TEXT("): "));
 		}
 		Result.Append(FPaths::ConvertRelativePathToFull(LocalFullPath));
-		Result.Append(LINE_TERMINATOR);
 	}
 
 	if (Result.Len())
