Re-enable VK SM5 preview for D3D
Force SM6 for preview VK SM5, Mac SM5
Enable Lumen Support for Preview VK SM5, Mac SM5
Modifying FeatureLevelSwitch to respect the parent ShaderPlatform FeatureLevel and not the one set by preview in order to have consistent material logic with the Shader Platform we are trying to preview

diff --git a/Engine/Config/Android/DataDrivenPlatformInfo.ini b/Engine/Config/Android/DataDrivenPlatformInfo.ini
--- a/Engine/Config/Android/DataDrivenPlatformInfo.ini
+++ b/Engine/Config/Android/DataDrivenPlatformInfo.ini
@@ -51,6 +51,7 @@ EnabledCVar=ini:Engine:Android:/Script/AndroidRuntimeSettings.AndroidRuntimeSett
 PlatformName=Android
 ShaderPlatform=VULKAN_SM5_ANDROID
 ShaderFormat=SF_VULKAN_SM5_ANDROID
+PreviewFeatureLevel=SM6
 ActiveIconName=LevelEditor.PreviewMode.AndroidVulkanSM5.Enabled
 InactiveIconName=LevelEditor.PreviewMode.AndroidVulkanSM5.Disabled
 MenuTooltip=NSLOCTEXT("PreviewPlatform", "PreviewMenuTooltip_AndroidVulkanSM5", "Mobile preview using Android Vulkan SM5 quality settings.")
@@ -134,4 +135,4 @@ bSupportsWaveOperations=RuntimeDependent
 MinimumWaveSize=4
 MaximumWaveSize=128
 bSupportsWave64=true
-FriendlyName=LOCTEXT("FriendlyShaderPlatformName_AndroidVulkanSM5", "Android Vulkan SM5 Mobile")
\ No newline at end of file
+FriendlyName=LOCTEXT("FriendlyShaderPlatformName_AndroidVulkanSM5", "Android Vulkan SM5 Mobile")
diff --git a/Engine/Config/Mac/DataDrivenPlatformInfo.ini b/Engine/Config/Mac/DataDrivenPlatformInfo.ini
--- a/Engine/Config/Mac/DataDrivenPlatformInfo.ini
+++ b/Engine/Config/Mac/DataDrivenPlatformInfo.ini
@@ -143,6 +143,7 @@ FriendlyName=LOCTEXT("FriendlyShaderPlatformName_MacMobile", "Mac Metal Mobile")
 PlatformName=Mac
 ShaderFormat=SF_METAL_SM5
 ShaderPlatform=METAL_SM5
+PreviewFeatureLevel=SM6
 MenuTooltip=LOCTEXT("PreviewMenuTooltip_METAL_SM5", "Mac using SM5 profile")
 
 [PreviewPlatform METAL_SM6]
diff --git a/Engine/Source/Editor/LevelEditor/Private/LevelEditorActions.cpp b/Engine/Source/Editor/LevelEditor/Private/LevelEditorActions.cpp
--- a/Engine/Source/Editor/LevelEditor/Private/LevelEditorActions.cpp
+++ b/Engine/Source/Editor/LevelEditor/Private/LevelEditorActions.cpp
@@ -819,18 +819,6 @@ bool FLevelEditorActionCallbacks::CanExecutePreviewPlatform(FPreviewPlatformInfo
 
 	if (FDataDrivenShaderPlatformInfo::IsValid(PreviewShaderPlatform) && FDataDrivenShaderPlatformInfo::GetIsPreviewPlatform(PreviewShaderPlatform))
 	{
-		const EShaderPlatform RealShaderPlatform = FDataDrivenShaderPlatformInfo::GetPreviewShaderPlatformParent(PreviewShaderPlatform);
-
-		// TODO: Prevent previewing VULKAN_SM5 with a D3D renderer until all issues are resolved.
-		// We have to use the real platform here since these fields are overridden in preview.
-		if (GDynamicRHI
-			&& (GDynamicRHI->GetInterfaceType() == ERHIInterfaceType::D3D12 || GDynamicRHI->GetInterfaceType() == ERHIInterfaceType::D3D11)
-			&& IsVulkanPlatform(RealShaderPlatform)
-			&& IsFeatureLevelSupported(RealShaderPlatform, ERHIFeatureLevel::SM5))
-		{
-			return false;
-		}
-
 		// When the preview platform's DDSPI MaxSamplers is > 16 and the current RHI device has support
 		// for > 16 samplers we rely on the shader compiler being able to choose an appropriate profile for the
 		// preview feature level that supports > 16 samplers. On D3D12 SM5 the D3D shader compiler will use Dxc and
diff --git a/Engine/Source/Runtime/Engine/Private/Materials/MaterialExpressions.cpp b/Engine/Source/Runtime/Engine/Private/Materials/MaterialExpressions.cpp
--- a/Engine/Source/Runtime/Engine/Private/Materials/MaterialExpressions.cpp
+++ b/Engine/Source/Runtime/Engine/Private/Materials/MaterialExpressions.cpp
@@ -10807,7 +10807,17 @@ UMaterialExpressionFeatureLevelSwitch::UMaterialExpressionFeatureLevelSwitch(con
 #if WITH_EDITOR
 int32 UMaterialExpressionFeatureLevelSwitch::Compile(class FMaterialCompiler* Compiler, int32 OutputIndex)
 {
-	const ERHIFeatureLevel::Type FeatureLevelToCompile = Compiler->GetFeatureLevel();
+	ERHIFeatureLevel::Type FeatureLevelToCompile = Compiler->GetFeatureLevel();
+
+	// PreviewPlatform can have a different feature level in order to support previewing the platform
+	// But we still want to respect the material logic of the parent platform
+	const EShaderPlatform ShaderPlatform = Compiler->GetShaderPlatform();
+	if (FDataDrivenShaderPlatformInfo::GetIsPreviewPlatform(ShaderPlatform))
+	{
+		const EShaderPlatform ParentShaderPlatform = FDataDrivenShaderPlatformInfo::GetPreviewShaderPlatformParent(ShaderPlatform);
+		FeatureLevelToCompile = FDataDrivenShaderPlatformInfo::GetMaxFeatureLevel(ParentShaderPlatform);
+	}
+
 	check(FeatureLevelToCompile < UE_ARRAY_COUNT(Inputs));
 	FExpressionInput& FeatureInput = Inputs[FeatureLevelToCompile];
 
diff --git a/Engine/Source/Runtime/RHI/Private/DataDrivenShaderPlatformInfo.cpp b/Engine/Source/Runtime/RHI/Private/DataDrivenShaderPlatformInfo.cpp
--- a/Engine/Source/Runtime/RHI/Private/DataDrivenShaderPlatformInfo.cpp
+++ b/Engine/Source/Runtime/RHI/Private/DataDrivenShaderPlatformInfo.cpp
@@ -535,6 +535,7 @@ void FGenericDataDrivenShaderPlatformInfo::UpdatePreviewPlatforms()
 				PREVIEW_USE_RUNTIME_VALUE(Language);
 				PREVIEW_USE_RUNTIME_VALUE(bIsHlslcc);
 				PREVIEW_USE_RUNTIME_VALUE(bSupportsDxc);
+				PREVIEW_USE_RUNTIME_VALUE(bIsSPIRV);
 
 				// Editor is always PC, never console and always supports debug view shaders
 				PREVIEW_FORCE_SETTING(bIsPC, true);
diff --git a/Engine/Source/Runtime/RenderCore/Private/RenderUtils.cpp b/Engine/Source/Runtime/RenderCore/Private/RenderUtils.cpp
--- a/Engine/Source/Runtime/RenderCore/Private/RenderUtils.cpp
+++ b/Engine/Source/Runtime/RenderCore/Private/RenderUtils.cpp
@@ -2091,7 +2091,7 @@ bool DoesPlatformSupportLumenGI(EShaderPlatform Platform, bool bSkipProjectCheck
 	#endif // WITH_EDITOR
 
 	const bool bMetalSM5 = Platform == SP_METAL_SM5 || ParentPreviewPlatform == SP_METAL_SM5;
-	const bool bVulkanSM5 = Platform == SP_VULKAN_SM5 || Platform == SP_VULKAN_SM5_ANDROID;
+	const bool bVulkanSM5 = Platform == SP_VULKAN_SM5 || Platform == SP_VULKAN_SM5_ANDROID || ParentPreviewPlatform == SP_VULKAN_SM5 || ParentPreviewPlatform == SP_VULKAN_SM5_ANDROID;
 	return (bSkipProjectCheck || bLumenSupported)
 		&& FDataDrivenShaderPlatformInfo::GetSupportsLumenGI(Platform)
 		&& (IsFeatureLevelSupported(Platform, ERHIFeatureLevel::SM6) || bVulkanSM5 || bMetalSM5) // Android and Mac can't rely on SM6 yet, but want to run Lumen
