[BUGFIX] Incorrect tooltip on Interchange Import into Level

diff --git a/Engine/Source/Editor/LevelEditor/Private/LevelEditorActions.cpp b/Engine/Source/Editor/LevelEditor/Private/LevelEditorActions.cpp
--- a/Engine/Source/Editor/LevelEditor/Private/LevelEditorActions.cpp
+++ b/Engine/Source/Editor/LevelEditor/Private/LevelEditorActions.cpp
@@ -3836,7 +3836,7 @@ void FLevelEditorCommands::RegisterCommands()
 	}
 
 	UI_COMMAND( ClearRecentFiles, "Clear Recent Levels", "Clear the list of recently opened levels", EUserInterfaceActionType::Button, FInputChord() );
-	UI_COMMAND( ImportScene, "Import Into Level...", "Imports a scene from a FBX or OBJ format into the current level", EUserInterfaceActionType::Button, FInputChord());
+	UI_COMMAND( ImportScene, "Import Into Level...", "Imports a scene from a FBX or GLTF format into the current level", EUserInterfaceActionType::Button, FInputChord());
 	UI_COMMAND( ExportAll, "Export All...", "Exports the entire level to a file on disk (multiple formats are supported.)", EUserInterfaceActionType::Button, FInputChord() );
 	UI_COMMAND( ExportSelected, "Export Selected...", "Exports currently-selected objects to a file on disk (multiple formats are supported.)", EUserInterfaceActionType::Button, FInputChord() );
 
