[Viewport Toolbar] Viewport Stats: Add FPS, Frame Times, RHI, Memory, and Scene Rendering above the sub-sub-submenus

Adds the following items:

- FPS
- Frame Times
- Memory
- RHI
- Scene Rendering

diff --git a/Engine/Source/Editor/LevelEditor/Private/ViewportToolbar/LevelEditorViewportToolbarSections.cpp b/Engine/Source/Editor/LevelEditor/Private/ViewportToolbar/LevelEditorViewportToolbarSections.cpp
--- a/Engine/Source/Editor/LevelEditor/Private/ViewportToolbar/LevelEditorViewportToolbarSections.cpp
+++ b/Engine/Source/Editor/LevelEditor/Private/ViewportToolbar/LevelEditorViewportToolbarSections.cpp
@@ -1314,14 +1314,29 @@ FToolMenuEntry CreateShowStatsSubmenu(bool bInAddToggleStatsCheckbox, TAttribute
 					FNewToolMenuDelegate::CreateLambda(
 						[](UToolMenu* InMenu) -> void
 						{
-							TArray<FLevelViewportCommands::FShowMenuCommand> HideStatsMenu;
-							HideStatsMenu.Add(FLevelViewportCommands::FShowMenuCommand(
-								FLevelViewportCommands::Get().HideAllStats, LOCTEXT("HideAllLabel", "Hide All")
-							));
+							// Hide All
+							{
+								FToolMenuSection& UnnamedSection = InMenu->AddSection(NAME_None);
+								UnnamedSection.AddMenuEntry(
+									FLevelViewportCommands::Get().HideAllStats,
+									LOCTEXT("HideAllLabel", "Hide All"),
+									FText(),
+									FSlateIconFinder::FindIcon("Cross")
+								);
+							}
 
-							UE::LevelEditor::Private::PopulateMenuWithCommands(InMenu, HideStatsMenu, 1);
+							// Common Stats
+							// The list of Stat Commands we want to show right below the Hide All Stats
+							TArray<FName> CommonStatCommandNames = { TEXT("STAT_FPS"),
+																	 TEXT("STAT_UNIT"),
+																	 TEXT("STATGROUP_Memory"),
+																	 TEXT("STATGROUP_RHI"),
+																	 TEXT("STATGROUP_SceneRendering") };
 
-							FToolMenuSection& Section = InMenu->FindOrAddSection("Section");
+							FToolMenuSection& CommonStatsSection =
+								InMenu->AddSection("CommonStats", LOCTEXT("CommonStatsLabel", "Common Stats"));
+
+							FToolMenuSection& Section = InMenu->AddSection("Section");
 
 							// Separate out stats into two list, those with and without submenus
 							TArray<FLevelViewportCommands::FShowMenuCommand> SingleStatCommands;
@@ -1351,8 +1366,33 @@ FToolMenuEntry CreateShowStatsSubmenu(bool bInAddToggleStatsCheckbox, TAttribute
 								{
 									SubbedStatCommands.Add(CategoryName, ShowStatCommands);
 								}
+
+								// Search for commands to be added to the Common Stats Section
+								for (const FLevelViewportCommands::FShowMenuCommand& ShowMenuCommand : ShowStatCommands)
+								{
+									if (TSharedPtr<FUICommandInfo> Command = ShowMenuCommand.ShowMenuItem)
+									{
+										for (FName StatCommandName : CommonStatCommandNames)
+										{
+											if (StatCommandName == Command->GetCommandName())
+											{
+												CommonStatsSection.AddMenuEntry(Command);
+											}
+										}
+									}
+								}
 							}
 
+							// Sort Common Stats section entries alphabetically
+							CommonStatsSection.Blocks.Sort(
+								[](const FToolMenuEntry& A, const FToolMenuEntry& B)
+								{
+									return A.Label.Get().ToLower().ToString() < B.Label.Get().ToLower().ToString();
+								}
+							);
+
+							CommonStatsSection.AddSeparator("CommonStatsSeparator");
+
 							// First add all the stats that don't have a sub menu
 							for (auto StatCatIt = SingleStatCommands.CreateConstIterator(); StatCatIt; ++StatCatIt)
 							{
