[Viewport Toolbar] Remove the piloting-camera pill when using the new viewport toolbar

diff --git a/Engine/Source/Editor/LevelEditor/Private/SLevelViewport.cpp b/Engine/Source/Editor/LevelEditor/Private/SLevelViewport.cpp
--- a/Engine/Source/Editor/LevelEditor/Private/SLevelViewport.cpp
+++ b/Engine/Source/Editor/LevelEditor/Private/SLevelViewport.cpp
@@ -4772,7 +4772,7 @@ void SLevelViewport::OnSimulateSessionFinished()
 
 EVisibility SLevelViewport::GetLockedIconVisibility() const
 {
-	return IsAnyActorLocked() ? EVisibility::Visible : EVisibility::Collapsed;
+	return IsAnyActorLocked() && UE::UnrealEd::ShowOldViewportToolbars() ? EVisibility::Visible : EVisibility::Collapsed;
 }
 
 FText SLevelViewport::GetLockedIconToolTip() const
