Add a new ViewportSizing menu to the new level editor viewport toolbar

* Add a new ViewportSizing menu to the new level editor viewport toolbar. In the toolbar it has a vertical ellipsis icon and no label, in the overflow menu it has the label "..." and no icon.
* Move the contents of the Settings > Layouts submenu of the new level editor viewport toolbar to the new ViewportSizing menu.
* Move the Immersive Mode toggle from the Settings submenu to the new ViewportSizing submenu. Rename it to Immersive View in the process.
* Move the minimize/maximize viewport entry from the Settings submenu to the new ViewportSizing submenu.
* Rename the "Advanced Settings..." command to "Open Viewport Preferences..." to match the design. It is only used in the new and old viewport toolbars.

diff --git a/Engine/Source/Editor/EditorStyle/Private/StarshipStyle.cpp b/Engine/Source/Editor/EditorStyle/Private/StarshipStyle.cpp
--- a/Engine/Source/Editor/EditorStyle/Private/StarshipStyle.cpp
+++ b/Engine/Source/Editor/EditorStyle/Private/StarshipStyle.cpp
@@ -459,6 +459,8 @@ void FStarshipEditorStyle::FStyle::SetupGeneralStyles()
 		Set("Icons.Heart", new IMAGE_BRUSH_SVG("Starship/Common/Heart", Icon16x16));
 		Set("Icons.RecentAssets", new IMAGE_BRUSH_SVG("Starship/Common/RecentAssets", Icon16x16));
 		Set("Icons.Scalability", new IMAGE_BRUSH_SVG("Starship/Common/Scalability_16", Icon20x20));
+		Set("Icons.EllipsisVerticalNarrow",
+			new CORE_IMAGE_BRUSH_SVG("Starship/Common/ellipsis-vertical-narrow", FVector2D(6, 24)));
 
 		Set("Icons.Toolbar.Play", new IMAGE_BRUSH_SVG("Starship/Common/play", Icon20x20));
 		Set("Icons.Toolbar.Pause", new IMAGE_BRUSH_SVG("Starship/MainToolbar/pause", Icon20x20));
@@ -2510,6 +2512,55 @@ void FStarshipEditorStyle::FStyle::SetupViewportStyles()
 				Set("ViewportToolbar.Raised", NewViewportToolbarTopLevelRaisedStyle);
 			}
 		}
+
+		// Special style for the "Viewport Sizing" submenu, with rounded corners on the left.
+		{
+			const FButtonStyle ViewportSizingButtonStyle =
+				FButtonStyle(ButtonStyle)
+					.SetNormal(FSlateRoundedBoxBrush(FStyleColors::Dropdown, FVector4(4.0f, 0.0f, 0.0f, 4.0f)))
+					.SetHovered(FSlateRoundedBoxBrush(FStyleColors::Dropdown, FVector4(4.0f, 0.0f, 0.0f, 4.0f)))
+					.SetPressed(FSlateRoundedBoxBrush(FStyleColors::Dropdown, FVector4(4.0f, 0.0f, 0.0f, 4.0f)))
+					.SetDisabled(FSlateRoundedBoxBrush(FStyleColors::Dropdown, FVector4(4.0f, 0.0f, 0.0f, 4.0f)))
+					.SetNormalPadding(FMargin(0, 0, 2.f, 0))
+					.SetPressedPadding(FMargin(0, 0, 2.f, 0));
+
+			FToolBarStyle ViewportSizingStyle = ToolBarStyle;
+
+			FComboButtonStyle ViewportSizingComboButtonStyle =
+				FComboButtonStyle(FStarshipCoreStyle::GetCoreStyle().GetWidgetStyle<FComboButtonStyle>("ComboButton"))
+					.SetDownArrowPadding(FMargin(1.f, 0.f, 2.f, 0.f))
+					.SetDownArrowImage(CORE_IMAGE_BRUSH_SVG("Starship/Common/ellipsis-vertical-narrow", FVector2D(6, 20)));
+			ViewportSizingComboButtonStyle.ButtonStyle = ViewportSizingButtonStyle;
+
+			ViewportSizingStyle.SetButtonStyle(ViewportSizingButtonStyle);
+			ViewportSizingStyle.SetComboButtonStyle(ViewportSizingComboButtonStyle);
+			ViewportSizingStyle.SetSeparatorPadding(FMargin(0.f, 0.f, 0.f, 0.f));
+
+			Set("ViewportToolbarViewportSizingSubmenu", ViewportSizingStyle);
+
+			// A raised button keeps the gray background here, with rounded corners on the right.
+			{
+				FButtonStyle ButtonStyleRaised = ButtonStyle;
+				{
+					FLinearColor BackgroundColorRaised = FStyleColors::Dropdown.GetSpecifiedColor();
+					BackgroundColorRaised.A = .90f;
+					const FSlateRoundedBoxBrush BackgroundBrushRaised(BackgroundColorRaised, 5.f, BackgroundColorRaised, 1.0);
+
+					ButtonStyleRaised.SetNormal(BackgroundBrushRaised)
+						.SetNormalForeground(FStyleColors::Foreground)
+						.SetNormal(FSlateRoundedBoxBrush(FStyleColors::Dropdown, FVector4(0, 4.0f, 4.0f, 0)))
+						.SetHovered(FSlateRoundedBoxBrush(FStyleColors::Dropdown, FVector4(0, 4.0f, 4.0f, 0)))
+						.SetPressed(FSlateRoundedBoxBrush(FStyleColors::Dropdown, FVector4(0, 4.0f, 4.0f, 0)))
+						.SetDisabled(FSlateRoundedBoxBrush(FStyleColors::Dropdown, FVector4(0, 4.0f, 4.0f, 0)));
+				}
+				FToolBarStyle ViewportSizingRaisedStyle = ToolBarStyle;
+				{
+					ViewportSizingRaisedStyle.SetButtonStyle(ButtonStyleRaised).SetComboButtonPadding(FMargin(2.0f, 0));
+
+					Set("ViewportToolbarViewportSizingSubmenu.Raised", ViewportSizingRaisedStyle);
+				}
+			}
+		}
 	}
 
 	// Old viewport toolbar.
diff --git a/Engine/Source/Editor/LevelEditor/Private/LevelViewportActions.cpp b/Engine/Source/Editor/LevelEditor/Private/LevelViewportActions.cpp
--- a/Engine/Source/Editor/LevelEditor/Private/LevelViewportActions.cpp
+++ b/Engine/Source/Editor/LevelEditor/Private/LevelViewportActions.cpp
@@ -39,7 +39,13 @@ void FLevelViewportCommands::RegisterCommands()
 {
 	UI_COMMAND( ToggleMaximize, "Maximize Viewport", "Toggles the Maximize state of the current viewport", EUserInterfaceActionType::ToggleButton, FInputChord() );
 	UI_COMMAND( ToggleGameView, "Game View", "Toggles game view.  Game view shows the scene as it appears in game", EUserInterfaceActionType::ToggleButton, FInputChord( EKeys::G ) );
-	UI_COMMAND( ToggleImmersive, "Immersive Mode", "Switches this viewport between immersive mode and regular mode", EUserInterfaceActionType::ToggleButton, PLATFORM_MAC ? FInputChord( EModifierKey::Control, EKeys::F11 ) : FInputChord( EKeys::F11 ) );
+	UI_COMMAND(
+		ToggleImmersive,
+		"Immersive View",
+		"Switches this viewport between immersive view and regular view",
+		EUserInterfaceActionType::ToggleButton,
+		PLATFORM_MAC ? FInputChord(EModifierKey::Control, EKeys::F11) : FInputChord(EKeys::F11)
+	);
 	UI_COMMAND(ToggleSidebarAllTabs, "Sidebar All Tabs", "Moves all tabs except the level editor to a sidebar or restores any previous state before if all tabs are already sidebared", EUserInterfaceActionType::ToggleButton, FInputChord(EKeys::F10));
 
 	UI_COMMAND( HighResScreenshot, "High Resolution Screenshot...", "Opens the control panel for high resolution screenshots", EUserInterfaceActionType::Button, FInputChord() );
@@ -82,7 +88,13 @@ void FLevelViewportCommands::RegisterCommands()
 	UI_COMMAND( ToggleCinematicPreview, "Allow Cinematic Control", "If enabled, allows cinematic (Sequencer) previews to play in this viewport", EUserInterfaceActionType::ToggleButton, FInputChord() );
 
 	UI_COMMAND( FindInLevelScriptBlueprint, "Find In Level Script", "Finds references of a selected actor in the level script blueprint", EUserInterfaceActionType::Button, FInputChord(EModifierKey::Control, EKeys::K) );
-	UI_COMMAND( AdvancedSettings, "Advanced Settings...", "Opens the advanced viewport settings", EUserInterfaceActionType::Button, FInputChord());
+	UI_COMMAND(
+		AdvancedSettings,
+		"Open Viewport Preferences...",
+		"Opens the advanced viewport settings",
+		EUserInterfaceActionType::Button,
+		FInputChord()
+	);
 
 	// Generate a command for each volume class
 	{
diff --git a/Engine/Source/Editor/LevelEditor/Private/SLevelViewport.cpp b/Engine/Source/Editor/LevelEditor/Private/SLevelViewport.cpp
--- a/Engine/Source/Editor/LevelEditor/Private/SLevelViewport.cpp
+++ b/Engine/Source/Editor/LevelEditor/Private/SLevelViewport.cpp
@@ -2084,6 +2084,13 @@ TSharedPtr<SWidget> SLevelViewport::BuildViewportToolbar()
 					SettingsSubmenu.InsertPosition.Position = EToolMenuInsertType::Last;
 					RightSection.AddEntry(SettingsSubmenu);
 				}
+
+				// Add the "..." submenu.
+				{
+					FToolMenuEntry ViewportSizingSubmenu = UE::LevelEditor::CreateLevelEditorViewportToolbarViewportSizingSubmenu();
+					ViewportSizingSubmenu.InsertPosition.Position = EToolMenuInsertType::Last;
+					RightSection.AddEntry(ViewportSizingSubmenu);
+				}
 			}
 		}
 	}
diff --git a/Engine/Source/Editor/LevelEditor/Private/ViewportToolbar/LevelEditorViewportToolbarSections.cpp b/Engine/Source/Editor/LevelEditor/Private/ViewportToolbar/LevelEditorViewportToolbarSections.cpp
--- a/Engine/Source/Editor/LevelEditor/Private/ViewportToolbar/LevelEditorViewportToolbarSections.cpp
+++ b/Engine/Source/Editor/LevelEditor/Private/ViewportToolbar/LevelEditorViewportToolbarSections.cpp
@@ -1876,37 +1876,6 @@ FToolMenuEntry CreateLevelEditorViewportToolbarSettingsSubmenu()
 		FNewToolMenuDelegate::CreateLambda(
 			[](UToolMenu* Submenu) -> void
 			{
-				{
-					FToolMenuSection& ViewportControlsSection = Submenu->FindOrAddSection(
-						"ViewportControls", LOCTEXT("ViewportControlsSectionLabel", "Viewport Controls")
-					);
-
-					ViewportControlsSection.AddSubMenu(
-						"ViewportLayouts",
-						LOCTEXT("ViewportLayoutsLabel", "Layouts"),
-						LOCTEXT("ViewportLayoutsTooltip", "Configure the layouts of the viewport windows"),
-						FNewToolMenuDelegate::CreateLambda(
-							[](UToolMenu* InMenu)
-							{
-								ULevelViewportContext* const LevelViewportContext =
-									InMenu->FindContext<ULevelViewportContext>();
-								if (!LevelViewportContext)
-								{
-									return;
-								}
-
-								if (const TSharedPtr<::SLevelViewport> LevelViewport =
-										LevelViewportContext->LevelViewport.Pin())
-								{
-									GenerateViewportLayoutsMenu(InMenu);
-								}
-							}
-						),
-						false,
-						FSlateIcon(FAppStyle::Get().GetStyleSetName(), "Icons.Layout")
-					);
-				}
-
 				{
 					FToolMenuSection& SettingsSection =
 						Submenu->FindOrAddSection("Settings", LOCTEXT("SettingsSectionLabel", "Settings"));
@@ -1915,11 +1884,49 @@ FToolMenuEntry CreateLevelEditorViewportToolbarSettingsSubmenu()
 						"Volume", BuildVolumeControlCustomWidget(), LOCTEXT("VolumeControlLabel", "Volume")
 					));
 
-					SettingsSection.AddSeparator("ViewportSizeSeparator");
+					SettingsSection.AddSeparator("AdvancedSeparator");
 
-					SettingsSection.AddMenuEntry(FLevelViewportCommands::Get().ToggleImmersive);
+					{
+						const FLevelViewportCommands& LevelViewportActions = FLevelViewportCommands::Get();
+						SettingsSection.AddMenuEntry(LevelViewportActions.AdvancedSettings);
+					}
+				}
+			}
+		)
+	);
+	Entry.Icon = FSlateIcon(FAppStyle::GetAppStyleSetName(), "LevelEditor.GameSettings");
+	Entry.ToolBarData.LabelOverride = FText();
+	return Entry;
+}
 
-					SettingsSection.AddDynamicEntry(
+FToolMenuEntry CreateLevelEditorViewportToolbarViewportSizingSubmenu()
+{
+	FToolMenuEntry Entry = FToolMenuEntry::InitSubMenu(
+		"ViewportSizing",
+		LOCTEXT("ViewportSizingLabel", "..."),
+		LOCTEXT("ViewportSizingTooltip", "Viewport-sizing settings"),
+		FNewToolMenuDelegate::CreateLambda(
+			[](UToolMenu* Submenu) -> void
+			{
+				ULevelViewportContext* const LevelViewportContext = Submenu->FindContext<ULevelViewportContext>();
+				if (!LevelViewportContext)
+				{
+					return;
+				}
+
+				if (const TSharedPtr<::SLevelViewport> LevelViewport = LevelViewportContext->LevelViewport.Pin())
+				{
+					GenerateViewportLayoutsMenu(Submenu);
+				}
+
+				{
+					FToolMenuSection& MaximizeSection = Submenu->FindOrAddSection("MaximizeSection");
+
+					MaximizeSection.AddSeparator("MaximizeSeparator");
+
+					MaximizeSection.AddMenuEntry(FLevelViewportCommands::Get().ToggleImmersive);
+
+					MaximizeSection.AddDynamicEntry(
 						"MaximizeRestoreDynamicEntry",
 						FNewToolMenuSectionDelegate::CreateLambda(
 							[](FToolMenuSection& InnerSection) -> void
@@ -2012,36 +2019,21 @@ FToolMenuEntry CreateLevelEditorViewportToolbarSettingsSubmenu()
 									}
 								);
 
-								// Only top-level flag the button if it is actually clickable.
-								const TAttribute<bool> TopLevel = TAttribute<bool>::CreateLambda(
-									[WeakLevelViewport = LevelViewportContext->LevelViewport]() -> bool
-									{
-										if (const TSharedPtr<::SLevelViewport> LevelViewport = WeakLevelViewport.Pin())
-										{
-											return LevelViewport->CanMaximize();
-										}
-
-										return false;
-									}
-								);
-								InnerSection.AddMenuEntry("MaximizeRestore", Label, Tooltip, Icon, Action)
-									.SetShowInToolbarTopLevel(TopLevel);
+								FToolMenuEntry& MaximizeRestoreEntry =
+									InnerSection.AddMenuEntry("MaximizeRestore", Label, Tooltip, Icon, Action);
+								MaximizeRestoreEntry.SetShowInToolbarTopLevel(true);
+								MaximizeRestoreEntry.StyleNameOverride = FName("ViewportToolbarViewportSizingSubmenu");
 							}
 						)
 					);
-
-					SettingsSection.AddSeparator("AdvancedSeparator");
-
-					{
-						const FLevelViewportCommands& LevelViewportActions = FLevelViewportCommands::Get();
-						SettingsSection.AddMenuEntry(LevelViewportActions.AdvancedSettings);
-					}
 				}
 			}
 		)
 	);
-	Entry.Icon = FSlateIcon(FAppStyle::GetAppStyleSetName(), "LevelEditor.GameSettings");
+
 	Entry.ToolBarData.LabelOverride = FText();
+	Entry.StyleNameOverride = FName("ViewportToolbarViewportSizingSubmenu");
+
 	return Entry;
 }
 
diff --git a/Engine/Source/Editor/LevelEditor/Private/ViewportToolbar/LevelEditorViewportToolbarSections.h b/Engine/Source/Editor/LevelEditor/Private/ViewportToolbar/LevelEditorViewportToolbarSections.h
--- a/Engine/Source/Editor/LevelEditor/Private/ViewportToolbar/LevelEditorViewportToolbarSections.h
+++ b/Engine/Source/Editor/LevelEditor/Private/ViewportToolbar/LevelEditorViewportToolbarSections.h
@@ -48,5 +48,7 @@ void GenerateViewportLayoutsMenu(UToolMenu* InMenu);
 TSharedRef<SWidget> BuildVolumeControlCustomWidget();
 FToolMenuEntry CreateLevelEditorViewportToolbarSettingsSubmenu();
 
+FToolMenuEntry CreateLevelEditorViewportToolbarViewportSizingSubmenu();
+
 void ExtendCameraSubmenu(FName InCameraOptionsSubmenuName);
 } // namespace UE::LevelEditor
