Fix surface-level issues with viewport toolbars

* Add icon and remove label from the new Niagara viewport toolbar's "Show" and "Settings" submenus.
* Sort "Show" before "Performance & Scalability" in the new Niagara viewport toolbar.
* Add icon and remove label from the new Blueprint viewport toolbar's "Show" and "Performance & Scalability" submenus.
* Add proper tooltip to the new Blueprint viewport toolbar's "Performance & Scalability" submenu.
* Add icon and remove label from the new Material viewport toolbar's "Show" submenu.
* Add icon and remove label from the "Performance & Scalability" submenu used by the new viewport toolbars of the Water Waves, Niagara, Smart Object, Material, and Static Mesh editors.

diff --git a/Engine/Plugins/FX/Niagara/Source/NiagaraEditor/Private/Widgets/NiagaraSystemEditorViewportToolbarSections.cpp b/Engine/Plugins/FX/Niagara/Source/NiagaraEditor/Private/Widgets/NiagaraSystemEditorViewportToolbarSections.cpp
--- a/Engine/Plugins/FX/Niagara/Source/NiagaraEditor/Private/Widgets/NiagaraSystemEditorViewportToolbarSections.cpp
+++ b/Engine/Plugins/FX/Niagara/Source/NiagaraEditor/Private/Widgets/NiagaraSystemEditorViewportToolbarSections.cpp
@@ -15,7 +15,7 @@
 
 FToolMenuEntry UE::NiagaraSystemEditor::CreateShowSubmenu()
 {
-	return FToolMenuEntry::InitSubMenu(
+	FToolMenuEntry Entry = FToolMenuEntry::InitSubMenu(
 		"Show",
 		LOCTEXT("ShowSubmenuLabel", "Show"),
 		LOCTEXT("ShowSubmenuTooltip", "Show options"),
@@ -26,6 +26,10 @@ FToolMenuEntry UE::NiagaraSystemEditor::CreateShowSubmenu()
 			}
 		)
 	);
+	Entry.ToolBarData.LabelOverride = FText();
+	Entry.Icon = FSlateIcon(FAppStyle::GetAppStyleSetName(), "Level.VisibleHighlightIcon16x");
+
+	return Entry;
 }
 
 void UE::NiagaraSystemEditor::FillShowSubmenu(UToolMenu* InMenu, bool bInShowViewportStatsToggle)
@@ -74,7 +78,7 @@ void UE::NiagaraSystemEditor::FillShowSubmenu(UToolMenu* InMenu, bool bInShowVie
 
 FToolMenuEntry UE::NiagaraSystemEditor::CreateSettingsSubmenu()
 {
-	return FToolMenuEntry::InitSubMenu(
+	FToolMenuEntry Entry = FToolMenuEntry::InitSubMenu(
 		"Settings",
 		LOCTEXT("SettingsSubmenuLabel", "Settings"),
 		LOCTEXT("SettingsSubmenuTooltip", "Settings options"),
@@ -85,6 +89,10 @@ FToolMenuEntry UE::NiagaraSystemEditor::CreateSettingsSubmenu()
 			}
 		)
 	);
+	Entry.Icon = FSlateIcon(FAppStyle::GetAppStyleSetName(), "LevelEditor.GameSettings");
+	Entry.ToolBarData.LabelOverride = FText();
+	
+	return Entry;
 }
 
 TSharedRef<SWidget> UE::NiagaraSystemEditor::CreateMotionMenuWidget(const TSharedRef<SNiagaraSystemViewport>& InNiagaraSystemEditorViewport
diff --git a/Engine/Plugins/FX/Niagara/Source/NiagaraEditor/Private/Widgets/SNiagaraSystemViewport.cpp b/Engine/Plugins/FX/Niagara/Source/NiagaraEditor/Private/Widgets/SNiagaraSystemViewport.cpp
--- a/Engine/Plugins/FX/Niagara/Source/NiagaraEditor/Private/Widgets/SNiagaraSystemViewport.cpp
+++ b/Engine/Plugins/FX/Niagara/Source/NiagaraEditor/Private/Widgets/SNiagaraSystemViewport.cpp
@@ -1235,7 +1235,7 @@ TSharedPtr<SWidget> SNiagaraSystemViewport::BuildViewportToolbar()
 				UE::NiagaraSystemEditor::Private::ExtendCameraSubmenu(SubmenuName);
 
 				FToolMenuEntry CameraSubmenu = UE::UnrealEd::CreateViewportToolbarCameraSubmenu();
-				CameraSubmenu.InsertPosition.Position = EToolMenuInsertType::First;
+				CameraSubmenu.InsertPosition.Position = EToolMenuInsertType::Last;
 				RightSection.AddEntry(CameraSubmenu);
 
 				// TODO: add navigation options --> Orbit Mode toggle
@@ -1257,28 +1257,31 @@ TSharedPtr<SWidget> SNiagaraSystemViewport::BuildViewportToolbar()
 				}
 
 				FToolMenuEntry ViewModesSubmenu = UE::UnrealEd::CreateViewportToolbarViewModesSubmenu();
-				ViewModesSubmenu.InsertPosition.Position = EToolMenuInsertType::First;
+				ViewModesSubmenu.InsertPosition.Position = EToolMenuInsertType::Last;
 				RightSection.AddEntry(ViewModesSubmenu);
-
-				FToolMenuEntry PerformanceAndScalabilitySubmenu = UE::UnrealEd::CreatePerformanceAndScalabilitySubmenu();
-				PerformanceAndScalabilitySubmenu.InsertPosition.Position = EToolMenuInsertType::First;
-				RightSection.AddEntry(PerformanceAndScalabilitySubmenu);
 			}
 
 			// Add the Show submenu.
 			{
 				FToolMenuEntry ShowSubmenu = UE::NiagaraSystemEditor::CreateShowSubmenu();
-				ShowSubmenu.InsertPosition.Position = EToolMenuInsertType::First;
+				ShowSubmenu.InsertPosition.Position = EToolMenuInsertType::Last;
 				RightSection.AddEntry(ShowSubmenu);
 			}
 
 			// Add the Performance and Scalability submenu.
 			{
 				FToolMenuEntry PerformanceAndScalabilitySubmenu = UE::UnrealEd::CreatePerformanceAndScalabilitySubmenu();
-				PerformanceAndScalabilitySubmenu.InsertPosition.Position = EToolMenuInsertType::First;
+				PerformanceAndScalabilitySubmenu.InsertPosition.Position = EToolMenuInsertType::Last;
 				RightSection.AddEntry(PerformanceAndScalabilitySubmenu);
 			}
 
+			// Add the Settings submenu.
+			{
+				FToolMenuEntry SettingsSubmenu = UE::NiagaraSystemEditor::CreateSettingsSubmenu();
+				SettingsSubmenu.InsertPosition.Position = EToolMenuInsertType::Last;
+				RightSection.AddEntry(SettingsSubmenu);
+			}
+
 			// Add the "Preview Profile" sub menu.
 			{
 				PreviewProfileController = MakeShared<FPreviewProfileController>();
@@ -1287,13 +1290,6 @@ TSharedPtr<SWidget> SNiagaraSystemViewport::BuildViewportToolbar()
 				PreviewProfileSubmenu.InsertPosition.Position = EToolMenuInsertType::Last;
 				RightSection.AddEntry(PreviewProfileSubmenu);
 			}
-
-			// Add the Settings submenu.
-			{
-				FToolMenuEntry SettingsSubmenu = UE::NiagaraSystemEditor::CreateSettingsSubmenu();
-				SettingsSubmenu.InsertPosition.Position = EToolMenuInsertType::First;
-				RightSection.AddEntry(SettingsSubmenu);
-			}
 		}
 	}
 
diff --git a/Engine/Source/Editor/Kismet/Private/SSCSEditorViewport.cpp b/Engine/Source/Editor/Kismet/Private/SSCSEditorViewport.cpp
--- a/Engine/Source/Editor/Kismet/Private/SSCSEditorViewport.cpp
+++ b/Engine/Source/Editor/Kismet/Private/SSCSEditorViewport.cpp
@@ -523,7 +523,8 @@ TSharedPtr<SWidget> SSCSEditorViewport::BuildViewportToolbar()
 						}
 					)
 				);
-
+				ShowSubmenu.ToolBarData.LabelOverride = FText();
+				ShowSubmenu.Icon = FSlateIcon(FAppStyle::GetAppStyleSetName(), "Level.VisibleHighlightIcon16x");
 				ShowSubmenu.InsertPosition.Position = EToolMenuInsertType::Last;
 				RightSection.AddEntry(ShowSubmenu);
 			}
@@ -533,7 +534,7 @@ TSharedPtr<SWidget> SSCSEditorViewport::BuildViewportToolbar()
 				FToolMenuEntry PerfSubmenu = FToolMenuEntry::InitSubMenu(
 					"PerformanceAndScalability",
 					LOCTEXT("PerformanceAndScalabilityLabel", "Performance and Scalability"),
-					LOCTEXT("PerformanceAndScalabilityTooltip", "Performance and Scalability tooltip"),
+					LOCTEXT("PerformanceAndScalabilityTooltip", "Performance and Scalability tools tied to this viewport"),
 					FNewToolMenuDelegate::CreateLambda(
 						[](UToolMenu* Submenu) -> void
 						{
@@ -542,7 +543,8 @@ TSharedPtr<SWidget> SSCSEditorViewport::BuildViewportToolbar()
 						}
 					)
 				);
-
+				PerfSubmenu.Icon = FSlateIcon(FAppStyle::GetAppStyleSetName(), "Icons.Scalability");
+				PerfSubmenu.ToolBarData.LabelOverride = FText();
 				PerfSubmenu.InsertPosition.Position = EToolMenuInsertType::Last;
 				RightSection.AddEntry(PerfSubmenu);
 			}
diff --git a/Engine/Source/Editor/LevelEditor/Private/SLevelViewport.cpp b/Engine/Source/Editor/LevelEditor/Private/SLevelViewport.cpp
--- a/Engine/Source/Editor/LevelEditor/Private/SLevelViewport.cpp
+++ b/Engine/Source/Editor/LevelEditor/Private/SLevelViewport.cpp
@@ -1981,6 +1981,37 @@ TSharedPtr<SWidget> SLevelViewport::BuildViewportToolbar()
 				FToolMenuSection& RightSection = ViewportToolbarMenu->FindOrAddSection("Right");
 				RightSection.Alignment = EToolMenuSectionAlign::Last;
 
+				// Add the "Camera" submenu.
+				{
+
+					const FName GrandParentSubmenuName = "UnrealEd.ViewportToolbar.Camera";
+					const FName ParentSubmenuName = "LevelEditor.ViewportToolbar.Camera";
+					const FName SubmenuName = "LevelEditor.ViewportToolbar.CameraOptions";
+
+					// Create our grandparent menu.
+					if (!UToolMenus::Get()->IsMenuRegistered(GrandParentSubmenuName))
+					{
+						UToolMenus::Get()->RegisterMenu(GrandParentSubmenuName);
+					}
+
+					// Create our parent menu.
+					if (!UToolMenus::Get()->IsMenuRegistered(ParentSubmenuName))
+					{
+						UToolMenus::Get()->RegisterMenu(ParentSubmenuName, GrandParentSubmenuName);
+					}
+
+					// Create our menu.
+					UToolMenus::Get()->RegisterMenu(SubmenuName, ParentSubmenuName);
+
+					// Extending using both Level Editor specific and UnrealEd generic entries
+					UE::LevelEditor::ExtendCameraSubmenu(SubmenuName);
+					UE::UnrealEd::ExtendCameraSubmenu(SubmenuName);
+
+					FToolMenuEntry CameraSubmenu = UE::LevelEditor::CreateLevelViewportToolbarCameraSubmenu();
+					CameraSubmenu.InsertPosition.Position = EToolMenuInsertType::Last;
+					RightSection.AddEntry(CameraSubmenu);
+				}
+
 				// Add the "View Modes" sub menu.
 				{
 					// Stay backward-compatible with the old viewport toolbar.
@@ -2015,37 +2046,6 @@ TSharedPtr<SWidget> SLevelViewport::BuildViewportToolbar()
 					RightSection.AddEntry(ViewModesSubmenu);
 				}
 
-				// Add the "Camera" submenu.
-				{
-
-					const FName GrandParentSubmenuName = "UnrealEd.ViewportToolbar.Camera";
-					const FName ParentSubmenuName = "LevelEditor.ViewportToolbar.Camera";
-					const FName SubmenuName = "LevelEditor.ViewportToolbar.CameraOptions";
-
-					// Create our grandparent menu.
-					if (!UToolMenus::Get()->IsMenuRegistered(GrandParentSubmenuName))
-					{
-						UToolMenus::Get()->RegisterMenu(GrandParentSubmenuName);
-					}
-
-					// Create our parent menu.
-					if (!UToolMenus::Get()->IsMenuRegistered(ParentSubmenuName))
-					{
-						UToolMenus::Get()->RegisterMenu(ParentSubmenuName, GrandParentSubmenuName);
-					}
-
-					// Create our menu.
-					UToolMenus::Get()->RegisterMenu(SubmenuName, ParentSubmenuName);
-
-					// Extending using both Level Editor specific and UnrealEd generic entries
-					UE::LevelEditor::ExtendCameraSubmenu(SubmenuName);
-					UE::UnrealEd::ExtendCameraSubmenu(SubmenuName);
-
-					FToolMenuEntry CameraSubmenu = UE::LevelEditor::CreateLevelViewportToolbarCameraSubmenu();
-					CameraSubmenu.InsertPosition.Position = EToolMenuInsertType::First;
-					RightSection.AddEntry(CameraSubmenu);
-				}
-
 				// Add the "Show" submenu.
 				{
 					// Stay backward-compatible with the old viewport toolbar.
diff --git a/Engine/Source/Editor/MaterialEditor/Private/MaterialEditorViewportToolbarSections.cpp b/Engine/Source/Editor/MaterialEditor/Private/MaterialEditorViewportToolbarSections.cpp
--- a/Engine/Source/Editor/MaterialEditor/Private/MaterialEditorViewportToolbarSections.cpp
+++ b/Engine/Source/Editor/MaterialEditor/Private/MaterialEditorViewportToolbarSections.cpp
@@ -59,7 +59,7 @@ TSharedRef<SWidget> UE::MaterialEditor::CreateShowMenuWidget(const TSharedRef<SM
 
 FToolMenuEntry UE::MaterialEditor::CreateShowSubmenu()
 {
-	return FToolMenuEntry::InitSubMenu(
+	FToolMenuEntry Entry = FToolMenuEntry::InitSubMenu(
 		"Show",
 		LOCTEXT("ShowSubmenuLabel", "Show"),
 		LOCTEXT("ShowSubmenuTooltip", "Show options"),
@@ -70,6 +70,10 @@ FToolMenuEntry UE::MaterialEditor::CreateShowSubmenu()
 			}
 		)
 	);
+	Entry.ToolBarData.LabelOverride = FText();
+	Entry.Icon = FSlateIcon(FAppStyle::GetAppStyleSetName(), "Level.VisibleHighlightIcon16x");
+
+	return Entry;
 }
 
 void UE::MaterialEditor::FillShowSubmenu(UToolMenu* InMenu, bool bInShowViewportStatsToggle)
diff --git a/Engine/Source/Editor/MaterialEditor/Private/SMaterialEditorViewport.cpp b/Engine/Source/Editor/MaterialEditor/Private/SMaterialEditorViewport.cpp
--- a/Engine/Source/Editor/MaterialEditor/Private/SMaterialEditorViewport.cpp
+++ b/Engine/Source/Editor/MaterialEditor/Private/SMaterialEditorViewport.cpp
@@ -993,21 +993,21 @@ TSharedPtr<SWidget> SMaterialEditor3DPreviewViewport::BuildViewportToolbar()
 					UE::UnrealEd::ExtendCameraSubmenu("MaterialEditor.ViewportToolbar.CameraOptions");
 
 					FToolMenuEntry CameraSubmenu = UE::UnrealEd::CreateViewportToolbarCameraSubmenu();
-					CameraSubmenu.InsertPosition.Position = EToolMenuInsertType::First;
+					CameraSubmenu.InsertPosition.Position = EToolMenuInsertType::Last;
 					RightSection.AddEntry(CameraSubmenu);
 				}
 
 				// Add the View Modes submenu.
 				{
 					FToolMenuEntry ViewModesSubmenu = UE::UnrealEd::CreateViewportToolbarViewModesSubmenu();
-					ViewModesSubmenu.InsertPosition.Position = EToolMenuInsertType::First;
+					ViewModesSubmenu.InsertPosition.Position = EToolMenuInsertType::Last;
 					RightSection.AddEntry(ViewModesSubmenu);
 				}
 
 				// Add the Show submenu.
 				{
 					FToolMenuEntry ShowSubmenu = UE::MaterialEditor::CreateShowSubmenu();
-					ShowSubmenu.InsertPosition.Position = EToolMenuInsertType::First;
+					ShowSubmenu.InsertPosition.Position = EToolMenuInsertType::Last;
 					RightSection.AddEntry(ShowSubmenu);
 				}
 
@@ -1015,7 +1015,7 @@ TSharedPtr<SWidget> SMaterialEditor3DPreviewViewport::BuildViewportToolbar()
 				{
 					FToolMenuEntry PerformanceAndScalabilitySubmenu =
 						UE::UnrealEd::CreatePerformanceAndScalabilitySubmenu();
-					PerformanceAndScalabilitySubmenu.InsertPosition.Position = EToolMenuInsertType::First;
+					PerformanceAndScalabilitySubmenu.InsertPosition.Position = EToolMenuInsertType::Last;
 					RightSection.AddEntry(PerformanceAndScalabilitySubmenu);
 				}
 
diff --git a/Engine/Source/Editor/UnrealEd/Private/ViewportToolbar/UnrealEdViewportToolbar.cpp b/Engine/Source/Editor/UnrealEd/Private/ViewportToolbar/UnrealEdViewportToolbar.cpp
--- a/Engine/Source/Editor/UnrealEd/Private/ViewportToolbar/UnrealEdViewportToolbar.cpp
+++ b/Engine/Source/Editor/UnrealEd/Private/ViewportToolbar/UnrealEdViewportToolbar.cpp
@@ -3002,7 +3002,7 @@ bool ShouldShowViewportRealtimeWarning(const FEditorViewportClient& ViewportClie
 
 FToolMenuEntry CreatePerformanceAndScalabilitySubmenu()
 {
-	return FToolMenuEntry::InitSubMenu(
+	FToolMenuEntry Entry = FToolMenuEntry::InitSubMenu(
 		"PerformanceAndScalability",
 		LOCTEXT("PerformanceAndScalabilitySubmenuLabel", "Performance and Scalability"),
 		LOCTEXT("PerformanceAndScalabilitySubmenuTooltip", "Performance and scalability tools tied to this viewport."),
@@ -3022,6 +3022,10 @@ FToolMenuEntry CreatePerformanceAndScalabilitySubmenu()
 			}
 		)
 	);
+	Entry.Icon = FSlateIcon(FAppStyle::GetAppStyleSetName(), "Icons.Scalability");
+	Entry.ToolBarData.LabelOverride = FText();
+
+	return Entry;
 }
 
 FToolMenuEntry CreateDefaultShowSubmenu()
