[Viewport Toolbar] Use new Selection menu in the Select menu

diff --git a/Engine/Plugins/Experimental/GeometryCollectionPlugin/Source/GeometryCollectionEditor/Private/GeometryCollection/GeometryCollectionEditorPlugin.cpp b/Engine/Plugins/Experimental/GeometryCollectionPlugin/Source/GeometryCollectionEditor/Private/GeometryCollection/GeometryCollectionEditorPlugin.cpp
--- a/Engine/Plugins/Experimental/GeometryCollectionPlugin/Source/GeometryCollectionEditor/Private/GeometryCollection/GeometryCollectionEditorPlugin.cpp
+++ b/Engine/Plugins/Experimental/GeometryCollectionPlugin/Source/GeometryCollectionEditor/Private/GeometryCollection/GeometryCollectionEditorPlugin.cpp
@@ -344,9 +344,14 @@ void IGeometryCollectionEditorPlugin::RegisterMenus()
 {
 	FToolMenuOwnerScoped OwnerScoped(this);
 
-	UToolMenu* Menu = UToolMenus::Get()->ExtendMenu("LevelEditor.MainMenu.Select");
-	FToolMenuSection& Section = Menu->AddSection("SelectBones", NSLOCTEXT("LevelViewportContextMenu", "GeometryCollectionHeading", "Geometry Collection"));
 	{
+		UToolMenu* Menu = UToolMenus::Get()->ExtendMenu("LevelEditor.MainMenu.Select.GeometryCollections");
+
+		FToolMenuSection& Section = Menu->AddSection(
+			"GeometryCollections",
+			NSLOCTEXT("LevelViewportContextMenu", "GeometryCollectionHeading", "Geometry Collection")
+		);
+
 		Section.AddMenuEntry(FGeometryCollectionSelectionCommands::Get().SelectAllGeometry);
 		Section.AddMenuEntry(FGeometryCollectionSelectionCommands::Get().SelectNone);
 		Section.AddMenuEntry(FGeometryCollectionSelectionCommands::Get().SelectInverseGeometry);
diff --git a/Engine/Source/Editor/Kismet/Private/SSCSEditorViewport.cpp b/Engine/Source/Editor/Kismet/Private/SSCSEditorViewport.cpp
--- a/Engine/Source/Editor/Kismet/Private/SSCSEditorViewport.cpp
+++ b/Engine/Source/Editor/Kismet/Private/SSCSEditorViewport.cpp
@@ -435,13 +435,6 @@ TSharedPtr<SWidget> SSCSEditorViewport::BuildViewportToolbar()
 				LeftSection.AddEntry(TransformsSubmenu);
 			}
 
-			// Add the "Selection" sub menu.
-			{
-				FToolMenuEntry SelectionSubmenu = UE::UnrealEd::CreateViewportToolbarSelectSection();
-				SelectionSubmenu.InsertPosition.Position = EToolMenuInsertType::First;
-				LeftSection.AddEntry(SelectionSubmenu);
-			}
-
 			// Add the "Snapping" sub menu.
 			{
 				FToolMenuEntry SnappingSubmenu = UE::UnrealEd::CreateViewportToolbarSnappingSubmenu();
diff --git a/Engine/Source/Editor/LevelEditor/Private/LevelEditorMenu.cpp b/Engine/Source/Editor/LevelEditor/Private/LevelEditorMenu.cpp
--- a/Engine/Source/Editor/LevelEditor/Private/LevelEditorMenu.cpp
+++ b/Engine/Source/Editor/LevelEditor/Private/LevelEditorMenu.cpp
@@ -26,6 +26,7 @@
 #include "AssetSelection.h"
 #include "EditorBuildUtils.h"
 #include "EditorViewportCommands.h"
+#include "Styling/SlateIconFinder.h"
 #include "Toolkits/GlobalEditorCommonCommands.h"
 
 #define LOCTEXT_NAMESPACE "LevelEditorMenu"
@@ -752,70 +753,268 @@ void FLevelEditorMenu::RegisterBuildMenu()
 void FLevelEditorMenu::RegisterSelectMenu()
 {
 	static const FName BaseMenuName = "LevelEditor.MainMenu.Select";
-	UToolMenu* Menu = UToolMenus::Get()->RegisterMenu(BaseMenuName, NAME_None, EMultiBoxType::Menu, /*bWarnIfAlreadyRegistered*/false);
+	UToolMenu* Menu =
+		UToolMenus::Get()->RegisterMenu(BaseMenuName, NAME_None, EMultiBoxType::Menu, /*bWarnIfAlreadyRegistered*/ false);
 
+	// Main section
 	{
-		FToolMenuSection& Section = Menu->AddSection("SelectActorGeneral", NSLOCTEXT("LevelViewportContextMenu", "SelectAnyHeading", "General"));
-		Section.AddMenuEntry(FGenericCommands::Get().SelectAll, TAttribute<FText>(), NSLOCTEXT("LevelViewportContextMenu", "SelectAll_ToolTip", "Selects all actors"));
-		Section.AddMenuEntry(FLevelEditorCommands::Get().SelectNone);
-		Section.AddMenuEntry(FLevelEditorCommands::Get().InvertSelection);
-		Section.AddDynamicEntry(NAME_None, FNewToolMenuSectionDelegate::CreateLambda([](FToolMenuSection& InSection)
-			{
-				FSelectedActorInfo SelectionInfo = AssetSelectionUtils::GetSelectedActorInfo();
-				TAttribute<FText> SelectAllActorsText;
-				if (SelectionInfo.bAllSelectedActorsOfSameType && !SelectionInfo.SelectionStr.IsEmpty())
-				{
-					SelectAllActorsText = FText::Format(NSLOCTEXT("LevelViewportContextMenu", "SelectActorsOfSameClass", "Select All {0}(s)"), FText::FromString(SelectionInfo.SelectionStr));
-				}
-				InSection.AddMenuEntry(FLevelEditorCommands::Get().SelectAllActorsOfSameClass, SelectAllActorsText);
-			}));
-	}
+		FToolMenuSection& UnnamedSection = Menu->FindOrAddSection(NAME_None);
 
-	{
-		FToolMenuSection& Section = Menu->AddSection("SelectActorHierarchy", NSLOCTEXT("LevelViewportContextMenu", "SelectHierarchyHeading", "Hierarchy"));
-		Section.AddMenuEntry(FLevelEditorCommands::Get().SelectImmediateChildren);
-		Section.AddMenuEntry(FLevelEditorCommands::Get().SelectAllDescendants);
-	}
+		UnnamedSection.AddMenuEntry(
+			FGenericCommands::Get().SelectAll,
+			FGenericCommands::Get().SelectAll->GetLabel(),
+			FGenericCommands::Get().SelectAll->GetDescription(),
+			FSlateIconFinder::FindIcon("FoliageEditMode.SelectAll")
+		);
 
-	{
-		FToolMenuSection& Section = Menu->AddSection("SelectBSP", NSLOCTEXT("LevelViewportContextMenu", "SelectBSPHeading", "BSP"));
-		Section.AddMenuEntry(FLevelEditorCommands::Get().SelectAllAddditiveBrushes);
-		Section.AddMenuEntry(FLevelEditorCommands::Get().SelectAllSubtractiveBrushes);
-		Section.AddMenuEntry(FLevelEditorCommands::Get().SelectAllSurfaces);
-	}
+		UnnamedSection.AddMenuEntry(
+			FLevelEditorCommands::Get().SelectNone,
+			FLevelEditorCommands::Get().SelectNone->GetLabel(),
+			FLevelEditorCommands::Get().SelectNone->GetDescription(),
+			FSlateIconFinder::FindIcon("Cross")
+		);
 
-	{
-		FToolMenuSection& Section = Menu->AddSection("SelectLights", NSLOCTEXT("LevelViewportContextMenu", "SelectLightHeading", "Lights"));
-		Section.AddMenuEntry(FLevelEditorCommands::Get().SelectRelevantLights);
-		Section.AddMenuEntry(FLevelEditorCommands::Get().SelectAllLights);
-		Section.AddMenuEntry(FLevelEditorCommands::Get().SelectStationaryLightsExceedingOverlap);
-	}
+		UnnamedSection.AddMenuEntry(
+			FLevelEditorCommands::Get().InvertSelection,
+			FLevelEditorCommands::Get().InvertSelection->GetLabel(),
+			FLevelEditorCommands::Get().InvertSelection->GetDescription(),
+			FSlateIconFinder::FindIcon("FoliageEditMode.DeselectAll")
+		);
 
-	{
-		FToolMenuSection& Section = Menu->AddSection("SelectMeshes", NSLOCTEXT("LevelViewportContextMenu", "SelectStaticMeshHeading", "Static Meshes"));
-		Section.AddMenuEntry(FLevelEditorCommands::Get().SelectStaticMeshesOfSameClass, NSLOCTEXT("LevelViewportContextMenu", "SelectStaticMeshesOfSameClass_Menu", "Select Matching (Selected Classes)"));
-		Section.AddMenuEntry(FLevelEditorCommands::Get().SelectStaticMeshesAllClasses, NSLOCTEXT("LevelViewportContextMenu", "SelectStaticMeshesAllClasses_Menu", "Select Matching (All Classes)"));
-	}
+		// Hierarchy based selection
+		{
+			UnnamedSection.AddSubMenu(
+				"Hierarchy",
+				LOCTEXT("HierarchyLabel", "Hierarchy"),
+				LOCTEXT("HierarchyTooltip", "Hierarchy selection tools"),
+				FNewToolMenuDelegate::CreateLambda(
+					[](UToolMenu* HierarchyMenu)
+					{
+						FToolMenuSection& HierarchySection = HierarchyMenu->FindOrAddSection(
+							"SelectAllHierarchy", LOCTEXT("SelectAllHierarchyLabel", "Hierarchy")
+						);
 
-	{
-		FToolMenuSection& Section = Menu->AddSection("SelectHLODCluster", NSLOCTEXT("LevelViewportContextMenu", "SelectHLODClusterHeading", "Hierachical LODs"));
-		Section.AddMenuEntry(FLevelEditorCommands::Get().SelectOwningHierarchicalLODCluster, NSLOCTEXT("LevelViewportContextMenu", "SelectOwningHierarchicalLODCluster_Menu", "Select Owning HierarchicalLODCluster"));
-	}
+						HierarchySection.AddMenuEntry(
+							FLevelEditorCommands::Get().SelectImmediateChildren,
+							LOCTEXT("HierarchySelectImmediateChildrenLabel", "Immediate Children")
+						);
 
-	{
-		FToolMenuSection& Section = Menu->AddSection("SelectSkeletalMeshes", NSLOCTEXT("LevelViewportContextMenu", "SelectSkeletalMeshHeading", "Skeletal Meshes"));
-		Section.AddMenuEntry(FLevelEditorCommands::Get().SelectSkeletalMeshesOfSameClass);
-		Section.AddMenuEntry(FLevelEditorCommands::Get().SelectSkeletalMeshesAllClasses);
-	}
+						HierarchySection.AddMenuEntry(
+							FLevelEditorCommands::Get().SelectAllDescendants,
+							LOCTEXT("HierarchySelectAllDescendantsLabel", "All Descendants")
+						);
+					}
+				),
+				false,
+				FSlateIconFinder::FindIcon("BTEditor.SwitchToBehaviorTreeMode")
+			);
+		}
 
-	{
-		FToolMenuSection& Section = Menu->AddSection("SelectEmitters", NSLOCTEXT("LevelViewportContextMenu", "SelectEmitterHeading", "Emitters"));
-		Section.AddMenuEntry(FLevelEditorCommands::Get().SelectMatchingEmitter);
+		UnnamedSection.AddSeparator("Advanced");
+
+		UnnamedSection.AddMenuEntry(
+			FLevelEditorCommands::Get().SelectAllActorsOfSameClass,
+			LOCTEXT("AdvancedSelectAllActorsOfSameClassLabel", "All of Same Class"),
+			FLevelEditorCommands::Get().SelectAllActorsOfSameClass->GetDescription(),
+			FSlateIconFinder::FindIcon("PlacementBrowser.Icons.All")
+		);
 	}
 
+	// By Type section
 	{
-		FToolMenuSection& Section = Menu->AddSection("SelectMaterial", NSLOCTEXT("LevelViewportContextMenu", "SelectMaterialHeading", "Materials"));
-		Section.AddMenuEntry(FLevelEditorCommands::Get().SelectAllWithSameMaterial);
+		FToolMenuSection& ByTypeSection =
+			Menu->FindOrAddSection("ByTypeSection", LOCTEXT("ByTypeSectionLabel", "By Type"));
+
+		ByTypeSection.AddSubMenu(
+			"BSP",
+			LOCTEXT("BspLabel", "BSP"),
+			LOCTEXT("BspTooltip", "BSP-related tools"),
+			FNewToolMenuDelegate::CreateLambda(
+				[](UToolMenu* BspMenu)
+				{
+					FToolMenuSection& SelectAllSection =
+						BspMenu->FindOrAddSection("SelectAllBSP", LOCTEXT("SelectAllBSPLabel", "Select All BSP"));
+
+					SelectAllSection.AddMenuEntry(
+						FLevelEditorCommands::Get().SelectAllAddditiveBrushes,
+						LOCTEXT("BSPSelectAllAdditiveBrushesLabel", "Addditive Brushes")
+					);
+
+					SelectAllSection.AddMenuEntry(
+						FLevelEditorCommands::Get().SelectAllSubtractiveBrushes,
+						LOCTEXT("BSPSelectAllSubtractiveBrushesLabel", "Subtractive Brushes")
+					);
+
+					SelectAllSection.AddMenuEntry(
+						FLevelEditorCommands::Get().SelectAllSurfaces, LOCTEXT("BSPSelectAllAllSurfacesLabel", "Surfaces")
+					);
+				}
+			),
+			false,
+			FSlateIconFinder::FindIcon("ShowFlagsMenu.BSP")
+		);
+
+		ByTypeSection.AddSubMenu(
+			"Emitters",
+			LOCTEXT("EmittersLabel", "Emitters"),
+			LOCTEXT("EmittersTooltip", "Emitters-related tools"),
+			FNewToolMenuDelegate::CreateLambda(
+				[](UToolMenu* EmittersMenu)
+				{
+					FToolMenuSection& SelectAllSection = EmittersMenu->FindOrAddSection(
+						"SelectAllEmitters", LOCTEXT("SelectAllEmittersLabel", "Select All Emitters")
+					);
+
+					SelectAllSection.AddMenuEntry(
+						FLevelEditorCommands::Get().SelectMatchingEmitter,
+						LOCTEXT("EmittersSelectMatchingEmitterLabel", "Matching Emitters")
+					);
+				}
+			),
+			false,
+			FSlateIconFinder::FindIcon("ClassIcon.Emitter")
+		);
+
+		ByTypeSection.AddSubMenu(
+			"GeometryCollections",
+			LOCTEXT("GeometryCollectionsLabel", "Geometry Collections"),
+			LOCTEXT("GeometryCollectionsTooltip", "GeometryCollections-related tools"),
+			FNewToolMenuDelegate::CreateLambda(
+				[](UToolMenu* GeometryCollectionsMenu)
+				{
+					// This one will be filled by extensions from GeometryCollectionEditorPlugin
+					// Hook is "SelectGeometryCollections"
+					FToolMenuSection& SelectAllSection = GeometryCollectionsMenu->FindOrAddSection(
+						"SelectGeometryCollections", LOCTEXT("SelectGeometryCollectionsLabel", "Geometry Collections")
+					);
+				}
+			),
+			false,
+			FSlateIconFinder::FindIcon("ClassIcon.GeometryCollection")
+		);
+
+		ByTypeSection.AddSubMenu(
+			"HLOD",
+			LOCTEXT("HLODLabel", "HLOD"),
+			LOCTEXT("HLODTooltip", "HLOD-related tools"),
+			FNewToolMenuDelegate::CreateLambda(
+				[](UToolMenu* HLODMenu)
+				{
+					FToolMenuSection& SelectAllSection =
+						HLODMenu->FindOrAddSection("SelectAllHLOD", LOCTEXT("SelectAllHLODLabel", "Select All HLOD"));
+
+					SelectAllSection.AddMenuEntry(
+						FLevelEditorCommands::Get().SelectOwningHierarchicalLODCluster,
+						LOCTEXT("HLODSelectOwningHierarchicalLODClusterLabel", "Owning HLOD Cluster")
+					);
+				}
+			),
+			false,
+			FSlateIconFinder::FindIcon("WorldPartition.ShowHLODActors")
+		);
+
+		ByTypeSection.AddSubMenu(
+			"Lights",
+			LOCTEXT("LightsLabel", "Lights"),
+			LOCTEXT("LightsTooltip", "Lights-related tools"),
+			FNewToolMenuDelegate::CreateLambda(
+				[](UToolMenu* LightsMenu)
+				{
+					FToolMenuSection& SelectAllSection = LightsMenu->FindOrAddSection(
+						"SelectAllLights", LOCTEXT("SelectAllLightsLabel", "Select All Lights")
+					);
+
+					SelectAllSection.AddMenuEntry(
+						FLevelEditorCommands::Get().SelectAllLights, LOCTEXT("LightsSelectAllLightsLabel", "All Lights")
+					);
+
+					SelectAllSection.AddMenuEntry(
+						FLevelEditorCommands::Get().SelectRelevantLights,
+						LOCTEXT("LightsSelectRelevantLightsLabel", "Relevant Lights")
+					);
+
+					SelectAllSection.AddMenuEntry(
+						FLevelEditorCommands::Get().SelectStationaryLightsExceedingOverlap,
+						LOCTEXT("LightsSelectStationaryLightsExceedingOverlapLabel", "Stationary Lights Exceeding Overlap")
+					);
+				}
+			),
+			false,
+			FSlateIconFinder::FindIcon("PlacementBrowser.Icons.Lights")
+		);
+
+		ByTypeSection.AddSubMenu(
+			"Material",
+			LOCTEXT("MaterialLabel", "Material"),
+			LOCTEXT("MaterialTooltip", "Material-related tools"),
+			FNewToolMenuDelegate::CreateLambda(
+				[](UToolMenu* MaterialMenu)
+				{
+					FToolMenuSection& SelectAllSection = MaterialMenu->FindOrAddSection(
+						"SelectAllMaterial", LOCTEXT("SelectAllMaterialLabel", "Select All Material")
+					);
+
+					SelectAllSection.AddMenuEntry(
+						FLevelEditorCommands::Get().SelectAllWithSameMaterial,
+						LOCTEXT("MaterialSelectAllWithSameMaterialLabel", "With Same Material")
+					);
+				}
+			),
+			false,
+			FSlateIconFinder::FindIcon("ClassIcon.Material")
+		);
+
+		ByTypeSection.AddSubMenu(
+			"SkeletalMeshes",
+			LOCTEXT("SkeletalMeshesLabel", "Skeletal Meshes"),
+			LOCTEXT("SkeletalMeshesTooltip", "SkeletalMeshes-related tools"),
+			FNewToolMenuDelegate::CreateLambda(
+				[](UToolMenu* SkeletalMeshesMenu)
+				{
+					FToolMenuSection& SelectAllSection = SkeletalMeshesMenu->FindOrAddSection(
+						"SelectAllSkeletalMeshes", LOCTEXT("SelectAllSkeletalMeshesLabel", "Select All SkeletalMeshes")
+					);
+
+					SelectAllSection.AddMenuEntry(
+						FLevelEditorCommands::Get().SelectSkeletalMeshesOfSameClass,
+						LOCTEXT("SkeletalMeshesSelectSkeletalMeshesOfSameClassLabel", "Using Selected Skeletal Meshes (Selected Actor Types)")
+					);
+
+					SelectAllSection.AddMenuEntry(
+						FLevelEditorCommands::Get().SelectSkeletalMeshesAllClasses,
+						LOCTEXT("SkeletalMeshesSelectSkeletalMeshesAllClassesLabel", "Using Selected Skeletal Meshes (All Actor Types)")
+					);
+				}
+			),
+			false,
+			FSlateIconFinder::FindIcon("SkeletonTree.Bone")
+		);
+
+		ByTypeSection.AddSubMenu(
+			"StaticMeshes",
+			LOCTEXT("StaticMeshesLabel", "Static Meshes"),
+			LOCTEXT("StaticMeshesTooltip", "StaticMeshes-related tools"),
+			FNewToolMenuDelegate::CreateLambda(
+				[](UToolMenu* StaticMeshesMenu)
+				{
+					FToolMenuSection& SelectAllSection = StaticMeshesMenu->FindOrAddSection(
+						"SelectAllStaticMeshes", LOCTEXT("SelectAllStaticMeshesLabel", "Select All StaticMeshes")
+					);
+
+					SelectAllSection.AddMenuEntry(
+						FLevelEditorCommands::Get().SelectStaticMeshesOfSameClass,
+						LOCTEXT("StaticMeshesSelectStaticMeshesOfSameClassLabel", "Matching Selected Class")
+					);
+
+					SelectAllSection.AddMenuEntry(
+						FLevelEditorCommands::Get().SelectStaticMeshesAllClasses,
+						LOCTEXT("StaticMeshesSelectStaticMeshesAllClassesLabel", "Matching All Classes")
+					);
+				}
+			),
+			false,
+			FSlateIconFinder::FindIcon("ShowFlagsMenu.StaticMeshes")
+		);
 	}
 }
 
diff --git a/Engine/Source/Editor/StaticMeshEditor/Private/SStaticMeshEditorViewport.cpp b/Engine/Source/Editor/StaticMeshEditor/Private/SStaticMeshEditorViewport.cpp
--- a/Engine/Source/Editor/StaticMeshEditor/Private/SStaticMeshEditorViewport.cpp
+++ b/Engine/Source/Editor/StaticMeshEditor/Private/SStaticMeshEditorViewport.cpp
@@ -801,13 +801,6 @@ TSharedPtr<SWidget> SStaticMeshEditorViewport::BuildViewportToolbar()
 				LeftSection.AddEntry(TransformsSubmenu);
 			}
 
-			// Add the "Selection" sub menu.
-			{
-				FToolMenuEntry SelectionSubmenu = UE::UnrealEd::CreateViewportToolbarSelectSection();
-				SelectionSubmenu.InsertPosition.Position = EToolMenuInsertType::First;
-				LeftSection.AddEntry(SelectionSubmenu);
-			}
-
 			// Add the "Snapping" sub menu.
 			{
 				FToolMenuEntry SnappingSubmenu = UE::UnrealEd::CreateViewportToolbarSnappingSubmenu();
diff --git a/Engine/Source/Editor/UnrealEd/Private/ViewportToolbar/UnrealEdViewportToolbar.cpp b/Engine/Source/Editor/UnrealEd/Private/ViewportToolbar/UnrealEdViewportToolbar.cpp
--- a/Engine/Source/Editor/UnrealEd/Private/ViewportToolbar/UnrealEdViewportToolbar.cpp
+++ b/Engine/Source/Editor/UnrealEd/Private/ViewportToolbar/UnrealEdViewportToolbar.cpp
@@ -829,6 +829,8 @@ FToolMenuEntry CreateViewportToolbarTransformsSection()
 	return Entry;
 }
 
+// The code in this function was moved to FLevelEditorMenu::RegisterSelectMenu.
+// This function just duplicates that until it is removed after its deprecation period.
 FToolMenuEntry CreateViewportToolbarSelectSection()
 {
 	FToolMenuEntry Entry = FToolMenuEntry::InitSubMenu(
diff --git a/Engine/Source/Editor/UnrealEd/Public/ViewportToolbar/UnrealEdViewportToolbar.h b/Engine/Source/Editor/UnrealEd/Public/ViewportToolbar/UnrealEdViewportToolbar.h
--- a/Engine/Source/Editor/UnrealEd/Public/ViewportToolbar/UnrealEdViewportToolbar.h
+++ b/Engine/Source/Editor/UnrealEd/Public/ViewportToolbar/UnrealEdViewportToolbar.h
@@ -40,6 +40,7 @@ UNREALED_API bool ShowNewViewportToolbars();
 UNREALED_API FSlateIcon GetIconFromCoordSystem(ECoordSystem InCoordSystem);
 UNREALED_API FToolMenuEntry CreateViewportToolbarTransformsSection();
 
+UE_DEPRECATED(5.6, "Viewport toolbar no longer has a Select Menu.")
 UNREALED_API FToolMenuEntry CreateViewportToolbarSelectSection();
 
 UNREALED_API FToolMenuEntry CreateViewportToolbarSnappingSubmenu();
