ViewportToolbar: Minimal fix for dangling shared pointer for level viewport

diff --git a/Engine/Source/Editor/LevelEditor/Private/ViewportToolbar/LevelEditorViewportToolbarSections.cpp b/Engine/Source/Editor/LevelEditor/Private/ViewportToolbar/LevelEditorViewportToolbarSections.cpp
--- a/Engine/Source/Editor/LevelEditor/Private/ViewportToolbar/LevelEditorViewportToolbarSections.cpp
+++ b/Engine/Source/Editor/LevelEditor/Private/ViewportToolbar/LevelEditorViewportToolbarSections.cpp
@@ -2106,7 +2106,7 @@ void CreateBookmarksMenu(UToolMenu* InMenu)
 					LOCTEXT("ManageBookmarksSubMenu", "Manage Bookmarks"),
 					LOCTEXT("ManageBookmarksSubMenu_ToolTip", "Bookmarks related actions"),
 					FNewToolMenuDelegate::CreateLambda(
-						[bFoundBookmarks, LevelViewport](UToolMenu* InMenu)
+						[bFoundBookmarks, LevelViewportWeak = LevelViewportContext->LevelViewport](UToolMenu* InMenu)
 						{
 							if (!bFoundBookmarks)
 							{
@@ -2123,7 +2123,7 @@ void CreateBookmarksMenu(UToolMenu* InMenu)
 								"ClearBookmark",
 								LOCTEXT("ClearBookmarkSubMenu", "Clear Bookmark"),
 								LOCTEXT("ClearBookmarkSubMenu_ToolTip", "Clear viewport bookmarks"),
-								FNewToolMenuDelegate::CreateLambda(&Private::AddClearBookmarkMenu, LevelViewport),
+								FNewToolMenuDelegate::CreateLambda(&Private::AddClearBookmarkMenu, LevelViewportWeak),
 								false,
 								FSlateIcon(FAppStyle::Get().GetStyleSetName(), "EditorViewport.SubMenu.Bookmarks")
 							);
