ViewportToolbar: Dangling shared pointer for level viewport

diff --git a/Engine/Source/Editor/LevelEditor/Private/ViewportToolbar/LevelEditorViewportToolbarSections.cpp b/Engine/Source/Editor/LevelEditor/Private/ViewportToolbar/LevelEditorViewportToolbarSections.cpp
--- a/Engine/Source/Editor/LevelEditor/Private/ViewportToolbar/LevelEditorViewportToolbarSections.cpp
+++ b/Engine/Source/Editor/LevelEditor/Private/ViewportToolbar/LevelEditorViewportToolbarSections.cpp
@@ -174,23 +174,21 @@ float GetFarViewPlaneValue(const TSharedRef<::SLevelViewport>& InLevelViewport)
 	return ViewportClient.GetFarClipPlaneOverride();
 }
 
-bool AddJumpToBookmarkMenu(UToolMenu* InMenu, const TWeakPtr<::SLevelViewport>& InViewport)
+bool AddJumpToBookmarkMenu(UToolMenu* InMenu, const FLevelEditorViewportClient* ViewportClient)
 {
 	FToolMenuSection& Section =
 		InMenu->FindOrAddSection("JumpToBookmark", LOCTEXT("JumpToBookmarksSectionName", "Jump to Bookmark"));
 
 	// Add a menu entry for each bookmark
-	TSharedPtr<::SLevelViewport> SharedViewport = InViewport.Pin();
-	FLevelEditorViewportClient& ViewportClient = SharedViewport->GetLevelViewportClient();
 
-	const int32 NumberOfBookmarks = static_cast<int32>(IBookmarkTypeTools::Get().GetMaxNumberOfBookmarks(&ViewportClient));
+	const int32 NumberOfBookmarks = static_cast<int32>(IBookmarkTypeTools::Get().GetMaxNumberOfBookmarks(ViewportClient));
 	const int32 NumberOfMappedBookmarks = FMath::Min<int32>(AWorldSettings::NumMappedBookmarks, NumberOfBookmarks);
 
 	bool bFoundAnyBookmarks = false;
 
 	for (int32 BookmarkIndex = 0; BookmarkIndex < NumberOfMappedBookmarks; ++BookmarkIndex)
 	{
-		if (IBookmarkTypeTools::Get().CheckBookmark(BookmarkIndex, &ViewportClient))
+		if (IBookmarkTypeTools::Get().CheckBookmark(BookmarkIndex, ViewportClient))
 		{
 			bFoundAnyBookmarks = true;
 			Section.AddMenuEntry(
@@ -212,38 +210,40 @@ void AddClearBookmarkMenu(UToolMenu* InMenu, const TWeakPtr<::SLevelViewport>& I
 
 	// Add a menu entry for each bookmark
 	// FEditorModeTools& Tools = GLevelEditorModeTools();
-	TSharedPtr<::SLevelViewport> SharedViewport = InViewport.Pin();
-	FLevelEditorViewportClient& ViewportClient = SharedViewport->GetLevelViewportClient();
+	if (TSharedPtr<::SLevelViewport> LevelViewportPinned = InViewport.Pin())
+	{
+		FLevelEditorViewportClient& ViewportClient = LevelViewportPinned->GetLevelViewportClient();
 
-	const int32 NumberOfBookmarks = static_cast<int32>(IBookmarkTypeTools::Get().GetMaxNumberOfBookmarks(&ViewportClient));
-	const int32 NumberOfMappedBookmarks = FMath::Min<int32>(AWorldSettings::NumMappedBookmarks, NumberOfBookmarks);
+		const int32 NumberOfBookmarks = static_cast<int32>(IBookmarkTypeTools::Get().GetMaxNumberOfBookmarks(&ViewportClient));
+		const int32 NumberOfMappedBookmarks = FMath::Min<int32>(AWorldSettings::NumMappedBookmarks, NumberOfBookmarks);
 
-	for (int32 BookmarkIndex = 0; BookmarkIndex < NumberOfMappedBookmarks; ++BookmarkIndex)
-	{
-		if (IBookmarkTypeTools::Get().CheckBookmark(BookmarkIndex, &ViewportClient))
+		for (int32 BookmarkIndex = 0; BookmarkIndex < NumberOfMappedBookmarks; ++BookmarkIndex)
 		{
-			Section.AddMenuEntry(
-				NAME_None,
-				FLevelViewportCommands::Get().ClearBookmarkCommands[BookmarkIndex],
-				FBookmarkUI::GetPlainLabel(BookmarkIndex)
-			);
+			if (IBookmarkTypeTools::Get().CheckBookmark(BookmarkIndex, &ViewportClient))
+			{
+				Section.AddMenuEntry(
+					NAME_None,
+					FLevelViewportCommands::Get().ClearBookmarkCommands[BookmarkIndex],
+					FBookmarkUI::GetPlainLabel(BookmarkIndex)
+				);
+			}
 		}
-	}
 
-	for (int32 BookmarkIndex = NumberOfMappedBookmarks; BookmarkIndex < NumberOfBookmarks; ++BookmarkIndex)
-	{
-		if (IBookmarkTypeTools::Get().CheckBookmark(BookmarkIndex, &ViewportClient))
+		for (int32 BookmarkIndex = NumberOfMappedBookmarks; BookmarkIndex < NumberOfBookmarks; ++BookmarkIndex)
 		{
-			FUIAction Action;
-			Action.ExecuteAction.BindSP(SharedViewport.ToSharedRef(), &::SLevelViewport::OnClearBookmark, BookmarkIndex);
+			if (IBookmarkTypeTools::Get().CheckBookmark(BookmarkIndex, &ViewportClient))
+			{
+				FUIAction Action;
+				Action.ExecuteAction.BindSP(LevelViewportPinned.ToSharedRef(), &::SLevelViewport::OnClearBookmark, BookmarkIndex);
 
-			Section.AddMenuEntry(
-				NAME_None,
-				FBookmarkUI::GetPlainLabel(BookmarkIndex),
-				FBookmarkUI::GetClearTooltip(BookmarkIndex),
-				FSlateIcon(FAppStyle::GetAppStyleSetName(), "GraphEditor.Clean"),
-				Action
-			);
+				Section.AddMenuEntry(
+					NAME_None,
+					FBookmarkUI::GetPlainLabel(BookmarkIndex),
+					FBookmarkUI::GetClearTooltip(BookmarkIndex),
+					FSlateIcon(FAppStyle::GetAppStyleSetName(), "GraphEditor.Clean"),
+					Action
+				);
+			}
 		}
 	}
 }
@@ -2040,25 +2040,19 @@ void CreateBookmarksMenu(UToolMenu* InMenu)
 		return;
 	}
 
-	const TSharedPtr<::SLevelViewport> LevelViewport = LevelViewportContext->LevelViewport.Pin();
-
-	if (!LevelViewport)
+	const TSharedPtr<::SLevelViewport> LevelViewportPinned = LevelViewportContext->LevelViewport.Pin();
+	if (!LevelViewportPinned)
 	{
 		return;
 	}
 
 	// Add a menu entry for each bookmark
-	FLevelEditorViewportClient& ViewportClient = LevelViewport->GetLevelViewportClient();
+	const FLevelEditorViewportClient& ViewportClient = LevelViewportPinned->GetLevelViewportClient();
 
 	FToolMenuSection& ManageBookmarksSection =
 		InMenu->FindOrAddSection("ManageBookmarks", LOCTEXT("ManageBookmarkSectionName", "Manage Bookmarks"));
 
-	bool bFoundBookmarks = false;
-
-	// Jump to Bookmark Section
-	{
-		bFoundBookmarks = Private::AddJumpToBookmarkMenu(InMenu, LevelViewport);
-	}
+	bool bFoundBookmarks = Private::AddJumpToBookmarkMenu(InMenu, &ViewportClient);
 
 	// Manage Bookmarks Section
 	{
