ViewportToolbar: Swap shared pointer to level viewport in tool menus callback for weak pointer.

diff --git a/Engine/Source/Editor/LevelEditor/Private/ViewportToolbar/LevelEditorViewportToolbarSections.cpp b/Engine/Source/Editor/LevelEditor/Private/ViewportToolbar/LevelEditorViewportToolbarSections.cpp
--- a/Engine/Source/Editor/LevelEditor/Private/ViewportToolbar/LevelEditorViewportToolbarSections.cpp
+++ b/Engine/Source/Editor/LevelEditor/Private/ViewportToolbar/LevelEditorViewportToolbarSections.cpp
@@ -2257,10 +2257,13 @@ void AddCameraActorSelectSection(UToolMenu* InMenu)
 			CameraActorsHeading,
 			LOCTEXT("LookThroughPlacedCameras_ToolTip", "Look through and pilot placed cameras"),
 			FNewToolMenuDelegate::CreateLambda(
-				[LookThroughActors, LevelViewport](UToolMenu* InMenu)
+				[LookThroughActors, LevelViewportWeak = LevelViewportContext->LevelViewport](UToolMenu* InMenu)
 				{
-					FToolMenuSection& Section = InMenu->FindOrAddSection(NAME_None);
-					UE::LevelEditor::Private::GeneratePlacedCameraMenuEntries(Section, LookThroughActors, LevelViewport);
+					if (TSharedPtr<::SLevelViewport> LevelViewport = LevelViewportWeak.Pin())
+					{
+						FToolMenuSection& Section = InMenu->FindOrAddSection(NAME_None);
+						UE::LevelEditor::Private::GeneratePlacedCameraMenuEntries(Section, LookThroughActors, LevelViewport);
+					}
 				}
 			)
 		);
