Optimized the uncontrolled and unsaved filter processing in the scene outliner

These used to be very slow (~8ms per item queried) as they would convert a package name to a filename, which involved querying the disk. This change switches that so that the cached state is built as package names (from the filenames, which is a fast in-memory conversion), which can then be queried almost instantly by the filters.

diff --git a/Engine/Source/Editor/LevelEditor/Private/LevelEditorOutlinerSettings.cpp b/Engine/Source/Editor/LevelEditor/Private/LevelEditorOutlinerSettings.cpp
--- a/Engine/Source/Editor/LevelEditor/Private/LevelEditorOutlinerSettings.cpp
+++ b/Engine/Source/Editor/LevelEditor/Private/LevelEditorOutlinerSettings.cpp
@@ -300,48 +300,41 @@ void FLevelEditorOutlinerSettings::CreateDefaultFilters()
 
 bool FLevelEditorOutlinerSettings::DoesActorPassUnsavedFilter(const ISceneOutlinerTreeItem& InItem)
 {
-	if (InItem.GetPackageName().IsEmpty())
-	{
-		return false;
-	}
-
-	return UnsavedAssets.Contains(USourceControlHelpers::PackageFilename(InItem.GetPackageName()));
+	return UnsavedPackages.Contains(*InItem.GetPackageName());
 }
 
 bool FLevelEditorOutlinerSettings::DoesActorPassUncontrolledFilter(const ISceneOutlinerTreeItem& InItem)
 {
-	if (InItem.GetPackageName().IsEmpty())
-	{
-		return false;
-	}
-
-	FString ExternalPackageFilename = USourceControlHelpers::PackageFilename(InItem.GetPackageName());
-		
-	for (const TSharedRef<FUncontrolledChangelistState>& UncontrolledChangelistState : UncontrolledChangelistStates)
-	{
-		if (UncontrolledChangelistState->GetFilenames().Contains(ExternalPackageFilename))
-		{
-			return true;
-		}
-	}
-	
-	return false;
+	return UncontrolledPackages.Contains(*InItem.GetPackageName());
 }
 
 void FLevelEditorOutlinerSettings::OnUnsavedAssetAdded(const FString& InAsset)
 {
-	UnsavedAssets.Add(InAsset);
+	FNameBuilder PackageName;
+	if (FPackageName::TryConvertFilenameToLongPackageName(InAsset, PackageName))
+	{
+		bool bAlreadyInSet = false;
+		UnsavedPackages.Add(FName(PackageName), &bAlreadyInSet);
 
-	// Refresh any outliners that have the unsaved assets filter active to refilter on adding an unsaved asset
-	RefreshOutlinersWithActiveFilter(/*bFullRefresh*/ true, UnsavedAssetsFilterName);
+		if (!bAlreadyInSet)
+		{
+			// Refresh any outliners that have the unsaved assets filter active to refilter on adding an unsaved asset
+			RefreshOutlinersWithActiveFilter(/*bFullRefresh*/ true, UnsavedAssetsFilterName);
+		}
+	}
 }
 
 void FLevelEditorOutlinerSettings::OnUnsavedAssetRemoved(const FString& InAsset)
 {
-	UnsavedAssets.Remove(InAsset);
-
-	// Refresh any outliners that have the unsaved assets filter active to refilter on adding an unsaved asset
-	RefreshOutlinersWithActiveFilter(/*bFullRefresh*/ true, UnsavedAssetsFilterName);
+	FNameBuilder PackageName;
+	if (FPackageName::TryConvertFilenameToLongPackageName(InAsset, PackageName))
+	{
+		if (UnsavedPackages.Remove(FName(PackageName)) > 0)
+		{
+			// Refresh any outliners that have the unsaved assets filter active to refilter on adding an unsaved asset
+			RefreshOutlinersWithActiveFilter(/*bFullRefresh*/ true, UnsavedAssetsFilterName);
+		}
+	}
 }
 
 void FLevelEditorOutlinerSettings::CreateSCCFilters()
@@ -350,9 +343,8 @@ void FLevelEditorOutlinerSettings::CreateSCCFilters()
 	TSharedPtr<FFilterCategory> SCCFiltersCategory = GetFilterCategory(FLevelEditorOutlinerBuiltInCategories::SourceControl());
 	
 	// Uncontrolled Actors Filter
-	FUncontrolledChangelistsModule& UncontrolledChangelistModule = FUncontrolledChangelistsModule::Get();
-	UncontrolledChangelistStates = FUncontrolledChangelistsModule::Get().GetChangelistStates();
-	UncontrolledChangelistModule.OnUncontrolledChangelistModuleChanged.AddSP(this, &FLevelEditorOutlinerSettings::OnUncontrolledChangelistModuleChanged);
+	FUncontrolledChangelistsModule::Get().OnUncontrolledChangelistModuleChanged.AddSP(this, &FLevelEditorOutlinerSettings::OnUncontrolledChangelistModuleChanged);
+	OnUncontrolledChangelistModuleChanged();
 	
 	FGenericFilter<const ISceneOutlinerTreeItem&>::FOnItemFiltered UncontrolledFilterDelegate = FGenericFilter<const ISceneOutlinerTreeItem&>::FOnItemFiltered::CreateSP(this, &FLevelEditorOutlinerSettings::DoesActorPassUncontrolledFilter);
 	TSharedPtr<FGenericFilter<const ISceneOutlinerTreeItem&>> UncontrolledFilter = MakeShared<FGenericFilter<const ISceneOutlinerTreeItem&>>(SCCFiltersCategory, UncontrolledAssetsFilterName, LOCTEXT("UncontrolledFilterName", "Uncontrolled"), UncontrolledFilterDelegate);
@@ -365,8 +357,19 @@ void FLevelEditorOutlinerSettings::CreateSCCFilters()
 	// Unsaved Actors
 	FUnsavedAssetsTrackerModule& UnsavedAssetsTrackerModule = FUnsavedAssetsTrackerModule::Get();
 	
-	UnsavedAssets.Empty();
-	UnsavedAssets.Append(UnsavedAssetsTrackerModule.GetUnsavedAssets());
+	// Update our cached unsaved packages
+	{
+		UnsavedPackages.Reset();
+
+		FNameBuilder PackageName;
+		for (const FString& Filename : UnsavedAssetsTrackerModule.GetUnsavedAssets())
+		{
+			if (FPackageName::TryConvertFilenameToLongPackageName(Filename, PackageName))
+			{
+				UnsavedPackages.Add(FName(PackageName));
+			}
+		}
+	}
 
 	UnsavedAssetsTrackerModule.OnUnsavedAssetAdded.AddSP(this, &FLevelEditorOutlinerSettings::OnUnsavedAssetAdded);
 	UnsavedAssetsTrackerModule.OnUnsavedAssetRemoved.AddSP(this, &FLevelEditorOutlinerSettings::OnUnsavedAssetRemoved);
@@ -379,8 +382,22 @@ void FLevelEditorOutlinerSettings::CreateSCCFilters()
 
 void FLevelEditorOutlinerSettings::OnUncontrolledChangelistModuleChanged()
 {
-	// Update our cached uncontrolled CL states
-	UncontrolledChangelistStates = FUncontrolledChangelistsModule::Get().GetChangelistStates();
+	// Update our cached uncontrolled packages
+	{
+		UncontrolledPackages.Reset();
+
+		FNameBuilder PackageName;
+		for (const TSharedRef<FUncontrolledChangelistState>& UncontrolledChangelistState : FUncontrolledChangelistsModule::Get().GetChangelistStates())
+		{
+			for (const FString& Filename : UncontrolledChangelistState->GetFilenames())
+			{
+				if (FPackageName::TryConvertFilenameToLongPackageName(Filename, PackageName))
+				{
+					UncontrolledPackages.Add(FName(PackageName));
+				}
+			}
+		}
+	}
 
 	// Refresh any Outliners that have the Uncontrolled Filter active
 	RefreshOutlinersWithActiveFilter(/*bFullRefresh*/ true, UncontrolledAssetsFilterName);
diff --git a/Engine/Source/Editor/LevelEditor/Public/LevelEditorOutlinerSettings.h b/Engine/Source/Editor/LevelEditor/Public/LevelEditorOutlinerSettings.h
--- a/Engine/Source/Editor/LevelEditor/Public/LevelEditorOutlinerSettings.h
+++ b/Engine/Source/Editor/LevelEditor/Public/LevelEditorOutlinerSettings.h
@@ -5,7 +5,6 @@
 #include "Filters/GenericFilter.h"
 
 struct ISceneOutlinerTreeItem;
-class FUncontrolledChangelistState;
 
 namespace SceneOutliner
 {
@@ -116,11 +115,11 @@ private:
 
 	// Source Control Cache
 
-	// List of currently unsaved assets. Using a set here because it can grow a lot and impact performances when looking up in the list.
-	TSet<FString> UnsavedAssets;
+	// List of currently unsaved packages. Using a set here because it can grow a lot and impact performances when looking up in the list.
+	TSet<FName> UnsavedPackages;
 
-	// List of uncontrolled changelist state
-	TArray<TSharedRef<FUncontrolledChangelistState>> UncontrolledChangelistStates;
+	// List of currently uncontrolled packages. Using a set here because it can grow a lot and impact performances when looking up in the list.
+	TSet<FName> UncontrolledPackages;
 
 	static const FString UnsavedAssetsFilterName;
 	static const FString UncontrolledAssetsFilterName;
