Set clipping priorities on viewport toolbar submenus

diff --git a/Engine/Source/Editor/LevelEditor/Private/SLevelViewport.cpp b/Engine/Source/Editor/LevelEditor/Private/SLevelViewport.cpp
--- a/Engine/Source/Editor/LevelEditor/Private/SLevelViewport.cpp
+++ b/Engine/Source/Editor/LevelEditor/Private/SLevelViewport.cpp
@@ -1963,6 +1963,7 @@ TSharedPtr<SWidget> SLevelViewport::BuildViewportToolbar()
 				{
 					FToolMenuEntry TransformsSubmenu = UE::UnrealEd::CreateViewportToolbarTransformsSection();
 					TransformsSubmenu.InsertPosition.Position = EToolMenuInsertType::First;
+					TransformsSubmenu.ToolBarData.ResizeParams.ClippingPriority = 1000;
 					LeftSection.AddEntry(TransformsSubmenu);
 				}
 
@@ -1970,6 +1971,7 @@ TSharedPtr<SWidget> SLevelViewport::BuildViewportToolbar()
 				{
 					FToolMenuEntry SnappingSubmenu = UE::UnrealEd::CreateViewportToolbarSnappingSubmenu();
 					SnappingSubmenu.InsertPosition.Position = EToolMenuInsertType::First;
+					SnappingSubmenu.ToolBarData.ResizeParams.ClippingPriority = 900;
 					LeftSection.AddEntry(SnappingSubmenu);
 				}
 			}
@@ -2007,7 +2009,8 @@ TSharedPtr<SWidget> SLevelViewport::BuildViewportToolbar()
 					UE::LevelEditor::ExtendCameraSubmenu(SubmenuName);
 					UE::UnrealEd::ExtendCameraSubmenu(SubmenuName);
 
-					FToolMenuEntry CameraSubmenu = UE::LevelEditor::CreateLevelViewportToolbarCameraSubmenu();
+					constexpr uint16 ResizePriority = 800;
+					FToolMenuEntry CameraSubmenu = UE::LevelEditor::CreateLevelViewportToolbarCameraSubmenu(ResizePriority);
 					CameraSubmenu.InsertPosition.Position = EToolMenuInsertType::Last;
 					RightSection.AddEntry(CameraSubmenu);
 				}
@@ -2041,7 +2044,8 @@ TSharedPtr<SWidget> SLevelViewport::BuildViewportToolbar()
 
 					// Create and add the submenu entry to make the menu we just created and extended appear in the
 					// viewport toolbar itself.
-					FToolMenuEntry ViewModesSubmenu = UE::UnrealEd::CreateViewportToolbarViewModesSubmenu();
+					constexpr uint16 ResizePriority = 800;
+					FToolMenuEntry ViewModesSubmenu = UE::UnrealEd::CreateViewportToolbarViewModesSubmenu(ResizePriority);
 					ViewModesSubmenu.InsertPosition.Position = EToolMenuInsertType::Last;
 					RightSection.AddEntry(ViewModesSubmenu);
 				}
@@ -2061,6 +2065,7 @@ TSharedPtr<SWidget> SLevelViewport::BuildViewportToolbar()
 
 					FToolMenuEntry ShowSubmenu = UE::LevelEditor::CreateViewportToolbarShowSubmenu();
 					ShowSubmenu.InsertPosition.Position = EToolMenuInsertType::Last;
+					ShowSubmenu.ToolBarData.ResizeParams.ClippingPriority = 800;
 					RightSection.AddEntry(ShowSubmenu);
 				}
 
@@ -2068,6 +2073,7 @@ TSharedPtr<SWidget> SLevelViewport::BuildViewportToolbar()
 				{
 					FToolMenuEntry PerfSubmenu = UE::LevelEditor::CreateViewportToolbarPerformanceAndScalabilitySubmenu();
 					PerfSubmenu.InsertPosition.Position = EToolMenuInsertType::Last;
+					PerfSubmenu.ToolBarData.ResizeParams.ClippingPriority = 800;
 					RightSection.AddEntry(PerfSubmenu);
 				}
 
@@ -2075,6 +2081,7 @@ TSharedPtr<SWidget> SLevelViewport::BuildViewportToolbar()
 				{
 					FToolMenuEntry SettingsSubmenu = UE::LevelEditor::CreateLevelEditorViewportToolbarSettingsSubmenu();
 					SettingsSubmenu.InsertPosition.Position = EToolMenuInsertType::Last;
+					SettingsSubmenu.ToolBarData.ResizeParams.ClippingPriority = 1000;
 					RightSection.AddEntry(SettingsSubmenu);
 				}
 
@@ -2082,6 +2089,7 @@ TSharedPtr<SWidget> SLevelViewport::BuildViewportToolbar()
 				{
 					FToolMenuEntry ViewportSizingSubmenu = UE::LevelEditor::CreateLevelEditorViewportToolbarViewportSizingSubmenu();
 					ViewportSizingSubmenu.InsertPosition.Position = EToolMenuInsertType::Last;
+					ViewportSizingSubmenu.ToolBarData.ResizeParams.AllowClipping = false;
 					RightSection.AddEntry(ViewportSizingSubmenu);
 				}
 			}
diff --git a/Engine/Source/Editor/LevelEditor/Private/ViewportToolbar/LevelEditorViewportToolbarSections.cpp b/Engine/Source/Editor/LevelEditor/Private/ViewportToolbar/LevelEditorViewportToolbarSections.cpp
--- a/Engine/Source/Editor/LevelEditor/Private/ViewportToolbar/LevelEditorViewportToolbarSections.cpp
+++ b/Engine/Source/Editor/LevelEditor/Private/ViewportToolbar/LevelEditorViewportToolbarSections.cpp
@@ -2006,6 +2006,7 @@ FToolMenuEntry CreateLevelEditorViewportToolbarViewportSizingSubmenu()
 								FToolMenuEntry& MaximizeRestoreEntry =
 									InnerSection.AddMenuEntry("MaximizeRestore", Label, Tooltip, Icon, Action);
 								MaximizeRestoreEntry.SetShowInToolbarTopLevel(true);
+								MaximizeRestoreEntry.ToolBarData.ResizeParams.AllowClipping = false;
 								MaximizeRestoreEntry.StyleNameOverride = FName("ViewportToolbarViewportSizingSubmenu");
 							}
 						)
@@ -2016,6 +2017,7 @@ FToolMenuEntry CreateLevelEditorViewportToolbarViewportSizingSubmenu()
 	);
 
 	Entry.ToolBarData.LabelOverride = FText();
+	Entry.ToolBarData.ResizeParams.AllowClipping = false;
 	Entry.StyleNameOverride = FName("ViewportToolbarViewportSizingSubmenu");
 
 	return Entry;
@@ -2498,12 +2500,12 @@ void ExtendCameraSubmenu(FName InCameraOptionsSubmenuName)
 	);
 }
 
-FToolMenuEntry CreateLevelViewportToolbarCameraSubmenu()
+FToolMenuEntry CreateLevelViewportToolbarCameraSubmenu(int32 InClippingPriority)
 {
-	return FToolMenuEntry::InitDynamicEntry(
+	FToolMenuEntry Entry = FToolMenuEntry::InitDynamicEntry(
 		"DynamicCameraOptions",
 		FNewToolMenuSectionDelegate::CreateLambda(
-			[](FToolMenuSection& InDynamicSection) -> void
+			[InClippingPriority](FToolMenuSection& InDynamicSection) -> void
 			{
 				if (ULevelViewportContext* const LevelViewportContext =
 						InDynamicSection.FindContext<ULevelViewportContext>())
@@ -2522,7 +2524,7 @@ FToolMenuEntry CreateLevelViewportToolbarCameraSubmenu()
 						}
 					);
 
-					InDynamicSection.AddSubMenu(
+					FToolMenuEntry& Entry = InDynamicSection.AddSubMenu(
 						"CameraOptions",
 						Label,
 						LOCTEXT("CameraSubmenuTooltip", "Camera options"),
@@ -2535,10 +2537,13 @@ FToolMenuEntry CreateLevelViewportToolbarCameraSubmenu()
 						false,
 						Icon
 					);
+					Entry.ToolBarData.ResizeParams.ClippingPriority = InClippingPriority;
 				}
 			}
 		)
 	);
+
+	return Entry;
 }
 
 } // namespace UE::LevelEditor
diff --git a/Engine/Source/Editor/LevelEditor/Private/ViewportToolbar/LevelEditorViewportToolbarSections.h b/Engine/Source/Editor/LevelEditor/Private/ViewportToolbar/LevelEditorViewportToolbarSections.h
--- a/Engine/Source/Editor/LevelEditor/Private/ViewportToolbar/LevelEditorViewportToolbarSections.h
+++ b/Engine/Source/Editor/LevelEditor/Private/ViewportToolbar/LevelEditorViewportToolbarSections.h
@@ -22,7 +22,7 @@ FToolMenuEntry CreateFOVMenu(TWeakPtr<SLevelViewport> InLevelViewportWeak);
 FToolMenuEntry CreateFarViewPlaneMenu(TWeakPtr<SLevelViewport> InLevelViewportWeak);
 FToolMenuEntry CreateCameraSpeedSlider(TWeakPtr<SLevelViewport> InLevelViewportWeak);
 FToolMenuEntry CreateCameraSpeedScalarSlider(TWeakPtr<SLevelViewport> InLevelViewportWeak);
-FToolMenuEntry CreateLevelViewportToolbarCameraSubmenu();
+FToolMenuEntry CreateLevelViewportToolbarCameraSubmenu(int32 InClippingPriority = 0);
 
 TSharedPtr<FExtender> GetViewModesLegacyExtenders();
 void PopulateViewModesMenu(UToolMenu* InMenu);
diff --git a/Engine/Source/Editor/UnrealEd/Private/ViewportToolbar/UnrealEdViewportToolbar.cpp b/Engine/Source/Editor/UnrealEd/Private/ViewportToolbar/UnrealEdViewportToolbar.cpp
--- a/Engine/Source/Editor/UnrealEd/Private/ViewportToolbar/UnrealEdViewportToolbar.cpp
+++ b/Engine/Source/Editor/UnrealEd/Private/ViewportToolbar/UnrealEdViewportToolbar.cpp
@@ -1698,13 +1698,13 @@ UUnrealEdViewportToolbarContext* CreateViewportToolbarDefaultContext(const TWeak
 	return ContextObject;
 }
 
-FToolMenuEntry CreateViewportToolbarViewModesSubmenu()
+FToolMenuEntry CreateViewportToolbarViewModesSubmenu(int32 InClippingPriority)
 {
 	// This has to be a dynamic entry for the ViewModes submenu's label to be able to access the context.
 	return FToolMenuEntry::InitDynamicEntry(
 		"DynamicViewModes",
 		FNewToolMenuSectionDelegate::CreateLambda(
-			[](FToolMenuSection& InDynamicSection) -> void
+			[InClippingPriority](FToolMenuSection& InDynamicSection) -> void
 			{
 				// Base the label on the current view mode.
 				TAttribute<FText> LabelAttribute = UE::UnrealEd::GetViewModesSubmenuLabel(nullptr);
@@ -1727,7 +1727,7 @@ FToolMenuEntry CreateViewportToolbarViewModesSubmenu()
 					);
 				}
 
-				InDynamicSection.AddSubMenu(
+				FToolMenuEntry& Entry = InDynamicSection.AddSubMenu(
 					"ViewModes",
 					LabelAttribute,
 					LOCTEXT("ViewModesSubmenuTooltip", "View mode settings for the current viewport."),
@@ -1740,6 +1740,7 @@ FToolMenuEntry CreateViewportToolbarViewModesSubmenu()
 					false,
 					IconAttribute
 				);
+				Entry.ToolBarData.ResizeParams.ClippingPriority = InClippingPriority;
 			}
 		)
 	);
@@ -3077,6 +3078,7 @@ FToolMenuEntry CreateScalabilitySubmenu()
 				if (UE::UnrealEd::IsScalabilityWarningVisible())
 				{
 					ScalabilitySubmenu.StyleNameOverride = "ViewportToolbarWarning";
+					ScalabilitySubmenu.ToolBarData.ResizeParams.AllowClipping = false;
 				}
 			}
 		)
@@ -3318,6 +3320,7 @@ FToolMenuEntry CreateToggleRealtimeEntry()
 					if (!EditorViewport->IsRealtime())
 					{
 						ToggleRealtime.StyleNameOverride = "ViewportToolbarWarning";
+						ToggleRealtime.ToolBarData.ResizeParams.AllowClipping = false;
 					}
 				}
 
diff --git a/Engine/Source/Editor/UnrealEd/Public/ViewportToolbar/UnrealEdViewportToolbar.h b/Engine/Source/Editor/UnrealEd/Public/ViewportToolbar/UnrealEdViewportToolbar.h
--- a/Engine/Source/Editor/UnrealEd/Public/ViewportToolbar/UnrealEdViewportToolbar.h
+++ b/Engine/Source/Editor/UnrealEd/Public/ViewportToolbar/UnrealEdViewportToolbar.h
@@ -60,7 +60,7 @@ UNREALED_API void PopulateViewModesMenu(UToolMenu* InMenu);
 UNREALED_API UUnrealEdViewportToolbarContext* CreateViewportToolbarDefaultContext(const TWeakPtr<SEditorViewport>& InViewport
 );
 
-UNREALED_API FToolMenuEntry CreateViewportToolbarViewModesSubmenu();
+UNREALED_API FToolMenuEntry CreateViewportToolbarViewModesSubmenu(int32 InClippingPriority = 0);
 
 DECLARE_DELEGATE_TwoParams(FRotationGridCheckboxListExecuteActionDelegate, int, ERotationGridMode);
 DECLARE_DELEGATE_RetVal_TwoParams(bool, FRotationGridCheckboxListIsCheckedDelegate, int, ERotationGridMode);
