Level Editor Outliner: During map change make sure the old outliner instances are cleaned up first and then create new instances

diff --git a/Engine/Source/Editor/LevelEditor/Private/SLevelEditor.cpp b/Engine/Source/Editor/LevelEditor/Private/SLevelEditor.cpp
--- a/Engine/Source/Editor/LevelEditor/Private/SLevelEditor.cpp
+++ b/Engine/Source/Editor/LevelEditor/Private/SLevelEditor.cpp
@@ -1974,10 +1974,16 @@ void SLevelEditor::HandleEditorMapChange( uint32 MapChangeFlags )
 	// We want to recreate all the Scene Outliners on loading a new map
 	if (MapChangeFlags == MapChangeEventFlags::NewMap)
 	{
+		SceneOutliners.Empty();
+		
 		for (TPair<FName, TWeakPtr<SDockTab>> OutlinerTab : SceneOutlinerTabs)
 		{
 			if (TSharedPtr<SDockTab> SceneOutlinerTabPin = OutlinerTab.Value.Pin())
 			{
+				// While this is not strictly necessary since the old tab contents will be cleaned up on SetContent, it is a cleaner flow if we
+				// first destruct the old Outliner and then construct the new one
+				SceneOutlinerTabPin->SetContent(SNullWidget::NullWidget);
+				
 				// Create the new outliner
 				TSharedRef<ISceneOutliner> SceneOutlinerRef = CreateSceneOutliner(OutlinerTab.Key);
 
