Made AddCustomFilterToOutliner virtual so it can be called from external modules.

diff --git a/Engine/Source/Editor/LevelEditor/Public/LevelEditor.h b/Engine/Source/Editor/LevelEditor/Public/LevelEditor.h
--- a/Engine/Source/Editor/LevelEditor/Public/LevelEditor.h
+++ b/Engine/Source/Editor/LevelEditor/Public/LevelEditor.h
@@ -351,7 +351,7 @@ public:
 	 *	 the level editor creates.
 	 *   @see FGenericFilter on how to create generic filters
 	 */
-	void AddCustomFilterToOutliner(FLevelEditorOutlinerSettings::FOutlinerFilterFactory InCreateCustomFilter);
+	virtual void AddCustomFilterToOutliner(FLevelEditorOutlinerSettings::FOutlinerFilterFactory InCreateCustomFilter);
 
 	/**  Add a custom class filter to the outliner filter bar. These represent asset/actor type filters and are OR'd
 	 *   @see FLevelEditorOutlinerSettings
