Rename LevelEditorViewportToolbarSection.cpp/h to LevelEditorSubmenus.cpp/h

diff --git a/Engine/Source/Editor/LevelEditor/Private/LevelEditorToolBar.cpp b/Engine/Source/Editor/LevelEditor/Private/LevelEditorToolBar.cpp
--- a/Engine/Source/Editor/LevelEditor/Private/LevelEditorToolBar.cpp
+++ b/Engine/Source/Editor/LevelEditor/Private/LevelEditorToolBar.cpp
@@ -63,7 +63,7 @@
 #include "Styling/ToolBarStyle.h"
 #include "PlatformInfo.h"
 #include "DataDrivenShaderPlatformInfo.h"
-#include "ViewportToolbar/LevelEditorViewportToolbarSections.h"
+#include "ViewportToolbar/LevelEditorSubmenus.h"
 
 FName FLevelEditorToolBar::SecondaryModeToolbarName("LevelEditor.SecondaryToolbar");
 
diff --git a/Engine/Source/Editor/LevelEditor/Private/SLevelViewport.cpp b/Engine/Source/Editor/LevelEditor/Private/SLevelViewport.cpp
--- a/Engine/Source/Editor/LevelEditor/Private/SLevelViewport.cpp
+++ b/Engine/Source/Editor/LevelEditor/Private/SLevelViewport.cpp
@@ -104,7 +104,7 @@
 #include "EditorViewportCommands.h"
 #include "FunctionalUIScreenshotTest.h"
 #include "SkeletalRenderPublic.h"
-#include "ViewportToolbar/LevelEditorViewportToolbarSections.h"
+#include "ViewportToolbar/LevelEditorSubmenus.h"
 #include "ViewportToolbar/LevelViewportContext.h"
 #include "ViewportToolbar/UnrealEdViewportToolbar.h"
 
diff --git a/Engine/Source/Editor/LevelEditor/Private/SLevelViewportToolBar.cpp b/Engine/Source/Editor/LevelEditor/Private/SLevelViewportToolBar.cpp
--- a/Engine/Source/Editor/LevelEditor/Private/SLevelViewportToolBar.cpp
+++ b/Engine/Source/Editor/LevelEditor/Private/SLevelViewportToolBar.cpp
@@ -57,7 +57,7 @@
 #include "SCommonEditorViewportToolbarBase.h"
 #include "SActionableMessageViewportWidget.h"
 #include "Engine/SceneCapture.h"
-#include "ViewportToolbar/LevelEditorViewportToolbarSections.h"
+#include "ViewportToolbar/LevelEditorSubmenus.h"
 #include "ViewportToolbar/LevelViewportContext.h"
 #include "ViewportToolbar/UnrealEdViewportToolbar.h"
 
diff --git a/Engine/Source/Editor/LevelEditor/Private/ViewportToolbar/LevelEditorSubmenus.cpp b/Engine/Source/Editor/LevelEditor/Private/ViewportToolbar/LevelEditorSubmenus.cpp
new file mode 100644
--- /dev/null
+++ b/Engine/Source/Editor/LevelEditor/Private/ViewportToolbar/LevelEditorSubmenus.cpp
@@ -0,0 +1,2559 @@
+// Copyright Epic Games, Inc. All Rights Reserved.
+
+#include "ViewportToolbar/LevelEditorSubmenus.h"
+
+#include "Bookmarks/BookmarkUI.h"
+#include "Bookmarks/IBookmarkTypeTools.h"
+#include "Camera/CameraActor.h"
+#include "EditorViewportCommands.h"
+#include "Engine/SceneCapture.h"
+#include "EngineUtils.h"
+#include "FoliageType.h"
+#include "GameFramework/ActorPrimitiveColorHandler.h"
+#include "GameFramework/WorldSettings.h"
+#include "GroomVisualizationData.h"
+#include "Layers/LayersSubsystem.h"
+#include "LevelEditor.h"
+#include "LevelEditorActions.h"
+#include "LevelViewportActions.h"
+#include "SLevelViewport.h"
+#include "Selection.h"
+#include "SortHelper.h"
+#include "Stats/StatsData.h"
+#include "Styling/SlateIconFinder.h"
+#include "Templates/SharedPointer.h"
+#include "ToolMenu.h"
+#include "ToolMenus.h"
+#include "ViewportToolbar/LevelViewportContext.h"
+#include "ViewportToolbar/UnrealEdViewportToolbar.h"
+#include "Widgets/Input/SSpinBox.h"
+#include "Widgets/Input/SVolumeControl.h"
+#include "Widgets/SBoxPanel.h"
+#include "Widgets/SToolTip.h"
+#include "WorldPartition/IWorldPartitionEditorModule.h"
+
+#define LOCTEXT_NAMESPACE "LevelEditorViewportToolbar"
+
+namespace UE::LevelEditor::Private
+{
+
+bool IsLandscapeLODSettingChecked(FLevelEditorViewportClient& ViewportClient, int32 Value)
+{
+	return ViewportClient.LandscapeLODOverride == Value;
+}
+
+void OnLandscapeLODChanged(FLevelEditorViewportClient& ViewportClient, int32 NewValue)
+{
+	ViewportClient.LandscapeLODOverride = NewValue;
+	ViewportClient.Invalidate();
+}
+
+TMap<FName, TArray<UFoliageType*>> GroupFoliageByOuter(const TArray<UFoliageType*> FoliageList)
+{
+	TMap<FName, TArray<UFoliageType*>> Result;
+
+	for (UFoliageType* FoliageType : FoliageList)
+	{
+		if (FoliageType->IsAsset())
+		{
+			Result.FindOrAdd(NAME_None).Add(FoliageType);
+		}
+		else
+		{
+			FName LevelName = FoliageType->GetOutermost()->GetFName();
+			Result.FindOrAdd(LevelName).Add(FoliageType);
+		}
+	}
+
+	Result.KeySort(
+		[](const FName& A, const FName& B)
+		{
+			return (A.LexicalLess(B) && B != NAME_None);
+		}
+	);
+	return Result;
+}
+
+void PopulateMenuWithCommands(UToolMenu* Menu, TArray<FLevelViewportCommands::FShowMenuCommand> MenuCommands, int32 EntryOffset)
+{
+	FToolMenuSection& Section = Menu->AddSection("Section");
+
+	// Generate entries for the standard show flags
+	// Assumption: the first 'n' entries types like 'Show All' and 'Hide All' buttons, so insert a separator after them
+	for (int32 EntryIndex = 0; EntryIndex < MenuCommands.Num(); ++EntryIndex)
+	{
+		FName EntryName = NAME_None;
+
+		if (MenuCommands[EntryIndex].ShowMenuItem)
+		{
+			EntryName = MenuCommands[EntryIndex].ShowMenuItem->GetCommandName();
+			ensure(Section.FindEntry(EntryName) == nullptr);
+		}
+
+		Section.AddMenuEntry(EntryName, MenuCommands[EntryIndex].ShowMenuItem, MenuCommands[EntryIndex].LabelOverride);
+
+		if (EntryIndex == EntryOffset - 1)
+		{
+			Section.AddSeparator(NAME_None);
+		}
+	}
+}
+
+void PopulateShowLayersSubmenu(UToolMenu* InMenu, TWeakPtr<::SLevelViewport> InViewport)
+{
+	{
+		FToolMenuSection& Section = InMenu->AddSection("LevelViewportLayers");
+		Section.AddMenuEntry(FLevelViewportCommands::Get().ShowAllLayers, LOCTEXT("ShowAllLabel", "Show All"));
+		Section.AddMenuEntry(FLevelViewportCommands::Get().HideAllLayers, LOCTEXT("HideAllLabel", "Hide All"));
+	}
+
+	if (TSharedPtr<::SLevelViewport> ViewportPinned = InViewport.Pin())
+	{
+		FToolMenuSection& Section = InMenu->AddSection("LevelViewportLayers2");
+		// Get all the layers and create an entry for each of them
+		TArray<FName> AllLayerNames;
+		ULayersSubsystem* Layers = GEditor->GetEditorSubsystem<ULayersSubsystem>();
+		Layers->AddAllLayerNamesTo(AllLayerNames);
+
+		for (int32 LayerIndex = 0; LayerIndex < AllLayerNames.Num(); ++LayerIndex)
+		{
+			const FName LayerName = AllLayerNames[LayerIndex];
+			// const FString LayerNameString = LayerName;
+
+			FUIAction Action(
+				FExecuteAction::CreateSP(ViewportPinned.ToSharedRef(), &::SLevelViewport::ToggleShowLayer, LayerName),
+				FCanExecuteAction(),
+				FIsActionChecked::CreateSP(ViewportPinned.ToSharedRef(), &::SLevelViewport::IsLayerVisible, LayerName)
+			);
+
+			Section.AddMenuEntry(
+				NAME_None, FText::FromName(LayerName), FText::GetEmpty(), FSlateIcon(), Action, EUserInterfaceActionType::ToggleButton
+			);
+		}
+	}
+}
+
+void SetLevelViewportFOV(const TSharedRef<::SLevelViewport>& InLevelViewport, float InValue)
+{
+	bool bUpdateStoredFOV = true;
+
+	if (InLevelViewport->GetLevelViewportClient().GetActiveActorLock().IsValid())
+	{
+		if (ACameraActor* CameraActor =
+				Cast<ACameraActor>(InLevelViewport->GetLevelViewportClient().GetActiveActorLock().Get()))
+		{
+			CameraActor->GetCameraComponent()->FieldOfView = InValue;
+			bUpdateStoredFOV = false;
+		}
+	}
+
+	if (bUpdateStoredFOV)
+	{
+		InLevelViewport->GetLevelViewportClient().FOVAngle = InValue;
+	}
+
+	InLevelViewport->GetLevelViewportClient().ViewFOV = InValue;
+	InLevelViewport->GetLevelViewportClient().Invalidate();
+}
+
+void SetFarViewPlaneValue(const TSharedRef<::SLevelViewport>& InLevelViewport, float InValue)
+{
+	FLevelEditorViewportClient& ViewportClient = InLevelViewport->GetLevelViewportClient();
+	return ViewportClient.OverrideFarClipPlane(InValue);
+}
+
+float GetLevelViewportFOV(const TSharedRef<::SLevelViewport>& InLevelViewport)
+{
+	FLevelEditorViewportClient& ViewportClient = InLevelViewport->GetLevelViewportClient();
+	return ViewportClient.ViewFOV;
+}
+
+float GetFarViewPlaneValue(const TSharedRef<::SLevelViewport>& InLevelViewport)
+{
+	FLevelEditorViewportClient& ViewportClient = InLevelViewport->GetLevelViewportClient();
+	return ViewportClient.GetFarClipPlaneOverride();
+}
+
+bool AddJumpToBookmarkMenu(UToolMenu* InMenu, const FLevelEditorViewportClient* ViewportClient)
+{
+	FToolMenuSection& Section =
+		InMenu->FindOrAddSection("JumpToBookmark", LOCTEXT("JumpToBookmarksSectionName", "Jump to Bookmark"));
+
+	// Add a menu entry for each bookmark
+
+	const int32 NumberOfBookmarks = static_cast<int32>(IBookmarkTypeTools::Get().GetMaxNumberOfBookmarks(ViewportClient));
+	const int32 NumberOfMappedBookmarks = FMath::Min<int32>(AWorldSettings::NumMappedBookmarks, NumberOfBookmarks);
+
+	bool bFoundAnyBookmarks = false;
+
+	for (int32 BookmarkIndex = 0; BookmarkIndex < NumberOfMappedBookmarks; ++BookmarkIndex)
+	{
+		if (IBookmarkTypeTools::Get().CheckBookmark(BookmarkIndex, ViewportClient))
+		{
+			bFoundAnyBookmarks = true;
+			Section.AddMenuEntry(
+				NAME_None,
+				FLevelViewportCommands::Get().JumpToBookmarkCommands[BookmarkIndex],
+				FBookmarkUI::GetPlainLabel(BookmarkIndex),
+				FBookmarkUI::GetJumpToTooltip(BookmarkIndex),
+				FSlateIcon(FAppStyle::Get().GetStyleSetName(), "EditorViewport.SubMenu.Bookmarks")
+			);
+		}
+	}
+
+	return bFoundAnyBookmarks;
+}
+
+void AddClearBookmarkMenu(UToolMenu* InMenu, const TWeakPtr<::SLevelViewport>& InViewport)
+{
+	FToolMenuSection& Section = InMenu->AddSection("Section");
+
+	// Add a menu entry for each bookmark
+	// FEditorModeTools& Tools = GLevelEditorModeTools();
+	if (TSharedPtr<::SLevelViewport> LevelViewportPinned = InViewport.Pin())
+	{
+		FLevelEditorViewportClient& ViewportClient = LevelViewportPinned->GetLevelViewportClient();
+
+		const int32 NumberOfBookmarks = static_cast<int32>(IBookmarkTypeTools::Get().GetMaxNumberOfBookmarks(&ViewportClient));
+		const int32 NumberOfMappedBookmarks = FMath::Min<int32>(AWorldSettings::NumMappedBookmarks, NumberOfBookmarks);
+
+		for (int32 BookmarkIndex = 0; BookmarkIndex < NumberOfMappedBookmarks; ++BookmarkIndex)
+		{
+			if (IBookmarkTypeTools::Get().CheckBookmark(BookmarkIndex, &ViewportClient))
+			{
+				Section.AddMenuEntry(
+					NAME_None,
+					FLevelViewportCommands::Get().ClearBookmarkCommands[BookmarkIndex],
+					FBookmarkUI::GetPlainLabel(BookmarkIndex)
+				);
+			}
+		}
+
+		for (int32 BookmarkIndex = NumberOfMappedBookmarks; BookmarkIndex < NumberOfBookmarks; ++BookmarkIndex)
+		{
+			if (IBookmarkTypeTools::Get().CheckBookmark(BookmarkIndex, &ViewportClient))
+			{
+				FUIAction Action;
+				Action.ExecuteAction.BindSP(LevelViewportPinned.ToSharedRef(), &::SLevelViewport::OnClearBookmark, BookmarkIndex);
+
+				Section.AddMenuEntry(
+					NAME_None,
+					FBookmarkUI::GetPlainLabel(BookmarkIndex),
+					FBookmarkUI::GetClearTooltip(BookmarkIndex),
+					FSlateIcon(FAppStyle::GetAppStyleSetName(), "GraphEditor.Clean"),
+					Action
+				);
+			}
+		}
+	}
+}
+
+void GeneratePlacedCameraMenuEntries(
+	FToolMenuSection& InSection, TArray<AActor*> InLookThroughActors, const TSharedPtr<::SLevelViewport>& InLevelViewport
+)
+{
+	// Sort the cameras to make the ordering predictable for users.
+	InLookThroughActors.StableSort(
+		[](const AActor& Left, const AActor& Right)
+		{
+			// Do "natural sorting" via SceneOutliner::FNumericStringWrapper to make more sense to humans (also matches
+			// the Scene Outliner). This sorts "Camera2" before "Camera10" which a normal lexicographical sort wouldn't.
+			SceneOutliner::FNumericStringWrapper LeftWrapper(FString(Left.GetActorLabel()));
+			SceneOutliner::FNumericStringWrapper RightWrapper(FString(Right.GetActorLabel()));
+
+			return LeftWrapper < RightWrapper;
+		}
+	);
+
+	for (AActor* LookThroughActor : InLookThroughActors)
+	{
+		// Needed for the delegate hookup to work below
+		AActor* GenericActor = LookThroughActor;
+
+		FText ActorDisplayName = FText::FromString(LookThroughActor->GetActorLabel());
+		FUIAction LookThroughCameraAction(
+			FExecuteAction::CreateSP(InLevelViewport.ToSharedRef(), &::SLevelViewport::OnActorLockToggleFromMenu, GenericActor),
+			FCanExecuteAction(),
+			FIsActionChecked::CreateSP(
+				InLevelViewport.ToSharedRef(), &::SLevelViewport::IsActorLocked, MakeWeakObjectPtr(GenericActor)
+			)
+		);
+
+		FSlateIcon ActorIcon;
+
+		if (LookThroughActor->IsA<ACameraActor>() || LookThroughActor->IsA<ASceneCapture>())
+		{
+			ActorIcon = FSlateIcon(FAppStyle::GetAppStyleSetName(), "ClassIcon.CameraComponent");
+		}
+		else
+		{
+			ActorIcon = FSlateIconFinder::FindIconForClass(LookThroughActor->GetClass());
+		}
+
+		InSection.AddMenuEntry(
+			NAME_None,
+			ActorDisplayName,
+			FText::Format(LOCTEXT("LookThroughCameraActor_ToolTip", "Look through and pilot {0}"), ActorDisplayName),
+			ActorIcon,
+			LookThroughCameraAction,
+			EUserInterfaceActionType::RadioButton
+		);
+	}
+}
+
+FToolMenuEntry CreateEjectActorPilotEntry()
+{
+	return FToolMenuEntry::InitDynamicEntry(
+		"EjectActorPilotDynamicSection",
+		FNewToolMenuSectionDelegate::CreateLambda(
+			[](FToolMenuSection& InnerSection) -> void
+			{
+				ULevelViewportContext* const LevelViewportContext = InnerSection.FindContext<ULevelViewportContext>();
+				if (!LevelViewportContext)
+				{
+					return;
+				}
+
+				FToolUIAction EjectActorPilotAction;
+
+				EjectActorPilotAction.ExecuteAction = FToolMenuExecuteAction::CreateLambda(
+					[LevelViewportWeak = LevelViewportContext->LevelViewport](const FToolMenuContext& Context) -> void
+					{
+						if (TSharedPtr<::SLevelViewport> LevelViewport = LevelViewportWeak.Pin())
+						{
+							LevelViewport->OnActorLockToggleFromMenu();
+						}
+					}
+				);
+
+				EjectActorPilotAction.CanExecuteAction = FToolMenuCanExecuteAction::CreateLambda(
+					[LevelViewportWeak = LevelViewportContext->LevelViewport](const FToolMenuContext& Context)
+					{
+						if (TSharedPtr<::SLevelViewport> EditorViewport = LevelViewportWeak.Pin())
+						{
+							return EditorViewport->IsAnyActorLocked();
+						}
+						return false;
+					}
+				);
+
+				// We use this entry to gather its Name, Tooltip and Icon. See comment below as to why we cannot directly use this entry.
+				FToolMenuEntry SourceEjectPilotEntry =
+					FToolMenuEntry::InitMenuEntry(FLevelViewportCommands::Get().EjectActorPilot);
+
+				// We want to use SetShowInToolbarTopLevel to show the Eject entry in the Top Level only when piloting is active.
+				// Currently, this will not work with Commands, e.g. AddMenuEntry(FLevelViewportCommands::Get().EjectActorPilot).
+				// So, we create the entry using FToolMenuEntry::InitMenuEntry, and we create our own Action to handle it.
+				FToolMenuEntry EjectPilotActor = FToolMenuEntry::InitMenuEntry(
+					"EjectActorPilot",
+					LOCTEXT("EjectActorPilotLabel", "Stop Piloting Actor"),
+					LOCTEXT(
+						"EjectActorPilotTooltip", "Stop piloting an actor with the current viewport. Unlocks the viewport's position and orientation from the actor the viewport is currently piloting."
+					),
+					FSlateIcon(FAppStyle::GetAppStyleSetName(), "LevelViewport.EjectActorPilot"),
+					EjectActorPilotAction,
+					EUserInterfaceActionType::Button
+				);
+
+				const TAttribute<bool> bShownInTopLevel = TAttribute<bool>::CreateLambda(
+					[ViewportContextWeak = TWeakObjectPtr<ULevelViewportContext>(LevelViewportContext)]() -> bool
+					{
+						if (TStrongObjectPtr<ULevelViewportContext> LevelViewportContext = ViewportContextWeak.Pin())
+						{
+							if (const TSharedPtr<::SLevelViewport> LevelViewport = LevelViewportContext->LevelViewport.Pin())
+							{
+								return LevelViewport->GetLevelViewportClient().IsAnyActorLocked();
+							}
+						}
+
+						return true;
+					}
+				);
+
+				EjectPilotActor.SetShowInToolbarTopLevel(bShownInTopLevel);
+
+				InnerSection.AddEntry(EjectPilotActor);
+			}
+		)
+	);
+}
+
+FText GetCameraSubmenuLabelFromLevelViewport(const TWeakPtr<::SLevelViewport>& InLevelEditorViewportClientWeak)
+{
+	if (TSharedPtr<::SLevelViewport> LevelViewport = InLevelEditorViewportClientWeak.Pin())
+	{
+		const FLevelEditorViewportClient& LevelViewportClient = LevelViewport->GetLevelViewportClient();
+
+		if (!LevelViewportClient.IsAnyActorLocked())
+		{
+			return UnrealEd::GetCameraSubmenuLabelFromViewportType(LevelViewportClient.GetViewportType());
+		}
+		else if (TStrongObjectPtr<AActor> ActorLock = LevelViewportClient.GetActiveActorLock().Pin())
+		{
+			return FText::FromString(ActorLock->GetActorNameOrLabel());
+		}
+	}
+
+	return LOCTEXT("MissingActiveCameraLabel", "No Active Camera");
+}
+
+FSlateIcon GetCameraSubmenuIconFromLevelViewport(const TWeakPtr<::SLevelViewport>& InLevelEditorViewportClientWeak)
+{
+	if (TSharedPtr<::SLevelViewport> LevelViewport = InLevelEditorViewportClientWeak.Pin())
+	{
+		const FLevelEditorViewportClient& LevelViewportClient = LevelViewport->GetLevelViewportClient();
+		if (!LevelViewportClient.IsAnyActorLocked())
+		{
+			const FName IconName =
+				UnrealEd::GetCameraSubmenuIconFNameFromViewportType(LevelViewportClient.GetViewportType());
+			return FSlateIcon(FAppStyle::GetAppStyleSetName(), IconName);
+		}
+		else if (TStrongObjectPtr<AActor> LockedActor = LevelViewportClient.GetActorLock().LockedActor.Pin())
+		{
+			if (!LockedActor->IsA<ACameraActor>() && !LockedActor->IsA<ASceneCapture>())
+			{
+				return FSlateIconFinder::FindIconForClass(LockedActor->GetClass());
+			}
+		}
+	}
+
+	return FSlateIcon(FAppStyle::GetAppStyleSetName(), "ClassIcon.CameraComponent");
+}
+} // namespace UE::LevelEditor::Private
+
+namespace UE::LevelEditor
+{
+
+TSharedPtr<FExtender> GetViewModesLegacyExtenders()
+{
+	FLevelEditorModule& LevelEditorModule = FModuleManager::GetModuleChecked<FLevelEditorModule>(TEXT("LevelEditor"));
+	return LevelEditorModule.GetMenuExtensibilityManager()->GetAllExtenders();
+}
+
+void PopulateViewModesMenu(UToolMenu* InMenu)
+{
+	FToolMenuInsert InsertPosition("ViewMode", EToolMenuInsertType::After);
+
+	ULevelViewportContext* const LevelViewportContext = InMenu->FindContext<ULevelViewportContext>();
+	if (!LevelViewportContext)
+	{
+		return;
+	}
+
+	const TSharedPtr<::SLevelViewport> LevelViewport = LevelViewportContext->LevelViewport.Pin();
+	if (!LevelViewport)
+	{
+		return;
+	}
+
+	{
+		FToolMenuSection& Section = InMenu->AddSection(
+			"LevelViewportDeferredRendering", LOCTEXT("DeferredRenderingHeader", "Deferred Rendering"), InsertPosition
+		);
+	}
+
+	{
+		FToolMenuSection& Section = InMenu->FindOrAddSection("ViewMode");
+		Section.AddSubMenu(
+			"VisualizeBufferViewMode",
+			LOCTEXT("VisualizeBufferViewModeDisplayName", "Buffer Visualization"),
+			LOCTEXT("BufferVisualizationMenu_ToolTip", "Select a mode for buffer visualization"),
+			FNewMenuDelegate::CreateStatic(&FBufferVisualizationMenuCommands::BuildVisualisationSubMenu),
+			FUIAction(
+				FExecuteAction(),
+				FCanExecuteAction(),
+				FIsActionChecked::CreateLambda(
+					[WeakViewport = LevelViewport.ToWeakPtr()]()
+					{
+						const TSharedPtr<::SLevelViewport> Viewport = WeakViewport.Pin();
+						check(Viewport.IsValid());
+						FLevelEditorViewportClient& ViewportClient = Viewport->GetLevelViewportClient();
+						return ViewportClient.IsViewModeEnabled(VMI_VisualizeBuffer);
+					}
+				)
+			),
+			EUserInterfaceActionType::RadioButton,
+			/* bInOpenSubMenuOnClick = */ false,
+			FSlateIcon(FAppStyle::GetAppStyleSetName(), "EditorViewport.VisualizeBufferMode")
+		);
+	}
+
+	{
+		FToolMenuSection& Section = InMenu->FindOrAddSection("ViewMode");
+		Section.AddSubMenu(
+			"VisualizeNaniteViewMode",
+			LOCTEXT("VisualizeNaniteViewModeDisplayName", "Nanite Visualization"),
+			LOCTEXT("NaniteVisualizationMenu_ToolTip", "Select a mode for Nanite visualization"),
+			FNewMenuDelegate::CreateStatic(&FNaniteVisualizationMenuCommands::BuildVisualisationSubMenu),
+			FUIAction(
+				FExecuteAction(),
+				FCanExecuteAction(),
+				FIsActionChecked::CreateLambda(
+					[WeakViewport = LevelViewport.ToWeakPtr()]()
+					{
+						const TSharedPtr<::SLevelViewport> Viewport = WeakViewport.Pin();
+						check(Viewport.IsValid());
+						FLevelEditorViewportClient& ViewportClient = Viewport->GetLevelViewportClient();
+						return ViewportClient.IsViewModeEnabled(VMI_VisualizeNanite);
+					}
+				)
+			),
+			EUserInterfaceActionType::RadioButton,
+			/* bInOpenSubMenuOnClick = */ false,
+			FSlateIcon(FAppStyle::GetAppStyleSetName(), "EditorViewport.VisualizeNaniteMode")
+		);
+	}
+
+	{
+		FToolMenuSection& Section = InMenu->FindOrAddSection("ViewMode");
+		Section.AddSubMenu(
+			"VisualizeLumenViewMode",
+			LOCTEXT("VisualizeLumenViewModeDisplayName", "Lumen"),
+			LOCTEXT("LumenVisualizationMenu_ToolTip", "Select a mode for Lumen visualization"),
+			FNewMenuDelegate::CreateStatic(&FLumenVisualizationMenuCommands::BuildVisualisationSubMenu),
+			FUIAction(
+				FExecuteAction(),
+				FCanExecuteAction(),
+				FIsActionChecked::CreateLambda(
+					[WeakViewport = LevelViewport.ToWeakPtr()]()
+					{
+						const TSharedPtr<::SLevelViewport> Viewport = WeakViewport.Pin();
+						check(Viewport.IsValid());
+						FLevelEditorViewportClient& ViewportClient = Viewport->GetLevelViewportClient();
+						return ViewportClient.IsViewModeEnabled(VMI_VisualizeLumen);
+					}
+				)
+			),
+			EUserInterfaceActionType::RadioButton,
+			/* bInOpenSubMenuOnClick = */ false,
+			FSlateIcon(FAppStyle::GetAppStyleSetName(), "EditorViewport.VisualizeLumenMode")
+		);
+	}
+
+	if (Substrate::IsSubstrateEnabled())
+	{
+		FToolMenuSection& Section = InMenu->FindOrAddSection("ViewMode");
+		Section.AddSubMenu(
+			"VisualizeSubstrateViewMode",
+			LOCTEXT("VisualizeSubstrateViewModeDisplayName", "Substrate"),
+			LOCTEXT("SubstrateVisualizationMenu_ToolTip", "Select a mode for Substrate visualization"),
+			FNewMenuDelegate::CreateStatic(&FSubstrateVisualizationMenuCommands::BuildVisualisationSubMenu),
+			FUIAction(
+				FExecuteAction(),
+				FCanExecuteAction(),
+				FIsActionChecked::CreateLambda(
+					[WeakViewport = LevelViewport.ToWeakPtr()]()
+					{
+						const TSharedPtr<::SLevelViewport> Viewport = WeakViewport.Pin();
+						check(Viewport.IsValid());
+						FLevelEditorViewportClient& ViewportClient = Viewport->GetLevelViewportClient();
+						return ViewportClient.IsViewModeEnabled(VMI_VisualizeSubstrate);
+					}
+				)
+			),
+			EUserInterfaceActionType::RadioButton,
+			/* bInOpenSubMenuOnClick = */ false,
+			FSlateIcon(FAppStyle::GetAppStyleSetName(), "EditorViewport.VisualizeSubstrateMode")
+		);
+	}
+
+	if (IsGroomEnabled())
+	{
+		FToolMenuSection& Section = InMenu->FindOrAddSection("ViewMode");
+		Section.AddSubMenu(
+			"VisualizeGroomViewMode",
+			LOCTEXT("VisualizeGroomViewModeDisplayName", "Groom"),
+			LOCTEXT("GroomVisualizationMenu_ToolTip", "Select a mode for Groom visualization"),
+			FNewMenuDelegate::CreateStatic(&FGroomVisualizationMenuCommands::BuildVisualisationSubMenu),
+			FUIAction(
+				FExecuteAction(),
+				FCanExecuteAction(),
+				FIsActionChecked::CreateLambda(
+					[WeakViewport = LevelViewport.ToWeakPtr()]()
+					{
+						const TSharedPtr<::SLevelViewport> Viewport = WeakViewport.Pin();
+						check(Viewport.IsValid());
+						FLevelEditorViewportClient& ViewportClient = Viewport->GetLevelViewportClient();
+						return ViewportClient.IsViewModeEnabled(VMI_VisualizeGroom);
+					}
+				)
+			),
+			EUserInterfaceActionType::RadioButton,
+			/* bInOpenSubMenuOnClick = */ false,
+			FSlateIcon(FAppStyle::GetAppStyleSetName(), "EditorViewport.VisualizeGroomMode")
+		);
+	}
+
+	{
+		FToolMenuSection& Section = InMenu->FindOrAddSection("ViewMode");
+		Section.AddSubMenu(
+			"VisualizeVirtualShadowMapViewMode",
+			LOCTEXT("VisualizeVirtualShadowMapViewModeDisplayName", "Virtual Shadow Map"),
+			LOCTEXT(
+				"VirtualShadowMapVisualizationMenu_ToolTip",
+				"Select a mode for virtual shadow map visualization. Select a light component in the world outliner to "
+				"visualize that light."
+			),
+			FNewMenuDelegate::CreateStatic(&FVirtualShadowMapVisualizationMenuCommands::BuildVisualisationSubMenu),
+			FUIAction(
+				FExecuteAction(),
+				FCanExecuteAction(),
+				FIsActionChecked::CreateLambda(
+					[WeakViewport = LevelViewport.ToWeakPtr()]()
+					{
+						const TSharedPtr<::SLevelViewport> Viewport = WeakViewport.Pin();
+						check(Viewport.IsValid());
+						FLevelEditorViewportClient& ViewportClient = Viewport->GetLevelViewportClient();
+						return ViewportClient.IsViewModeEnabled(VMI_VisualizeVirtualShadowMap);
+					}
+				)
+			),
+			EUserInterfaceActionType::RadioButton,
+			/* bInOpenSubMenuOnClick = */ false,
+			FSlateIcon(FAppStyle::GetAppStyleSetName(), "EditorViewport.VisualizeVirtualShadowMapMode")
+		);
+	}
+
+	{
+		auto BuildActorColorationMenu = [WeakViewport = LevelViewport.ToWeakPtr()](UToolMenu* InMenu)
+		{
+			FToolMenuSection& SubMenuSection =
+				InMenu->AddSection("LevelViewportActorColoration", LOCTEXT("ActorColorationHeader", "Actor Coloration"));
+
+			TArray<FActorPrimitiveColorHandler::FPrimitiveColorHandler> PrimitiveColorHandlers;
+			FActorPrimitiveColorHandler::Get().GetRegisteredPrimitiveColorHandlers(PrimitiveColorHandlers);
+
+			for (const FActorPrimitiveColorHandler::FPrimitiveColorHandler& PrimitiveColorHandler : PrimitiveColorHandlers)
+			{
+				if (!PrimitiveColorHandler.bAvailalbleInEditor)
+				{
+					continue;
+				}
+
+				SubMenuSection.AddMenuEntry(
+					NAME_None,
+					PrimitiveColorHandler.HandlerText,
+					PrimitiveColorHandler.HandlerToolTipText,
+					FSlateIcon(),
+					FUIAction(
+						FExecuteAction::CreateLambda(
+							[WeakViewport, PrimitiveColorHandler]()
+							{
+								if (TSharedPtr<::SLevelViewport> Viewport = WeakViewport.Pin())
+								{
+									FLevelEditorViewportClient& ViewportClient = Viewport->GetLevelViewportClient();
+									ViewportClient.ChangeActorColorationVisualizationMode(PrimitiveColorHandler.HandlerName);
+								}
+							}
+						),
+						FCanExecuteAction::CreateLambda(
+							[WeakViewport]()
+							{
+								if (TSharedPtr<::SLevelViewport> Viewport = WeakViewport.Pin())
+								{
+									return true;
+								}
+								return false;
+							}
+						),
+						FGetActionCheckState::CreateLambda(
+							[WeakViewport, PrimitiveColorHandler]()
+							{
+								if (TSharedPtr<::SLevelViewport> Viewport = WeakViewport.Pin())
+								{
+									FLevelEditorViewportClient& ViewportClient = Viewport->GetLevelViewportClient();
+									return ViewportClient.IsActorColorationVisualizationModeSelected(PrimitiveColorHandler.HandlerName) ? ECheckBoxState::Checked : ECheckBoxState::Unchecked;
+								}
+
+								return ECheckBoxState::Unchecked;
+							}
+						)
+					),
+					EUserInterfaceActionType::RadioButton
+				);
+			}
+		};
+
+		FToolMenuSection& Section = InMenu->FindOrAddSection("ViewMode");
+		Section.AddSubMenu(
+			"VisualizeActorColorationViewMode",
+			LOCTEXT("VisualizeActorColorationViewModeDisplayName", "Actor Coloration"),
+			LOCTEXT("ActorColorationVisualizationMenu_ToolTip", "Select a mode for actor coloration visualization."),
+			FNewToolMenuDelegate::CreateLambda(BuildActorColorationMenu),
+			FUIAction(
+				FExecuteAction(),
+				FCanExecuteAction(),
+				FIsActionChecked::CreateLambda(
+					[WeakViewport = LevelViewport.ToWeakPtr()]()
+					{
+						const TSharedPtr<::SLevelViewport> Viewport = WeakViewport.Pin();
+						check(Viewport.IsValid());
+						FLevelEditorViewportClient& ViewportClient = Viewport->GetLevelViewportClient();
+						return ViewportClient.IsViewModeEnabled(VMI_VisualizeActorColoration);
+					}
+				)
+			),
+			EUserInterfaceActionType::RadioButton,
+			/* bInOpenSubMenuOnClick = */ false,
+			FSlateIcon(FAppStyle::GetAppStyleSetName(), "EditorViewport.VisualizeActorColorationMode")
+		);
+	}
+
+	{
+		FToolMenuSection& Section =
+			InMenu->AddSection("LevelViewportLandscape", LOCTEXT("LandscapeHeader", "Landscape"), InsertPosition);
+
+		auto BuildLandscapeLODMenu = [WeakViewport = LevelViewport.ToWeakPtr()](UToolMenu* InMenu)
+		{
+			FToolMenuSection& SubMenuSection =
+				InMenu->AddSection("LevelViewportLandScapeLOD", LOCTEXT("LandscapeLODHeader", "Landscape LOD"));
+
+			auto CreateLandscapeLODAction = [WeakViewport](int32 LODValue)
+			{
+				FUIAction LandscapeLODAction;
+				LandscapeLODAction.ExecuteAction = FExecuteAction::CreateLambda(
+					[WeakViewport, LODValue]()
+					{
+						if (const TSharedPtr<::SLevelViewport> Viewport = WeakViewport.Pin())
+						{
+							UE::LevelEditor::Private::OnLandscapeLODChanged(Viewport->GetLevelViewportClient(), LODValue);
+						}
+					}
+				);
+				LandscapeLODAction.GetActionCheckState = FGetActionCheckState::CreateLambda(
+					[WeakViewport, LODValue]() -> ECheckBoxState
+					{
+						bool bChecked = false;
+						if (const TSharedPtr<::SLevelViewport> Viewport = WeakViewport.Pin())
+						{
+							bChecked = UE::LevelEditor::Private::IsLandscapeLODSettingChecked(
+								Viewport->GetLevelViewportClient(), LODValue
+							);
+						}
+						return bChecked ? ECheckBoxState::Checked : ECheckBoxState::Unchecked;
+					}
+				);
+
+				return LandscapeLODAction;
+			};
+
+			SubMenuSection.AddMenuEntry(
+				"LandscapeLODAuto",
+				LOCTEXT("LandscapeLODAuto", "Auto"),
+				FText(),
+				FSlateIcon(),
+				CreateLandscapeLODAction(-1),
+				EUserInterfaceActionType::RadioButton
+			);
+
+			SubMenuSection.AddSeparator("LandscapeLODSeparator");
+
+			static const FText FormatString = LOCTEXT("LandscapeLODFixed", "Fixed at {0}");
+			for (int32 i = 0; i < 8; ++i)
+			{
+				SubMenuSection.AddMenuEntry(
+					NAME_None,
+					FText::Format(FormatString, FText::AsNumber(i)),
+					FText(),
+					FSlateIcon(),
+					CreateLandscapeLODAction(i),
+					EUserInterfaceActionType::RadioButton
+				);
+			}
+		};
+
+		Section.AddSubMenu(
+			"LandscapeLOD",
+			LOCTEXT("LandscapeLODDisplayName", "LOD"),
+			LOCTEXT("LandscapeLODMenu_ToolTip", "Override Landscape LOD in this viewport"),
+			FNewToolMenuDelegate::CreateLambda(BuildLandscapeLODMenu),
+			/*bInOpenSubMenuOnClick=*/false,
+			FSlateIcon(FAppStyle::GetAppStyleSetName(), "EditorViewport.LOD")
+		);
+	}
+}
+
+void ExtendViewModesSubmenu(FName InViewModesSubmenuName)
+{
+	UToolMenu* const Submenu = UToolMenus::Get()->ExtendMenu(InViewModesSubmenuName);
+
+	Submenu->AddDynamicSection(
+		"LevelEditorViewModesExtensionDynamicSection",
+		FNewToolMenuDelegate::CreateLambda(
+			[](UToolMenu* InDynamicMenu)
+			{
+				PopulateViewModesMenu(InDynamicMenu);
+			}
+		)
+	);
+}
+
+FToolMenuEntry CreateShowFoliageSubmenu()
+{
+	return FToolMenuEntry::InitSubMenu(
+		"ShowFoliage",
+		LOCTEXT("ShowFoliageTypesMenu", "Foliage Types"),
+		LOCTEXT("ShowFoliageTypesMenu_ToolTip", "Show/hide specific foliage types"),
+		FNewToolMenuDelegate::CreateLambda(
+			[](UToolMenu* Submenu)
+			{
+				ULevelViewportContext* const LevelViewportContext = Submenu->FindContext<ULevelViewportContext>();
+				if (!LevelViewportContext)
+				{
+					return;
+				}
+
+				TSharedPtr<::SLevelViewport> Viewport = LevelViewportContext->LevelViewport.Pin();
+				if (!Viewport)
+				{
+					return;
+				}
+
+				if (!Viewport->GetWorld())
+				{
+					return;
+				}
+
+				{
+					FToolMenuSection& Section = Submenu->AddSection("LevelViewportFoliageMeshes");
+					// Map 'Show All' and 'Hide All' commands
+					FUIAction ShowAllFoliage(
+						FExecuteAction::CreateSP(Viewport.ToSharedRef(), &::SLevelViewport::ToggleAllFoliageTypes, true)
+					);
+					FUIAction HideAllFoliage(
+						FExecuteAction::CreateSP(Viewport.ToSharedRef(), &::SLevelViewport::ToggleAllFoliageTypes, false)
+					);
+
+					Section.AddMenuEntry(
+						"ShowAll", LOCTEXT("ShowAllLabel", "Show All"), FText::GetEmpty(), FSlateIcon(), ShowAllFoliage
+					);
+					Section.AddMenuEntry(
+						"HideAll", LOCTEXT("HideAllLabel", "Hide All"), FText::GetEmpty(), FSlateIcon(), HideAllFoliage
+					);
+				}
+
+				// Gather all foliage types used in this world and group them by sub-levels
+				auto AllFoliageMap =
+					UE::LevelEditor::Private::GroupFoliageByOuter(GEditor->GetFoliageTypesInWorld(Viewport->GetWorld()));
+
+				for (auto& FoliagePair : AllFoliageMap)
+				{
+					// Name foliage group by an outer sub-level name, or empty if foliage type is an asset
+					FText EntryName =
+						(FoliagePair.Key == NAME_None ? FText::GetEmpty()
+													  : FText::FromName(FPackageName::GetShortFName(FoliagePair.Key)));
+					FToolMenuSection& Section = Submenu->AddSection(NAME_None, EntryName);
+
+					TArray<UFoliageType*>& FoliageList = FoliagePair.Value;
+					for (UFoliageType* FoliageType : FoliageList)
+					{
+						FName MeshName = FoliageType->GetDisplayFName();
+						TWeakObjectPtr<UFoliageType> FoliageTypePtr = FoliageType;
+
+						FUIAction Action(
+							FExecuteAction::CreateSP(
+								Viewport.ToSharedRef(), &::SLevelViewport::ToggleShowFoliageType, FoliageTypePtr
+							),
+							FCanExecuteAction(),
+							FIsActionChecked::CreateSP(Viewport.ToSharedRef(), &::SLevelViewport::IsFoliageTypeVisible, FoliageTypePtr)
+						);
+
+						Section.AddMenuEntry(
+							NAME_None,
+							FText::FromName(MeshName),
+							FText::GetEmpty(),
+							FSlateIcon(),
+							Action,
+							EUserInterfaceActionType::ToggleButton
+						);
+					}
+				}
+			}
+		),
+		false,
+		FSlateIcon(FAppStyle::Get().GetStyleSetName(), "ShowFlagsMenu.SubMenu.FoliageTypes")
+	);
+}
+
+FToolMenuEntry CreateShowHLODsSubmenu()
+{
+	// This is a dynamic entry so we can skip adding the submenu if the context
+	// indicates that the viewport's world isn't partitioned.
+	return FToolMenuEntry::InitDynamicEntry(
+		"ShowHLODsDynamic",
+		FNewToolMenuSectionDelegate::CreateLambda(
+			[](FToolMenuSection& InDynamicSection) -> void
+			{
+				ULevelViewportContext* const LevelViewportContext = InDynamicSection.FindContext<ULevelViewportContext>();
+				if (!LevelViewportContext)
+				{
+					return;
+				}
+
+				TSharedPtr<::SLevelViewport> Viewport = LevelViewportContext->LevelViewport.Pin();
+				if (!Viewport)
+				{
+					return;
+				}
+
+				UWorld* World = Viewport->GetWorld();
+				if (!World)
+				{
+					return;
+				}
+
+				// Only add this submenu for partitioned worlds.
+				if (!World->IsPartitionedWorld())
+				{
+					return;
+				}
+
+				InDynamicSection.AddSubMenu(
+					"ShowHLODsMenu",
+					LOCTEXT("ShowHLODsMenu", "HLODs"),
+					LOCTEXT("ShowHLODsMenu_ToolTip", "Settings for HLODs in editor"),
+					FNewToolMenuDelegate::CreateLambda(
+						[](UToolMenu* Submenu)
+						{
+							ULevelViewportContext* const LevelViewportContext =
+								Submenu->FindContext<ULevelViewportContext>();
+							if (!LevelViewportContext)
+							{
+								return;
+							}
+
+							TSharedPtr<::SLevelViewport> Viewport = LevelViewportContext->LevelViewport.Pin();
+							if (!Viewport)
+							{
+								return;
+							}
+
+							UWorld* World = Viewport->GetWorld();
+							UWorldPartition* WorldPartition = World ? World->GetWorldPartition() : nullptr;
+							if (!WorldPartition)
+							{
+								return;
+							}
+
+							IWorldPartitionEditorModule* WorldPartitionEditorModule =
+								FModuleManager::GetModulePtr<IWorldPartitionEditorModule>("WorldPartitionEditor");
+							if (!WorldPartitionEditorModule)
+							{
+								return;
+							}
+
+							FText HLODInEditorDisallowedReason;
+							const bool bHLODInEditorAllowed =
+								WorldPartitionEditorModule->IsHLODInEditorAllowed(World, &HLODInEditorDisallowedReason);
+
+							// Show HLODs
+							{
+								FToolUIAction UIAction;
+								UIAction.ExecuteAction = FToolMenuExecuteAction::CreateLambda(
+									[WorldPartitionEditorModule](const FToolMenuContext& InContext)
+									{
+										WorldPartitionEditorModule->SetShowHLODsInEditor(
+											!WorldPartitionEditorModule->GetShowHLODsInEditor()
+										);
+									}
+								);
+								UIAction.CanExecuteAction = FToolMenuCanExecuteAction::CreateLambda(
+									[bHLODInEditorAllowed](const FToolMenuContext& InContext)
+									{
+										return bHLODInEditorAllowed;
+									}
+								);
+								UIAction.GetActionCheckState = FToolMenuGetActionCheckState::CreateLambda(
+									[WorldPartitionEditorModule](const FToolMenuContext& InContext)
+									{
+										return WorldPartitionEditorModule->GetShowHLODsInEditor()
+												 ? ECheckBoxState::Checked
+												 : ECheckBoxState::Unchecked;
+									}
+								);
+								FToolMenuEntry MenuEntry = FToolMenuEntry::InitMenuEntry(
+									"ShowHLODs",
+									LOCTEXT("ShowHLODs", "Show HLODs"),
+									bHLODInEditorAllowed ? LOCTEXT("ShowHLODsToolTip", "Show/Hide HLODs")
+														 : HLODInEditorDisallowedReason,
+									FSlateIcon(),
+									UIAction,
+									EUserInterfaceActionType::ToggleButton
+								);
+								Submenu->AddMenuEntry(NAME_None, MenuEntry);
+							}
+
+							// Show HLODs Over Loaded Regions
+							{
+								FToolUIAction UIAction;
+								UIAction.ExecuteAction = FToolMenuExecuteAction::CreateLambda(
+									[WorldPartitionEditorModule](const FToolMenuContext& InContext)
+									{
+										WorldPartitionEditorModule->SetShowHLODsOverLoadedRegions(
+											!WorldPartitionEditorModule->GetShowHLODsOverLoadedRegions()
+										);
+									}
+								);
+								UIAction.CanExecuteAction = FToolMenuCanExecuteAction::CreateLambda(
+									[bHLODInEditorAllowed](const FToolMenuContext& InContext)
+									{
+										return bHLODInEditorAllowed;
+									}
+								);
+								UIAction.GetActionCheckState = FToolMenuGetActionCheckState::CreateLambda(
+									[WorldPartitionEditorModule](const FToolMenuContext& InContext)
+									{
+										return WorldPartitionEditorModule->GetShowHLODsOverLoadedRegions()
+												 ? ECheckBoxState::Checked
+												 : ECheckBoxState::Unchecked;
+									}
+								);
+								FToolMenuEntry ShowHLODsEntry = FToolMenuEntry::InitMenuEntry(
+									"ShowHLODsOverLoadedRegions",
+									LOCTEXT("ShowHLODsOverLoadedRegions", "Show HLODs Over Loaded Regions"),
+									bHLODInEditorAllowed
+										? LOCTEXT("ShowHLODsOverLoadedRegions_ToolTip", "Show/Hide HLODs over loaded actors or regions")
+										: HLODInEditorDisallowedReason,
+									FSlateIcon(),
+									UIAction,
+									EUserInterfaceActionType::ToggleButton
+								);
+								Submenu->AddMenuEntry(NAME_None, ShowHLODsEntry);
+							}
+
+							// Min/Max Draw Distance
+							{
+								const double MinDrawDistanceMinValue = 0;
+								const double MinDrawDistanceMaxValue = 102400;
+
+								const double MaxDrawDistanceMinValue = 0;
+								const double MaxDrawDistanceMaxValue = 1638400;
+
+								// double SLevelViewportToolBar::OnGetHLODInEditorMinDrawDistanceValue() const
+								auto OnGetHLODInEditorMinDrawDistanceValue = []() -> double
+								{
+									IWorldPartitionEditorModule* WorldPartitionEditorModule =
+										FModuleManager::GetModulePtr<IWorldPartitionEditorModule>("WorldPartitionEditor");
+									return WorldPartitionEditorModule
+											 ? WorldPartitionEditorModule->GetHLODInEditorMinDrawDistance()
+											 : 0;
+								};
+
+								// void SLevelViewportToolBar::OnHLODInEditorMinDrawDistanceValueChanged(double NewValue) const
+								auto OnHLODInEditorMinDrawDistanceValueChanged = [](double NewValue) -> void
+								{
+									IWorldPartitionEditorModule* WorldPartitionEditorModule =
+										FModuleManager::GetModulePtr<IWorldPartitionEditorModule>("WorldPartitionEditor");
+									if (WorldPartitionEditorModule)
+									{
+										WorldPartitionEditorModule->SetHLODInEditorMinDrawDistance(NewValue);
+										GEditor->RedrawLevelEditingViewports(true);
+									}
+								};
+
+								TSharedRef<SSpinBox<double>> MinDrawDistanceSpinBox =
+									SNew(SSpinBox<double>)
+										.MinValue(MinDrawDistanceMinValue)
+										.MaxValue(MinDrawDistanceMaxValue)
+										.IsEnabled(bHLODInEditorAllowed)
+										.Value_Lambda(OnGetHLODInEditorMinDrawDistanceValue)
+										.OnValueChanged_Lambda(OnHLODInEditorMinDrawDistanceValueChanged)
+										.ToolTipText(
+											bHLODInEditorAllowed
+												? LOCTEXT(
+													  "HLODsInEditor_MinDrawDistance_Tooltip",
+													  "Sets the minimum distance at which HLOD will be rendered"
+												  )
+												: HLODInEditorDisallowedReason
+										)
+										.OnBeginSliderMovement_Lambda(
+											[]()
+											{
+												// Disable Slate throttling during slider drag to ensure immediate updates while moving the slider.
+												FSlateThrottleManager::Get().DisableThrottle(true);
+											}
+										)
+										.OnEndSliderMovement_Lambda(
+											[](float)
+											{
+												FSlateThrottleManager::Get().DisableThrottle(false);
+											}
+										);
+
+								// double SLevelViewportToolBar::OnGetHLODInEditorMaxDrawDistanceValue() const
+								auto OnGetHLODInEditorMaxDrawDistanceValue = []() -> double
+								{
+									IWorldPartitionEditorModule* WorldPartitionEditorModule =
+										FModuleManager::GetModulePtr<IWorldPartitionEditorModule>("WorldPartitionEditor");
+									return WorldPartitionEditorModule
+										? WorldPartitionEditorModule->GetHLODInEditorMaxDrawDistance()
+										: 0;
+								};
+
+								// void SLevelViewportToolBar::OnHLODInEditorMaxDrawDistanceValueChanged(double NewValue) const
+								auto OnHLODInEditorMaxDrawDistanceValueChanged = [](double NewValue) -> void
+								{
+									IWorldPartitionEditorModule* WorldPartitionEditorModule =
+										FModuleManager::GetModulePtr<IWorldPartitionEditorModule>("WorldPartitionEditor");
+									if (WorldPartitionEditorModule)
+									{
+										WorldPartitionEditorModule->SetHLODInEditorMaxDrawDistance(NewValue);
+										GEditor->RedrawLevelEditingViewports(true);
+									}
+								};
+
+								TSharedRef<SSpinBox<double>> MaxDrawDistanceSpinBox =
+									SNew(SSpinBox<double>)
+										.MinValue(MaxDrawDistanceMinValue)
+										.MaxValue(MaxDrawDistanceMaxValue)
+										.IsEnabled(bHLODInEditorAllowed)
+										.Value_Lambda(OnGetHLODInEditorMaxDrawDistanceValue)
+										.OnValueChanged_Lambda(OnHLODInEditorMaxDrawDistanceValueChanged)
+										.ToolTipText(
+											bHLODInEditorAllowed
+												? LOCTEXT(
+													  "HLODsInEditor_MaxDrawDistance_Tooltip",
+													  "Sets the maximum distance at which HLODs will be rendered (0.0 means infinite)"
+												  )
+												: HLODInEditorDisallowedReason
+										)
+										.OnBeginSliderMovement_Lambda(
+											[]()
+											{
+												// Disable Slate throttling during slider drag to ensure immediate updates while moving the slider.
+												FSlateThrottleManager::Get().DisableThrottle(true);
+											}
+										)
+										.OnEndSliderMovement_Lambda(
+											[](float)
+											{
+												FSlateThrottleManager::Get().DisableThrottle(false);
+											}
+										);
+
+								auto CreateDrawDistanceWidget = [](TSharedRef<SSpinBox<double>> InSpinBoxWidget)
+								{
+									// clang-format off
+									return SNew(SBox)
+										.HAlign(HAlign_Right)
+										[
+											SNew(SBox)
+										  .Padding(FMargin(0.0f, 0.0f, 0.0f, 0.0f))
+										  .WidthOverride(100.0f)
+											[
+												SNew(SBorder)
+												.BorderImage(FAppStyle::Get().GetBrush("Menu.WidgetBorder"))
+												.Padding(FMargin(1.0f))
+												[
+													InSpinBoxWidget
+												]
+											]
+										];
+									// clang-format on
+								};
+
+								FToolMenuEntry MinDrawDistanceMenuEntry = FToolMenuEntry::InitWidget(
+									"Min Draw Distance",
+									CreateDrawDistanceWidget(MinDrawDistanceSpinBox),
+									LOCTEXT("MinDrawDistance", "Min Draw Distance")
+								);
+								Submenu->AddMenuEntry(NAME_None, MinDrawDistanceMenuEntry);
+
+								FToolMenuEntry MaxDrawDistanceMenuEntry = FToolMenuEntry::InitWidget(
+									"Max Draw Distance",
+									CreateDrawDistanceWidget(MaxDrawDistanceSpinBox),
+									LOCTEXT("MaxDrawDistance", "Max Draw Distance")
+								);
+								Submenu->AddMenuEntry(NAME_None, MaxDrawDistanceMenuEntry);
+							}
+						}
+					),
+					false,
+					FSlateIcon(FAppStyle::Get().GetStyleSetName(), "ShowFlagsMenu.SubMenu.HLODs")
+				);
+			}
+		)
+	);
+}
+
+FToolMenuEntry CreateShowLayersSubmenu()
+{
+	// This is a dynamic entry so we can skip adding the submenu if the context
+	// indicates that the viewport's world is partitioned.
+	return FToolMenuEntry::InitDynamicEntry(
+		"ShowHLODsDynamic",
+		FNewToolMenuSectionDelegate::CreateLambda(
+			[](FToolMenuSection& InDynamicSection) -> void
+			{
+				ULevelViewportContext* const LevelViewportContext = InDynamicSection.FindContext<ULevelViewportContext>();
+				if (!LevelViewportContext)
+				{
+					return;
+				}
+
+				TSharedPtr<::SLevelViewport> Viewport = LevelViewportContext->LevelViewport.Pin();
+				if (!Viewport)
+				{
+					return;
+				}
+
+				UWorld* World = Viewport->GetWorld();
+				if (!World)
+				{
+					return;
+				}
+
+				// Only add this submenu for non-partitioned worlds.
+				if (World->IsPartitionedWorld())
+				{
+					return;
+				}
+
+				InDynamicSection.AddSubMenu(
+					"ShowLayers",
+					LOCTEXT("ShowLayersMenu", "Layers"),
+					LOCTEXT("ShowLayersMenu_ToolTip", "Show layers flags"),
+					FNewToolMenuDelegate::CreateStatic(
+						&UE::LevelEditor::Private::PopulateShowLayersSubmenu, Viewport.ToWeakPtr()
+					),
+					false,
+					FSlateIcon(FAppStyle::Get().GetStyleSetName(), "ShowFlagsMenu.SubMenu.Layers")
+				);
+			}
+		)
+	);
+}
+
+FToolMenuEntry CreateShowSpritesSubmenu()
+{
+	const FLevelViewportCommands& Actions = FLevelViewportCommands::Get();
+	TArray<FLevelViewportCommands::FShowMenuCommand> ShowSpritesMenu;
+
+	// 'Show All' and 'Hide All' buttons
+	ShowSpritesMenu.Add(
+		FLevelViewportCommands::FShowMenuCommand(Actions.ShowAllSprites, LOCTEXT("ShowAllLabel", "Show All"))
+	);
+	ShowSpritesMenu.Add(
+		FLevelViewportCommands::FShowMenuCommand(Actions.HideAllSprites, LOCTEXT("HideAllLabel", "Hide All"))
+	);
+
+	// Get each show flag command and put them in their corresponding groups
+	ShowSpritesMenu += Actions.ShowSpriteCommands;
+
+	return FToolMenuEntry::InitSubMenu(
+		"ShowSprites",
+		LOCTEXT("ShowSpritesMenu", "Sprites"),
+		LOCTEXT("ShowSpritesMenu_ToolTip", "Show sprites flags"),
+		FNewToolMenuDelegate::CreateStatic(&UE::LevelEditor::Private::PopulateMenuWithCommands, ShowSpritesMenu, 2),
+		false,
+		FSlateIcon(FAppStyle::Get().GetStyleSetName(), "ShowFlagsMenu.SubMenu.Sprites")
+	);
+}
+
+FToolMenuEntry CreateShowVolumesSubmenu()
+{
+	const FLevelViewportCommands& Actions = FLevelViewportCommands::Get();
+	TArray<FLevelViewportCommands::FShowMenuCommand> ShowVolumesMenu;
+
+	// 'Show All' and 'Hide All' buttons
+	ShowVolumesMenu.Add(
+		FLevelViewportCommands::FShowMenuCommand(Actions.ShowAllVolumes, LOCTEXT("ShowAllLabel", "Show All"))
+	);
+	ShowVolumesMenu.Add(
+		FLevelViewportCommands::FShowMenuCommand(Actions.HideAllVolumes, LOCTEXT("HideAllLabel", "Hide All"))
+	);
+
+	// Get each show flag command and put them in their corresponding groups
+	ShowVolumesMenu += Actions.ShowVolumeCommands;
+
+	return FToolMenuEntry::InitSubMenu(
+		"ShowVolumes",
+		LOCTEXT("ShowVolumesMenu", "Volumes"),
+		LOCTEXT("ShowVolumesMenu_ToolTip", "Show volumes flags"),
+		FNewToolMenuDelegate::CreateStatic(&UE::LevelEditor::Private::PopulateMenuWithCommands, ShowVolumesMenu, 2),
+		false,
+		FSlateIcon(FAppStyle::Get().GetStyleSetName(), "ShowFlagsMenu.SubMenu.Volumes")
+	);
+}
+
+#if STATS
+FToolMenuEntry CreateShowStatsSubmenu(bool bInAddToggleStatsCheckbox, TAttribute<FText> InLabelOverride)
+{
+	return FToolMenuEntry::InitDynamicEntry(
+		"DynamicShowStatsEntry",
+		FNewToolMenuSectionDelegate::CreateLambda(
+			[bInAddToggleStatsCheckbox, InLabelOverride](FToolMenuSection& InDynamicSection)
+			{
+				TWeakPtr<::SLevelViewport> WeakLevelViewport;
+				if (ULevelViewportContext* const LevelViewportContext =
+						InDynamicSection.FindContext<ULevelViewportContext>())
+				{
+					WeakLevelViewport = LevelViewportContext->LevelViewport;
+				}
+
+				FToolUIActionChoice CommandAction;
+				if (bInAddToggleStatsCheckbox)
+				{
+					if (TSharedPtr<::SLevelViewport> Viewport = WeakLevelViewport.Pin())
+					{
+						CommandAction =
+							FToolUIActionChoice(FEditorViewportCommands::Get().ToggleStats, *Viewport->GetCommandList());
+					}
+				}
+
+				const TAttribute<FText> Label = InLabelOverride.IsSet() ? InLabelOverride
+																		: LOCTEXT("ShowStatsMenu", "Stat");
+
+				InDynamicSection.AddSubMenu(
+					"ShowStatsMenu",
+					Label,
+					LOCTEXT("ShowStatsMenu_ToolTip", "Show Stat commands"),
+					FNewToolMenuDelegate::CreateLambda(
+						[](UToolMenu* InMenu) -> void
+						{
+							// Hide All
+							{
+								FToolMenuSection& UnnamedSection = InMenu->AddSection(NAME_None);
+								UnnamedSection.AddMenuEntry(
+									FLevelViewportCommands::Get().HideAllStats,
+									LOCTEXT("HideAllLabel", "Hide All"),
+									FText(),
+									FSlateIconFinder::FindIcon("Cross")
+								);
+							}
+
+							// Common Stats
+							// The list of Stat Commands we want to show right below the Hide All Stats
+							TArray<FName> CommonStatCommandNames = { TEXT("STAT_FPS"),
+																	 TEXT("STAT_UNIT"),
+																	 TEXT("STATGROUP_Memory"),
+																	 TEXT("STATGROUP_RHI"),
+																	 TEXT("STATGROUP_SceneRendering") };
+
+							FToolMenuSection& CommonStatsSection =
+								InMenu->AddSection("CommonStats", LOCTEXT("CommonStatsLabel", "Common Stats"));
+
+							FToolMenuSection& Section = InMenu->AddSection("Section");
+
+							// Separate out stats into two list, those with and without submenus
+							TArray<FLevelViewportCommands::FShowMenuCommand> SingleStatCommands;
+							TMap<FString, TArray<FLevelViewportCommands::FShowMenuCommand>> SubbedStatCommands;
+							for (auto StatCatIt = FLevelViewportCommands::Get().ShowStatCatCommands.CreateConstIterator();
+								 StatCatIt;
+								 ++StatCatIt)
+							{
+								const TArray<FLevelViewportCommands::FShowMenuCommand>& ShowStatCommands =
+									StatCatIt.Value();
+								const FString& CategoryName = StatCatIt.Key();
+
+								// If no category is specified, or there's only one category, don't use submenus
+								FString NoCategory = FStatConstants::NAME_NoCategory.ToString();
+								NoCategory.RemoveFromStart(TEXT("STATCAT_"));
+								if (CategoryName == NoCategory
+									|| FLevelViewportCommands::Get().ShowStatCatCommands.Num() == 1)
+								{
+									for (int32 StatIndex = 0; StatIndex < ShowStatCommands.Num(); ++StatIndex)
+									{
+										const FLevelViewportCommands::FShowMenuCommand& StatCommand =
+											ShowStatCommands[StatIndex];
+										SingleStatCommands.Add(StatCommand);
+									}
+								}
+								else
+								{
+									SubbedStatCommands.Add(CategoryName, ShowStatCommands);
+								}
+
+								// Search for commands to be added to the Common Stats Section
+								for (const FLevelViewportCommands::FShowMenuCommand& ShowMenuCommand : ShowStatCommands)
+								{
+									if (TSharedPtr<FUICommandInfo> Command = ShowMenuCommand.ShowMenuItem)
+									{
+										for (FName StatCommandName : CommonStatCommandNames)
+										{
+											if (StatCommandName == Command->GetCommandName())
+											{
+												CommonStatsSection.AddMenuEntry(Command);
+											}
+										}
+									}
+								}
+							}
+
+							// Sort Common Stats section entries alphabetically
+							CommonStatsSection.Blocks.Sort(
+								[](const FToolMenuEntry& A, const FToolMenuEntry& B)
+								{
+									return A.Label.Get().ToLower().ToString() < B.Label.Get().ToLower().ToString();
+								}
+							);
+
+							CommonStatsSection.AddSeparator("CommonStatsSeparator");
+
+							// First add all the stats that don't have a sub menu
+							for (auto StatCatIt = SingleStatCommands.CreateConstIterator(); StatCatIt; ++StatCatIt)
+							{
+								const FLevelViewportCommands::FShowMenuCommand& StatCommand = *StatCatIt;
+								Section.AddMenuEntry(NAME_None, StatCommand.ShowMenuItem, StatCommand.LabelOverride);
+							}
+
+							// Now add all the stats that have sub menus
+							for (auto StatCatIt = SubbedStatCommands.CreateConstIterator(); StatCatIt; ++StatCatIt)
+							{
+								const TArray<FLevelViewportCommands::FShowMenuCommand>& StatCommands = StatCatIt.Value();
+								const FText CategoryName = FText::FromString(StatCatIt.Key());
+
+								FFormatNamedArguments Args;
+								Args.Add(TEXT("StatCat"), CategoryName);
+								const FText CategoryDescription =
+									FText::Format(NSLOCTEXT("UICommands", "StatShowCatName", "Show {StatCat} stats"), Args);
+
+								Section.AddSubMenu(
+									NAME_None,
+									CategoryName,
+									CategoryDescription,
+									FNewToolMenuDelegate::CreateStatic(
+										&UE::LevelEditor::Private::PopulateMenuWithCommands, StatCommands, 0
+									)
+								);
+							}
+						}
+					),
+					CommandAction,
+					bInAddToggleStatsCheckbox ? EUserInterfaceActionType::ToggleButton : EUserInterfaceActionType::Button,
+					false,
+					FSlateIcon(FAppStyle::Get().GetStyleSetName(), "EditorViewport.SubMenu.Stats")
+				);
+			}
+		)
+	);
+}
+#endif
+
+FToolMenuEntry CreateShowSubmenu()
+{
+	FToolMenuEntry Entry = FToolMenuEntry::InitSubMenu(
+		"Show",
+		LOCTEXT("ShowSubmenuLabel", "Show"),
+		LOCTEXT("ShowSubmenuTooltip", "Show flags related to the current viewport"),
+		FNewToolMenuDelegate::CreateLambda(
+			[](UToolMenu* InMenu) -> void
+			{
+				{
+					FToolMenuSection& UnnamedSection = InMenu->FindOrAddSection(NAME_None);
+
+					UnnamedSection.AddMenuEntry(FLevelViewportCommands::Get().UseDefaultShowFlags);
+
+					UnnamedSection.AddSeparator("ViewportStatsSeparator");
+
+#if STATS
+					// Override the label of the stats submenu for the new viewport toolbar.
+					UnnamedSection.AddEntry(
+						UE::LevelEditor::CreateShowStatsSubmenu(true, LOCTEXT("ViewportStatsLabel", "Viewport Stats"))
+					);
+#endif
+				}
+
+				// Starting from commonly used flags
+				UE::UnrealEd::AddDefaultShowFlags(InMenu);
+
+				// Add Level Editor specific entries to the All Show Flags Section
+				{
+					FToolMenuSection& AllShowFlagsSection =
+						InMenu->FindOrAddSection("AllShowFlags", LOCTEXT("AllShowFlagsLabel", "All Show Flags"));
+
+					// Show Foliage
+					{
+						FToolMenuEntry ShowFoliageSubmenu = CreateShowFoliageSubmenu();
+						ShowFoliageSubmenu.Label = LOCTEXT("ShowFoliageLabel", "Foliage");
+						ShowFoliageSubmenu.InsertPosition.Position = EToolMenuInsertType::First;
+						AllShowFlagsSection.AddEntry(ShowFoliageSubmenu);
+					}
+
+					// Show HLODs
+					{
+						FToolMenuEntry ShowHLODSubmenu = CreateShowHLODsSubmenu();
+						ShowHLODSubmenu.InsertPosition.Position = EToolMenuInsertType::First;
+						AllShowFlagsSection.AddEntry(ShowHLODSubmenu);
+					}
+
+					// Show Layers
+					{
+						FToolMenuEntry ShowLayersSubmenu = CreateShowLayersSubmenu();
+						ShowLayersSubmenu.InsertPosition.Position = EToolMenuInsertType::First;
+						AllShowFlagsSection.AddEntry(ShowLayersSubmenu);
+					}
+
+					// Show Sprites
+					{
+						FToolMenuEntry ShowSpriteSubmenu = CreateShowSpritesSubmenu();
+						ShowSpriteSubmenu.InsertPosition.Position = EToolMenuInsertType::First;
+						AllShowFlagsSection.AddEntry(ShowSpriteSubmenu);
+					}
+
+					// Show Volumes
+					{
+						FToolMenuEntry ShowVolumesSubmenu = CreateShowVolumesSubmenu();
+						ShowVolumesSubmenu.InsertPosition.Position = EToolMenuInsertType::First;
+						AllShowFlagsSection.AddEntry(ShowVolumesSubmenu);
+					}
+				}
+
+				// Adds show flags sections for backward compatibility with the old viewport toolbar.
+				//  If your entries end up in this section, you should move it to the new "CommonShowFlags" section instead.
+				InMenu->FindOrAddSection(
+					"ShowFlagsMenuSectionCommon",
+					LOCTEXT("ShowFlagsMenuSectionCommonLabel", "Common Show Flags (Deprecated section)")
+				);
+
+				// If your entries end up in these sections, you should move them to the above "AllShowFlags" section instead.
+				InMenu->FindOrAddSection(
+					"LevelViewportShowFlags", LOCTEXT("LevelViewportShowFlagsLabel", "All Show Flags (Deprecated section)")
+				);
+				InMenu->FindOrAddSection(
+					"LevelViewportEditorShow", LOCTEXT("LevelViewportEditorShowLabel", "Editor (Deprecated section)")
+				);
+			}
+		)
+	);
+	Entry.Icon = FSlateIcon(FAppStyle::GetAppStyleSetName(), "Level.VisibleHighlightIcon16x");
+	Entry.InsertPosition.Position = EToolMenuInsertType::Last;
+	Entry.ToolBarData.LabelOverride = FText();
+	Entry.ToolBarData.ResizeParams.ClippingPriority = 800;
+	return Entry;
+}
+
+FToolMenuEntry CreateFeatureLevelPreviewSubmenu()
+{
+	return FToolMenuEntry::InitSubMenu(
+		"FeatureLevelPreview",
+		NSLOCTEXT("LevelToolBarViewMenu", "PreviewPlatformSubMenu", "Preview Platform"),
+		NSLOCTEXT("LevelToolBarViewMenu", "PreviewPlatformSubMenu_ToolTip", "Sets the preview platform used by the main editor"),
+		FNewToolMenuDelegate::CreateLambda(
+			[](UToolMenu* InMenu) -> void
+			{
+				FToolMenuSection& Section =
+					InMenu->AddSection("EditorPreviewMode", LOCTEXT("EditorPreviewModePlatforms", "Preview Platforms"));
+				
+				if (FLevelEditorCommands::Get().DisablePlatformPreview.IsValid())
+				{
+					Section.AddMenuEntry(FLevelEditorCommands::Get().DisablePlatformPreview);
+				}
+				Section.AddSeparator("DisablePlatformPreviewSeparator");
+
+				for (auto Iter = FLevelEditorCommands::Get().PlatformToPreviewPlatformOverrides.CreateConstIterator(); Iter; ++Iter)
+				{
+					FName PlatformName = Iter.Key();
+					const TArray<FLevelEditorCommands::PreviewPlatformCommand>& CommandList = Iter.Value();
+					const TArray<FLevelEditorCommands::PreviewPlatformCommand>* CommandListJson = FLevelEditorCommands::Get().PlatformToPreviewJsonPlatformOverrides.Find(PlatformName);
+
+					Section.AddSubMenu(
+						FName(PlatformName), 
+						FText::FromString(PlatformName.ToString()), 
+						FText(),
+						FNewToolMenuDelegate::CreateLambda(
+							[CommandList, CommandListJson](UToolMenu* InSubMenu)
+							{
+								for (const FLevelEditorCommands::PreviewPlatformCommand& Command : CommandList)
+								{
+									FToolMenuSection& Section = InSubMenu->FindOrAddSection(Command.SectionName, FText::Format(LOCTEXT("PreviewJson", "{0}"), FText::FromName(Command.SectionName)));
+									Section.AddMenuEntry(Command.CommandInfo);
+								}
+
+								if (CommandListJson != nullptr)
+								{
+									FToolMenuSection& SectionJson = InSubMenu->FindOrAddSection("PreviewWithJson", LOCTEXT("PreviewWithJsonLabel", "Preview With Json"));
+									TMap<FName, TArray<TSharedPtr<FUICommandInfo>>> SectionNameToCommandList;
+									for (const FLevelEditorCommands::PreviewPlatformCommand& PreviewJsonPlatform : *CommandListJson)
+									{
+										if (PreviewJsonPlatform.bIsGeneratingJsonCommand)
+										{
+											SectionJson.AddMenuEntry(PreviewJsonPlatform.CommandInfo);
+										}
+										else
+										{
+											SectionNameToCommandList.FindOrAdd(PreviewJsonPlatform.SectionName).Add(PreviewJsonPlatform.CommandInfo);
+										}
+									}
+
+									for (auto Iter = SectionNameToCommandList.CreateConstIterator(); Iter; ++Iter)
+									{
+										FName SectionName = Iter.Key();
+										const TArray<TSharedPtr<FUICommandInfo>>& CommandListValue = Iter.Value();
+										SectionJson.AddSubMenu(
+											SectionName,
+											FText::Format(LOCTEXT("PreviewJsonLabel", "Preview {0}"), FText::FromName(SectionName)),
+											FText::Format(LOCTEXT("PreviewJsonTooltip", "Preview {0}"), FText::FromName(SectionName)),
+											FNewToolMenuDelegate::CreateLambda(
+												[CommandListValue](UToolMenu* InSubMenu)
+												{
+													FToolMenuSection& Section = InSubMenu->AddSection(NAME_None);
+													for (const TSharedPtr<FUICommandInfo>& Command : CommandListValue)
+													{
+														Section.AddMenuEntry(Command);
+													}
+												})
+										);
+									}
+								}
+							})
+					);
+				}
+			}
+		)
+	);
+}
+
+FToolMenuEntry CreateMaterialQualityLevelSubmenu()
+{
+	return FToolMenuEntry::InitSubMenu(
+		"MaterialQualityLevel",
+		NSLOCTEXT("LevelToolBarViewMenu", "MaterialQualityLevelSubMenu", "Material Quality Level"),
+		NSLOCTEXT(
+			"LevelToolBarViewMenu",
+			"MaterialQualityLevelSubMenu_ToolTip",
+			"Sets the value of the CVar \"r.MaterialQualityLevel\" (low=0, high=1, medium=2, Epic=3). This affects "
+			"materials via the QualitySwitch material expression."
+		),
+		FNewToolMenuDelegate::CreateLambda(
+			[](UToolMenu* InMenu) -> void
+			{
+				FToolMenuSection& Section = InMenu->AddSection(
+					"LevelEditorMaterialQualityLevel",
+					NSLOCTEXT("LevelToolBarViewMenu", "MaterialQualityLevelHeading", "Material Quality Level")
+				);
+				Section.AddMenuEntry(FLevelEditorCommands::Get().MaterialQualityLevel_Low);
+				Section.AddMenuEntry(FLevelEditorCommands::Get().MaterialQualityLevel_Medium);
+				Section.AddMenuEntry(FLevelEditorCommands::Get().MaterialQualityLevel_High);
+				Section.AddMenuEntry(FLevelEditorCommands::Get().MaterialQualityLevel_Epic);
+			}
+		)
+	);
+}
+
+FToolMenuEntry CreatePerformanceAndScalabilitySubmenu()
+{
+	FToolMenuEntry Entry = FToolMenuEntry::InitSubMenu(
+		"PerformanceAndScalability",
+		LOCTEXT("PerformanceAndScalabilityLabel", "Performance & Scalability"),
+		LOCTEXT("PerformanceAndScalabilityTooltip", "Performance and scalability tools tied to this viewport."),
+		FNewToolMenuDelegate::CreateLambda(
+			[](UToolMenu* Submenu) -> void
+			{
+				{
+					FToolMenuSection& UnnamedSection = Submenu->FindOrAddSection(NAME_None);
+
+					UnnamedSection.AddEntry(UE::UnrealEd::CreateToggleRealtimeEntry());
+				}
+
+				{
+					FToolMenuSection& PerformanceAndScalabilitySection = Submenu->FindOrAddSection(
+						"PerformanceAndScalability",
+						LOCTEXT("PerformanceAndScalabilitySectionLabel", "Performance & Scalability")
+					);
+
+					PerformanceAndScalabilitySection.AddEntry(CreateFeatureLevelPreviewSubmenu());
+
+					PerformanceAndScalabilitySection.AddSeparator("PerformanceAndScalabilitySettings");
+
+					PerformanceAndScalabilitySection.AddEntry(UE::UnrealEd::CreateScalabilitySubmenu());
+
+					PerformanceAndScalabilitySection.AddEntry(CreateMaterialQualityLevelSubmenu());
+
+					PerformanceAndScalabilitySection.AddSubMenu(
+						"ScreenPercentageSubMenu",
+						LOCTEXT("ScreenPercentageSubMenu", "Screen Percentage"),
+						LOCTEXT("ScreenPercentageSubMenu_ToolTip", "Customize the viewport's screen percentage"),
+						FNewToolMenuDelegate::CreateLambda(
+							[](UToolMenu* ScreenPercentageSubMenu)
+							{
+								ScreenPercentageSubMenu->AddDynamicSection(
+									NAME_None,
+									FNewToolMenuDelegate::CreateLambda(
+										[](UToolMenu* InMenu) -> void
+										{
+											UE::UnrealEd::ConstructScreenPercentageMenu(InMenu);
+										}
+									)
+								);
+							}
+						)
+					);
+				}
+			}
+		)
+	);
+	Entry.Icon = FSlateIcon(FAppStyle::GetAppStyleSetName(), "Icons.Scalability");
+	Entry.InsertPosition.Position = EToolMenuInsertType::Last;
+	Entry.ToolBarData.LabelOverride = FText();
+	Entry.ToolBarData.ResizeParams.ClippingPriority = 800;
+	return Entry;
+}
+
+void GenerateViewportLayoutsMenu(UToolMenu* InMenu)
+{
+	ULevelViewportContext* const LevelViewportContext = InMenu->FindContext<ULevelViewportContext>();
+	if (!LevelViewportContext)
+	{
+		return;
+	}
+
+	const TSharedPtr<::SLevelViewport> LevelViewport = LevelViewportContext->LevelViewport.Pin();
+	if (!LevelViewport)
+	{
+		return;
+	}
+	TSharedPtr<FUICommandList> CommandList = LevelViewport->GetCommandList();
+
+	// Disable searching in this menu because it only contains visual representations of
+	// viewport layouts without any searchable text.
+	InMenu->bSearchable = false;
+
+	{
+		FToolMenuSection& Section =
+			InMenu->AddSection("LevelViewportOnePaneConfigs", LOCTEXT("OnePaneConfigHeader", "One Pane"));
+
+		FSlimHorizontalToolBarBuilder OnePaneButton(CommandList, FMultiBoxCustomization::None);
+		OnePaneButton.SetLabelVisibility(EVisibility::Collapsed);
+		OnePaneButton.SetStyle(&FAppStyle::Get(), "ViewportLayoutToolbar");
+
+		OnePaneButton.AddToolBarButton(FLevelViewportCommands::Get().ViewportConfig_OnePane);
+
+		Section.AddEntry(FToolMenuEntry::InitWidget(
+			"LevelViewportOnePaneConfigs",
+			// clang-format off
+			SNew(SHorizontalBox)
+			+ SHorizontalBox::Slot()
+			.AutoWidth()
+			[
+				OnePaneButton.MakeWidget()
+			]
+			+ SHorizontalBox::Slot()
+			.FillWidth(1)
+			[
+				SNullWidget::NullWidget
+			],
+			// clang-format on
+			FText::GetEmpty(),
+			true
+		));
+	}
+
+	{
+		FToolMenuSection& Section =
+			InMenu->AddSection("LevelViewportTwoPaneConfigs", LOCTEXT("TwoPaneConfigHeader", "Two Panes"));
+		FSlimHorizontalToolBarBuilder TwoPaneButtons(CommandList, FMultiBoxCustomization::None);
+		TwoPaneButtons.SetLabelVisibility(EVisibility::Collapsed);
+		TwoPaneButtons.SetStyle(&FAppStyle::Get(), "ViewportLayoutToolbar");
+
+		TwoPaneButtons.AddToolBarButton(FLevelViewportCommands::Get().ViewportConfig_TwoPanesH, NAME_None, FText());
+		TwoPaneButtons.AddToolBarButton(FLevelViewportCommands::Get().ViewportConfig_TwoPanesV, NAME_None, FText());
+
+		Section.AddEntry(FToolMenuEntry::InitWidget(
+			"LevelViewportTwoPaneConfigs",
+			// clang-format off
+			SNew(SHorizontalBox)
+			+SHorizontalBox::Slot()
+			.AutoWidth()
+			[
+				TwoPaneButtons.MakeWidget()
+			]
+			+SHorizontalBox::Slot()
+			.FillWidth(1)
+			[
+				SNullWidget::NullWidget
+			],
+			// clang-format on
+			FText::GetEmpty(),
+			true
+		));
+	}
+
+	{
+		FToolMenuSection& Section =
+			InMenu->AddSection("LevelViewportThreePaneConfigs", LOCTEXT("ThreePaneConfigHeader", "Three Panes"));
+		FSlimHorizontalToolBarBuilder ThreePaneButtons(CommandList, FMultiBoxCustomization::None);
+		ThreePaneButtons.SetLabelVisibility(EVisibility::Collapsed);
+		ThreePaneButtons.SetStyle(&FAppStyle::Get(), "ViewportLayoutToolbar");
+
+		ThreePaneButtons.AddToolBarButton(FLevelViewportCommands::Get().ViewportConfig_ThreePanesLeft, NAME_None, FText());
+		ThreePaneButtons.AddToolBarButton(FLevelViewportCommands::Get().ViewportConfig_ThreePanesRight, NAME_None, FText());
+		ThreePaneButtons.AddToolBarButton(FLevelViewportCommands::Get().ViewportConfig_ThreePanesTop, NAME_None, FText());
+		ThreePaneButtons.AddToolBarButton(FLevelViewportCommands::Get().ViewportConfig_ThreePanesBottom, NAME_None, FText());
+
+		Section.AddEntry(FToolMenuEntry::InitWidget(
+			"LevelViewportThreePaneConfigs",
+			// clang-format off
+			SNew(SHorizontalBox)
+			+SHorizontalBox::Slot()
+			.AutoWidth()
+			[
+				ThreePaneButtons.MakeWidget()
+			]
+			+SHorizontalBox::Slot()
+			.FillWidth(1)
+			[
+				SNullWidget::NullWidget
+			],
+			// clang-format on
+			FText::GetEmpty(),
+			true
+		));
+	}
+
+	{
+		FToolMenuSection& Section =
+			InMenu->AddSection("LevelViewportFourPaneConfigs", LOCTEXT("FourPaneConfigHeader", "Four Panes"));
+		FSlimHorizontalToolBarBuilder FourPaneButtons(CommandList, FMultiBoxCustomization::None);
+		FourPaneButtons.SetLabelVisibility(EVisibility::Collapsed);
+		FourPaneButtons.SetStyle(&FAppStyle::Get(), "ViewportLayoutToolbar");
+
+		FourPaneButtons.AddToolBarButton(FLevelViewportCommands::Get().ViewportConfig_FourPanes2x2, NAME_None, FText());
+		FourPaneButtons.AddToolBarButton(FLevelViewportCommands::Get().ViewportConfig_FourPanesLeft, NAME_None, FText());
+		FourPaneButtons.AddToolBarButton(FLevelViewportCommands::Get().ViewportConfig_FourPanesRight, NAME_None, FText());
+		FourPaneButtons.AddToolBarButton(FLevelViewportCommands::Get().ViewportConfig_FourPanesTop, NAME_None, FText());
+		FourPaneButtons.AddToolBarButton(FLevelViewportCommands::Get().ViewportConfig_FourPanesBottom, NAME_None, FText());
+
+		Section.AddEntry(FToolMenuEntry::InitWidget(
+			"LevelViewportFourPaneConfigs",
+			// clang-format off
+			SNew(SHorizontalBox)
+			+SHorizontalBox::Slot()
+			.AutoWidth()
+			[
+				FourPaneButtons.MakeWidget()
+			]
+			+SHorizontalBox::Slot()
+			.FillWidth(1)
+			[
+				SNullWidget::NullWidget
+			],
+			// clang-format on
+			FText::GetEmpty(),
+			true
+		));
+	}
+}
+
+TSharedRef<SWidget> BuildVolumeControlCustomWidget()
+{
+	// clang-format off
+	return SNew(SHorizontalBox)
+		+ SHorizontalBox::Slot()
+		.FillWidth(0.9f)
+		.Padding(FMargin(2.0f, 0.0f, 0.0f, 0.0f))
+		[
+			SNew(SVolumeControl)
+				.ToolTipText_Static(&FLevelEditorActionCallbacks::GetAudioVolumeToolTip)
+				.Volume_Static(&FLevelEditorActionCallbacks::GetAudioVolume)
+				.OnVolumeChanged_Static(&FLevelEditorActionCallbacks::OnAudioVolumeChanged)
+				.Muted_Static(&FLevelEditorActionCallbacks::GetAudioMuted)
+				.OnMuteChanged_Static(&FLevelEditorActionCallbacks::OnAudioMutedChanged)
+		]
+		+ SHorizontalBox::Slot()
+		.FillWidth(0.1f);
+	// clang-format on
+}
+
+FToolMenuEntry CreateSettingsSubmenu()
+{
+	FToolMenuEntry Entry = FToolMenuEntry::InitSubMenu(
+		"Settings",
+		LOCTEXT("SettingsSubmenuLabel", "Settings"),
+		LOCTEXT("SettingsSubmenuTooltip", "Viewport-related settings"),
+		FNewToolMenuDelegate::CreateLambda(
+			[](UToolMenu* Submenu) -> void
+			{
+				{
+					FToolMenuSection& SettingsSection =
+						Submenu->FindOrAddSection("Settings", LOCTEXT("SettingsSectionLabel", "Settings"));
+
+					SettingsSection.AddEntry(FToolMenuEntry::InitWidget(
+						"Volume", BuildVolumeControlCustomWidget(), LOCTEXT("VolumeControlLabel", "Volume")
+					));
+
+					SettingsSection.AddSeparator("AdvancedSeparator");
+
+					{
+						const FLevelViewportCommands& LevelViewportActions = FLevelViewportCommands::Get();
+						SettingsSection.AddMenuEntry(LevelViewportActions.AdvancedSettings);
+					}
+				}
+			}
+		)
+	);
+	Entry.Icon = FSlateIcon(FAppStyle::GetAppStyleSetName(), "LevelEditor.GameSettings");
+	Entry.InsertPosition.Position = EToolMenuInsertType::Last;
+	Entry.ToolBarData.LabelOverride = FText();
+	Entry.ToolBarData.ResizeParams.ClippingPriority = 1000;
+	return Entry;
+}
+
+FToolMenuEntry CreateViewportSizingSubmenu()
+{
+	FToolMenuEntry Entry = FToolMenuEntry::InitSubMenu(
+		"ViewportSizing",
+		LOCTEXT("ViewportSizingLabel", "..."),
+		LOCTEXT("ViewportSizingTooltip", "Viewport-sizing settings"),
+		FNewToolMenuDelegate::CreateLambda(
+			[](UToolMenu* Submenu) -> void
+			{
+				ULevelViewportContext* const LevelViewportContext = Submenu->FindContext<ULevelViewportContext>();
+				if (!LevelViewportContext)
+				{
+					return;
+				}
+
+				if (const TSharedPtr<::SLevelViewport> LevelViewport = LevelViewportContext->LevelViewport.Pin())
+				{
+					GenerateViewportLayoutsMenu(Submenu);
+				}
+
+				{
+					FToolMenuSection& MaximizeSection = Submenu->FindOrAddSection("MaximizeSection");
+
+					MaximizeSection.AddSeparator("MaximizeSeparator");
+
+					MaximizeSection.AddMenuEntry(FLevelViewportCommands::Get().ToggleImmersive);
+
+					MaximizeSection.AddDynamicEntry(
+						"MaximizeRestoreDynamicEntry",
+						FNewToolMenuSectionDelegate::CreateLambda(
+							[](FToolMenuSection& InnerSection) -> void
+							{
+								ULevelViewportContext* const LevelViewportContext =
+									InnerSection.FindContext<ULevelViewportContext>();
+								if (!LevelViewportContext)
+								{
+									return;
+								}
+
+								const TAttribute<FText> Label = TAttribute<FText>::CreateLambda(
+									[WeakLevelViewport = LevelViewportContext->LevelViewport]() -> FText
+									{
+										if (const TSharedPtr<::SLevelViewport> LevelViewport = WeakLevelViewport.Pin())
+										{
+											if (!LevelViewport->IsMaximized())
+											{
+												return LOCTEXT("MaximizeRestoreLabel_Maximize", "Maximize Viewport");
+											}
+										}
+										return LOCTEXT("MaximizeRestoreLabel_Restore", "Restore All Viewports");
+									}
+								);
+
+								const TAttribute<FText> Tooltip = TAttribute<FText>::CreateLambda(
+									[WeakLevelViewport = LevelViewportContext->LevelViewport]() -> FText
+									{
+										if (const TSharedPtr<::SLevelViewport> LevelViewport = WeakLevelViewport.Pin())
+										{
+											if (!LevelViewport->IsMaximized())
+											{
+												return LOCTEXT("MaximizeRestoreTooltip_Maximize", "Maximizes this viewport");
+											}
+										}
+										return LOCTEXT(
+											"MaximizeRestoreTooltip_Restore", "Restores the layout to show all viewports"
+										);
+									}
+								);
+
+								const TAttribute<FSlateIcon> Icon = TAttribute<FSlateIcon>::CreateLambda(
+									[WeakLevelViewport = LevelViewportContext->LevelViewport]() -> FSlateIcon
+									{
+										if (const TSharedPtr<::SLevelViewport> LevelViewport = WeakLevelViewport.Pin())
+										{
+											if (!LevelViewport->IsMaximized())
+											{
+												return FSlateIcon(
+													FAppStyle::GetAppStyleSetName(), "EditorViewportToolBar.Maximize.Normal"
+												);
+											}
+										}
+
+										return FSlateIcon(
+											FAppStyle::GetAppStyleSetName(), "EditorViewportToolBar.Maximize.Checked"
+										);
+									}
+								);
+
+								FToolUIAction Action;
+								Action.ExecuteAction = FToolMenuExecuteAction::CreateLambda(
+									[](const FToolMenuContext& Context) -> void
+									{
+										if (ULevelViewportContext* const LevelViewportContext =
+												Context.FindContext<ULevelViewportContext>())
+										{
+											if (const TSharedPtr<::SLevelViewport> LevelViewport =
+													LevelViewportContext->LevelViewport.Pin())
+											{
+												LevelViewport->OnToggleMaximize();
+											}
+										}
+									}
+								);
+								Action.CanExecuteAction = FToolMenuCanExecuteAction::CreateLambda(
+									[](const FToolMenuContext& Context) -> bool
+									{
+										if (ULevelViewportContext* const LevelViewportContext =
+												Context.FindContext<ULevelViewportContext>())
+										{
+											if (const TSharedPtr<::SLevelViewport> LevelViewport =
+													LevelViewportContext->LevelViewport.Pin())
+											{
+												return LevelViewport->CanMaximize();
+											}
+										}
+
+										return false;
+									}
+								);
+
+								FToolMenuEntry& MaximizeRestoreEntry =
+									InnerSection.AddMenuEntry("MaximizeRestore", Label, Tooltip, Icon, Action);
+								MaximizeRestoreEntry.SetShowInToolbarTopLevel(true);
+								MaximizeRestoreEntry.ToolBarData.ResizeParams.AllowClipping = false;
+								MaximizeRestoreEntry.StyleNameOverride = FName("ViewportToolbarViewportSizingSubmenu");
+							}
+						)
+					);
+				}
+			}
+		)
+	);
+
+	Entry.StyleNameOverride = FName("ViewportToolbarViewportSizingSubmenu");
+	Entry.InsertPosition.Position = EToolMenuInsertType::Last;
+	Entry.ToolBarData.LabelOverride = FText();
+	Entry.ToolBarData.ResizeParams.AllowClipping = false;
+
+	return Entry;
+}
+
+void CreateCameraSpawnMenu(UToolMenu* InMenu)
+{
+	FToolMenuSection& Section = InMenu->AddSection("Section");
+	const FLevelViewportCommands& Actions = FLevelViewportCommands::Get();
+
+	for (TSharedPtr<FUICommandInfo> Camera : Actions.CreateCameras)
+	{
+		Section.AddMenuEntry(NAME_None, Camera);
+	}
+}
+
+void CreateBookmarksMenu(UToolMenu* InMenu)
+{
+	ULevelViewportContext* const LevelViewportContext = InMenu->FindContext<ULevelViewportContext>();
+	if (!LevelViewportContext)
+	{
+		return;
+	}
+
+	const TSharedPtr<::SLevelViewport> LevelViewportPinned = LevelViewportContext->LevelViewport.Pin();
+	if (!LevelViewportPinned)
+	{
+		return;
+	}
+
+	// Add a menu entry for each bookmark
+	const FLevelEditorViewportClient& ViewportClient = LevelViewportPinned->GetLevelViewportClient();
+
+	FToolMenuSection& ManageBookmarksSection =
+		InMenu->FindOrAddSection("ManageBookmarks", LOCTEXT("ManageBookmarkSectionName", "Manage Bookmarks"));
+
+	bool bFoundBookmarks = Private::AddJumpToBookmarkMenu(InMenu, &ViewportClient);
+
+	// Manage Bookmarks Section
+	{
+		// Set Bookmark Submenu
+		{
+			const int32 NumberOfBookmarks =
+				static_cast<int32>(IBookmarkTypeTools::Get().GetMaxNumberOfBookmarks(&ViewportClient));
+			const int32 NumberOfMappedBookmarks = FMath::Min<int32>(AWorldSettings::NumMappedBookmarks, NumberOfBookmarks);
+
+			ManageBookmarksSection.AddSubMenu(
+				"SetBookmark",
+				LOCTEXT("SetBookmarkSubMenu", "Set Bookmark"),
+				LOCTEXT("SetBookmarkSubMenu_ToolTip", "Setting bookmarks"),
+				FNewToolMenuDelegate::CreateLambda(
+					[NumberOfMappedBookmarks](UToolMenu* InMenu)
+					{
+						const FLevelViewportCommands& Actions = FLevelViewportCommands::Get();
+
+						FToolMenuSection& SetBookmarksSection =
+							InMenu->FindOrAddSection("SetBookmark", LOCTEXT("SetBookmarkSectionName", "Set Bookmark"));
+
+						for (int32 BookmarkIndex = 0; BookmarkIndex < NumberOfMappedBookmarks; ++BookmarkIndex)
+						{
+							SetBookmarksSection.AddMenuEntry(
+								NAME_None,
+								Actions.SetBookmarkCommands[BookmarkIndex],
+								FBookmarkUI::GetPlainLabel(BookmarkIndex),
+								FBookmarkUI::GetSetTooltip(BookmarkIndex),
+								FSlateIcon(FAppStyle::Get().GetStyleSetName(), "LevelViewport.ToggleActorPilotCameraView")
+							);
+						}
+					}
+				),
+				false,
+				FSlateIcon(FSlateIcon(FAppStyle::GetAppStyleSetName(), "LevelViewport.ToggleActorPilotCameraView"))
+			);
+		}
+
+		// Manage Bookmarks Submenu
+		{
+			if (bFoundBookmarks)
+			{
+				ManageBookmarksSection.AddSubMenu(
+					"ManageBookmarks",
+					LOCTEXT("ManageBookmarksSubMenu", "Manage Bookmarks"),
+					LOCTEXT("ManageBookmarksSubMenu_ToolTip", "Bookmarks related actions"),
+					FNewToolMenuDelegate::CreateLambda(
+						[bFoundBookmarks, LevelViewportWeak = LevelViewportContext->LevelViewport](UToolMenu* InMenu)
+						{
+							if (!bFoundBookmarks)
+							{
+								return;
+							}
+
+							const FLevelViewportCommands& Actions = FLevelViewportCommands::Get();
+
+							FToolMenuSection& ManageBookmarksSubsection = InMenu->FindOrAddSection(
+								"ManageBookmarks", LOCTEXT("ManageBookmarkSectionName", "Manage Bookmarks")
+							);
+
+							ManageBookmarksSubsection.AddSubMenu(
+								"ClearBookmark",
+								LOCTEXT("ClearBookmarkSubMenu", "Clear Bookmark"),
+								LOCTEXT("ClearBookmarkSubMenu_ToolTip", "Clear viewport bookmarks"),
+								FNewToolMenuDelegate::CreateLambda(&Private::AddClearBookmarkMenu, LevelViewportWeak),
+								false,
+								FSlateIcon(FAppStyle::Get().GetStyleSetName(), "EditorViewport.SubMenu.Bookmarks")
+							);
+
+							FToolMenuEntry& CompactBookmarks =
+								ManageBookmarksSubsection.AddMenuEntry(Actions.CompactBookmarks);
+							CompactBookmarks.Icon =
+								FSlateIcon(FAppStyle::GetAppStyleSetName(), "AnimationEditor.ApplyCompression");
+
+							FToolMenuEntry& ClearBookmarks =
+								ManageBookmarksSubsection.AddMenuEntry(Actions.ClearAllBookmarks);
+							ClearBookmarks.Icon = FSlateIcon(FAppStyle::GetAppStyleSetName(), "GraphEditor.Clean");
+						}
+					),
+					false,
+					FSlateIcon(FAppStyle::Get().GetStyleSetName(), "EditorViewport.SubMenu.Bookmarks")
+				);
+			}
+		}
+	}
+}
+
+FToolMenuEntry CreateFOVMenu(TWeakPtr<::SLevelViewport> InLevelViewportWeak)
+{
+	constexpr float FOVMin = 5.0f;
+	constexpr float FOVMax = 170.0f;
+
+	return UnrealEd::CreateNumericEntry(
+		"FOVAngle",
+		LOCTEXT("FOVAngle", "Field of View"),
+		LOCTEXT("FOVAngleTooltip", "Field of View"),
+		FCanExecuteAction(),
+		UnrealEd::FNumericEntryExecuteActionDelegate::CreateLambda(
+			[InLevelViewportWeak](float InValue)
+			{
+				if (TSharedPtr<::SLevelViewport> LevelViewport = InLevelViewportWeak.Pin())
+				{
+					Private::SetLevelViewportFOV(LevelViewport.ToSharedRef(), InValue);
+				}
+			}
+		),
+		TAttribute<float>::CreateLambda(
+			[InLevelViewportWeak]()
+			{
+				if (TSharedPtr<::SLevelViewport> Viewport = InLevelViewportWeak.Pin())
+				{
+					return Private::GetLevelViewportFOV(Viewport.ToSharedRef());
+				}
+
+				return FOVMin;
+			}
+		),
+		FOVMin,
+		FOVMax,
+		1
+	);
+}
+
+FToolMenuEntry CreateFarViewPlaneMenu(TWeakPtr<::SLevelViewport> InLevelViewportWeak)
+{
+	constexpr float FarMin = 0.0f;
+	constexpr float FarMax = 100000.0f;
+
+	return UnrealEd::CreateNumericEntry(
+		"FarViewPlane",
+		LOCTEXT("FarViewPlane", "Far View Plane"),
+		LOCTEXT("FarViewPlaneTooltip", "Far View Plane"),
+		FCanExecuteAction(),
+		UnrealEd::FNumericEntryExecuteActionDelegate::CreateLambda(
+			[InLevelViewportWeak](float InValue)
+			{
+				if (TSharedPtr<::SLevelViewport> LevelViewport = InLevelViewportWeak.Pin())
+				{
+					Private::SetFarViewPlaneValue(LevelViewport.ToSharedRef(), InValue);
+				}
+			}
+		),
+		TAttribute<float>::CreateLambda(
+			[InLevelViewportWeak]()
+			{
+				if (TSharedPtr<::SLevelViewport> Viewport = InLevelViewportWeak.Pin())
+				{
+					return Private::GetFarViewPlaneValue(Viewport.ToSharedRef());
+				}
+
+				return FarMax;
+			}
+		),
+		FarMin,
+		FarMax,
+		1
+	);
+}
+
+void AddCameraActorSelectSection(UToolMenu* InMenu)
+{
+	ULevelViewportContext* const LevelViewportContext = InMenu->FindContext<ULevelViewportContext>();
+	if (!LevelViewportContext)
+	{
+		return;
+	}
+
+	TSharedPtr<::SLevelViewport> LevelViewport = LevelViewportContext->LevelViewport.Pin();
+	if (!LevelViewport)
+	{
+		return;
+	}
+
+	TArray<AActor*> LookThroughActors;
+
+	if (UWorld* World = LevelViewport->GetWorld())
+	{
+		for (TActorIterator<ACameraActor> It(World); It; ++It)
+		{
+			LookThroughActors.Add(Cast<AActor>(*It));
+		}
+
+		for (TActorIterator<ASceneCapture> It(World); It; ++It)
+		{
+			LookThroughActors.Add(Cast<AActor>(*It));
+		}
+	}
+
+	FText CameraActorsHeading = LOCTEXT("CameraActorsHeading", "Cameras");
+
+	FToolMenuInsert InsertPosition("LevelViewportCameraType_Perspective", EToolMenuInsertType::After);
+
+	FToolMenuSection& Section = InMenu->AddSection("CameraActors");
+	Section.InsertPosition = InsertPosition;
+
+	// Don't add too many cameras to the top level menu or else it becomes too large
+	constexpr uint32 MaxCamerasInTopLevelMenu = 10;
+	if (LookThroughActors.Num() > MaxCamerasInTopLevelMenu)
+	{
+		Section.AddSubMenu(
+			"CameraActors",
+			CameraActorsHeading,
+			LOCTEXT("LookThroughPlacedCameras_ToolTip", "Look through and pilot placed cameras"),
+			FNewToolMenuDelegate::CreateLambda(
+				[LookThroughActors, LevelViewportWeak = LevelViewportContext->LevelViewport](UToolMenu* InMenu)
+				{
+					if (TSharedPtr<::SLevelViewport> LevelViewport = LevelViewportWeak.Pin())
+					{
+						FToolMenuSection& Section = InMenu->FindOrAddSection(NAME_None);
+						UE::LevelEditor::Private::GeneratePlacedCameraMenuEntries(Section, LookThroughActors, LevelViewport);
+					}
+				}
+			)
+		);
+	}
+	else if (!LookThroughActors.IsEmpty())
+	{
+		Section.AddSeparator(NAME_None);
+		UE::LevelEditor::Private::GeneratePlacedCameraMenuEntries(Section, LookThroughActors, LevelViewport);
+	}
+
+	TWeakObjectPtr<AActor> LockedActorWeak = LevelViewport->GetLevelViewportClient().GetActorLock().LockedActor;
+
+	if (TStrongObjectPtr<AActor> LockedActor = LockedActorWeak.Pin())
+	{
+		if (!LockedActor->IsA<ACameraActor>() && !LockedActor->IsA<ASceneCapture>())
+		{
+			UE::LevelEditor::Private::GeneratePlacedCameraMenuEntries(Section, { LockedActor.Get() }, LevelViewport);
+		}
+	}
+}
+
+void ExtendCameraSubmenu(FName InCameraOptionsSubmenuName)
+{
+	UToolMenu* const Submenu = UToolMenus::Get()->ExtendMenu(InCameraOptionsSubmenuName);
+
+	Submenu->AddDynamicSection(
+		"LevelEditorCameraExtensionDynamicSection",
+		FNewToolMenuDelegate::CreateLambda(
+			[](UToolMenu* InDynamicMenu)
+			{
+				ULevelViewportContext* LevelViewportContext = InDynamicMenu->FindContext<ULevelViewportContext>();
+				TWeakPtr<::SLevelViewport> LevelViewportWeak = LevelViewportContext->LevelViewport;
+
+				FToolMenuSection& SelectPilotedActorSection = InDynamicMenu->AddSection("SelectPilotedActor");
+
+				SelectPilotedActorSection.InsertPosition =
+					FToolMenuInsert("LevelViewportCameraType_Perspective", EToolMenuInsertType::Before);
+
+				SelectPilotedActorSection.AddMenuEntry(FLevelViewportCommands::Get().SelectPilotedActor);
+
+				// Camera Selection elements
+				{
+					AddCameraActorSelectSection(InDynamicMenu);
+				}
+
+				// Pilot Section
+				{
+					FToolMenuSection& PilotSection = InDynamicMenu->FindOrAddSection("Pilot");
+
+					PilotSection.InsertPosition =
+						FToolMenuInsert("LevelViewportCameraType_Ortho", EToolMenuInsertType::After);
+					PilotSection.AddSeparator("PilotSectionSeparator");
+
+					bool bShowPilotSelectedActorEntry = false;
+
+					AActor* SelectedActor = nullptr;
+					if (TSharedPtr<::SLevelViewport> LevelViewport = LevelViewportWeak.Pin())
+					{
+						TArray<AActor*> SelectedActors;
+						GEditor->GetSelectedActors()->GetSelectedObjects<AActor>(SelectedActors);
+
+						if (!SelectedActors.IsEmpty() && !LevelViewport->IsSelectedActorLocked())
+						{
+							SelectedActor = SelectedActors[0];
+							const FLevelEditorViewportClient& ViewportClient = LevelViewport->GetLevelViewportClient();
+
+							bShowPilotSelectedActorEntry = SelectedActor && ViewportClient.IsPerspective()
+														&& !ViewportClient.IsLockedToCinematic();
+						}
+					}
+
+					if (bShowPilotSelectedActorEntry)
+					{
+						// Pilot Selected Actor Entry
+						PilotSection.AddMenuEntry(
+							FLevelViewportCommands::Get().PilotSelectedActor,
+							FText::Format(
+								LOCTEXT("PilotActor", "Pilot '{0}'"), FText::FromString(SelectedActor->GetActorLabel())
+							)
+						);
+					}
+
+					// Stop Piloting Entry
+					PilotSection.AddEntry(UE::LevelEditor::Private::CreateEjectActorPilotEntry());
+
+					// Exact Camera View Entry
+					{
+						FToolMenuEntry& ToggleCameraView =
+							PilotSection.AddMenuEntry(FLevelViewportCommands::Get().ToggleActorPilotCameraView);
+						ToggleCameraView.Label = LOCTEXT("ToggleCameraViewLabel", "Exact Camera View");
+						ToggleCameraView.SetShowInToolbarTopLevel(TAttribute<bool>::CreateLambda(
+							[LevelViewportWeak]()
+							{
+								if (TSharedPtr<::SLevelViewport> EditorViewport = LevelViewportWeak.Pin())
+								{
+									return EditorViewport->IsAnyActorLocked();
+								}
+								return false;
+							}
+						));
+					}
+
+					UE::UnrealEd::GenerateViewportTypeMenu(InDynamicMenu);
+
+				}
+
+				// Create Section
+				{
+					FToolMenuSection& CreateSection = InDynamicMenu->FindOrAddSection("Create", LOCTEXT("CreateLabel", "Create"));
+
+					CreateSection.AddSubMenu(
+						"CreateCamera",
+						LOCTEXT("CameraSubMenu", "Create Camera"),
+						LOCTEXT("CameraSubMenu_ToolTip", "Select a camera type to create at current viewport's location"),
+						FNewToolMenuDelegate::CreateLambda(
+							[](UToolMenu* InMenu)
+							{
+								CreateCameraSpawnMenu(InMenu);
+							}
+						),
+						false,
+						FSlateIcon(FAppStyle::Get().GetStyleSetName(), "EditorViewport.SubMenu.CreateCamera")
+					);
+
+					CreateSection.AddSubMenu(
+						"Bookmarks",
+						LOCTEXT("BookmarksSubMenu", "Bookmarks"),
+						LOCTEXT("BookmarksSubMenu_ToolTip", "Bookmarks related actions"),
+						FNewToolMenuDelegate::CreateLambda(
+							[](UToolMenu* InMenu)
+							{
+								CreateBookmarksMenu(InMenu);
+							}
+						),
+						false,
+						FSlateIcon(FAppStyle::Get().GetStyleSetName(), "EditorViewport.SubMenu.Bookmarks")
+					);
+				}
+
+				// Positioning Section
+				{
+					FToolMenuSection& PositioningSection =
+						InDynamicMenu->FindOrAddSection("Positioning", LOCTEXT("PositioningLabel", "Positioning"));
+
+					// Camera Speed Submenu
+					{
+						FToolMenuEntry& Entry = PositioningSection.AddSubMenu(
+							"CameraSpeed",
+							LOCTEXT("CameraSpeedSubMenu", "Camera Speed"),
+							LOCTEXT("CameraSpeedSubMenu_ToolTip", "Camera Speed related actions"),
+							FNewToolMenuDelegate::CreateLambda(
+								[LevelViewportWeak](UToolMenu* InMenu)
+								{
+									UE::UnrealEd::CreateCameraSpeedMenu(InMenu);
+								}
+							),
+							false,
+							FSlateIcon(FSlateIcon(FAppStyle::GetAppStyleSetName(), "EditorViewport.Speed"))
+						);
+
+						Entry.ToolBarData.LabelOverride = TAttribute<FText>::CreateLambda(
+							[LevelViewportWeak]()
+							{
+								return UE::UnrealEd::GetCameraSpeedLabel(LevelViewportWeak);
+							}
+						);
+						Entry.SetShowInToolbarTopLevel(true);
+						Entry.ToolBarData.PlacementOverride = MenuPlacement_BelowRightAnchor;
+					}
+
+					PositioningSection.AddSeparator("PositioningSeparator_1");
+
+					// Frame Selection
+					{
+						FToolMenuEntry FocusViewportToSelection =
+							FToolMenuEntry::InitMenuEntry(FEditorViewportCommands::Get().FocusViewportToSelection);
+						FocusViewportToSelection.UserInterfaceActionType = EUserInterfaceActionType::ToggleButton;
+						FocusViewportToSelection.Icon = FSlateIcon(FAppStyle::GetAppStyleSetName(), "Icons.FrameActor");
+						PositioningSection.AddEntry(FocusViewportToSelection);
+					}
+
+					PositioningSection.AddSeparator("PositioningSeparator_2");
+
+					// Move Camera/Object
+					{
+						FToolMenuEntry CameraToObjectMenu =
+							FToolMenuEntry::InitMenuEntry(FLevelEditorCommands::Get().SnapCameraToObject);
+						CameraToObjectMenu.UserInterfaceActionType = EUserInterfaceActionType::ToggleButton;
+						CameraToObjectMenu.Label = LOCTEXT("CameraToObjectLabel", "Move Camera to Object");
+						PositioningSection.AddEntry(CameraToObjectMenu);
+
+						FToolMenuEntry ObjectToCameraMenu =
+							FToolMenuEntry::InitMenuEntry(FLevelEditorCommands::Get().SnapObjectToCamera);
+						ObjectToCameraMenu.UserInterfaceActionType = EUserInterfaceActionType::ToggleButton;
+						ObjectToCameraMenu.Label = LOCTEXT("ObjectToCameraLabel", "Move Object to Camera");
+						PositioningSection.AddEntry(ObjectToCameraMenu);
+					}
+				}
+
+				// Options Section
+				{
+					FToolMenuSection& OptionsSection =
+						InDynamicMenu->FindOrAddSection("CameraOptions", LOCTEXT("OptionsLabel", "Options"));
+					// add Cinematic Viewport
+					// add Allow Cinematic Control
+					// add Game View
+
+					FToolMenuEntry AllowCinematicControl =
+						FToolMenuEntry::InitMenuEntry(FLevelViewportCommands::Get().ToggleCinematicPreview);
+					AllowCinematicControl.UserInterfaceActionType = EUserInterfaceActionType::ToggleButton;
+					OptionsSection.AddEntry(AllowCinematicControl);
+
+					FToolMenuEntry ToggleGameView =
+						FToolMenuEntry::InitMenuEntry(FLevelViewportCommands::Get().ToggleGameView);
+					ToggleGameView.UserInterfaceActionType = EUserInterfaceActionType::ToggleButton;
+					OptionsSection.AddEntry(ToggleGameView);
+
+					// This additional options section is used to force certain elements to appear after extensions
+					{
+						FToolMenuSection& AdditionalOptions =
+							InDynamicMenu->FindOrAddSection("AdditionalOptions");
+						AdditionalOptions.AddSeparator("AdditionalOptionsSeparator");
+
+						FToolMenuEntry HighResolutionScreenshot =
+							FToolMenuEntry::InitMenuEntry(FLevelViewportCommands::Get().HighResScreenshot);
+						HighResolutionScreenshot.UserInterfaceActionType = EUserInterfaceActionType::ToggleButton;
+						AdditionalOptions.AddEntry(HighResolutionScreenshot);
+					}
+				}
+			}
+		)
+	);
+}
+
+FToolMenuEntry CreateToolbarCameraSubmenu()
+{
+	FToolMenuEntry Entry = FToolMenuEntry::InitDynamicEntry(
+		"DynamicCameraOptions",
+		FNewToolMenuSectionDelegate::CreateLambda(
+			[](FToolMenuSection& InDynamicSection) -> void
+			{
+				if (ULevelViewportContext* const LevelViewportContext =
+						InDynamicSection.FindContext<ULevelViewportContext>())
+				{
+					const TAttribute<FText> Label = TAttribute<FText>::CreateLambda(
+						[LevelViewportWeak = LevelViewportContext->LevelViewport]()
+						{
+							return UE::LevelEditor::Private::GetCameraSubmenuLabelFromLevelViewport(LevelViewportWeak);
+						}
+					);
+
+					const TAttribute<FSlateIcon> Icon = TAttribute<FSlateIcon>::CreateLambda(
+						[LevelViewportWeak = LevelViewportContext->LevelViewport]()
+						{
+							return UE::LevelEditor::Private::GetCameraSubmenuIconFromLevelViewport(LevelViewportWeak);
+						}
+					);
+
+					FToolMenuEntry& Entry = InDynamicSection.AddSubMenu(
+						"CameraOptions",
+						Label,
+						LOCTEXT("CameraSubmenuTooltip", "Camera options"),
+						FNewToolMenuDelegate::CreateLambda(
+							[](UToolMenu* Submenu) -> void
+							{
+								UnrealEd::PopulateCameraMenu(Submenu);
+							}
+						),
+						false,
+						Icon
+					);
+					Entry.InsertPosition.Position = EToolMenuInsertType::Last;
+					Entry.ToolBarData.ResizeParams.ClippingPriority = 800;
+				}
+			}
+		)
+	);
+
+	return Entry;
+}
+
+} // namespace UE::LevelEditor
+
+#undef LOCTEXT_NAMESPACE
diff --git a/Engine/Source/Editor/LevelEditor/Private/ViewportToolbar/LevelEditorSubmenus.h b/Engine/Source/Editor/LevelEditor/Private/ViewportToolbar/LevelEditorSubmenus.h
new file mode 100644
--- /dev/null
+++ b/Engine/Source/Editor/LevelEditor/Private/ViewportToolbar/LevelEditorSubmenus.h
@@ -0,0 +1,54 @@
+// Copyright Epic Games, Inc. All Rights Reserved.
+
+#pragma once
+
+#include "Templates/SharedPointerFwd.h"
+#include "ToolMenuEntry.h"
+
+class FExtender;
+class FLevelEditorViewportClient;
+class SLevelViewport;
+class SWidget;
+class UToolMenu;
+struct FToolMenuSection;
+
+namespace UE::LevelEditor
+{
+void CreateCameraSpawnMenu(UToolMenu* InMenu);
+void CreateBookmarksMenu(UToolMenu* InMenu);
+void AddCameraActorSelectSection(UToolMenu* InMenu);
+
+FToolMenuEntry CreateFOVMenu(TWeakPtr<SLevelViewport> InLevelViewportWeak);
+FToolMenuEntry CreateFarViewPlaneMenu(TWeakPtr<SLevelViewport> InLevelViewportWeak);
+FToolMenuEntry CreateCameraSpeedSlider(TWeakPtr<SLevelViewport> InLevelViewportWeak);
+FToolMenuEntry CreateCameraSpeedScalarSlider(TWeakPtr<SLevelViewport> InLevelViewportWeak);
+FToolMenuEntry CreateToolbarCameraSubmenu();
+
+TSharedPtr<FExtender> GetViewModesLegacyExtenders();
+void PopulateViewModesMenu(UToolMenu* InMenu);
+void ExtendViewModesSubmenu(FName InViewModesSubmenuName);
+
+FToolMenuEntry CreateShowFoliageSubmenu();
+FToolMenuEntry CreateShowHLODsSubmenu();
+FToolMenuEntry CreateShowLayersSubmenu();
+FToolMenuEntry CreateShowSpritesSubmenu();
+#if STATS
+FToolMenuEntry CreateShowStatsSubmenu(
+	bool bInAddToggleStatsCheckbox = false, TAttribute<FText> InLabelOverride = TAttribute<FText>()
+);
+#endif
+FToolMenuEntry CreateShowVolumesSubmenu();
+FToolMenuEntry CreateShowSubmenu();
+
+FToolMenuEntry CreateFeatureLevelPreviewSubmenu();
+FToolMenuEntry CreateMaterialQualityLevelSubmenu();
+FToolMenuEntry CreatePerformanceAndScalabilitySubmenu();
+
+void GenerateViewportLayoutsMenu(UToolMenu* InMenu);
+TSharedRef<SWidget> BuildVolumeControlCustomWidget();
+FToolMenuEntry CreateSettingsSubmenu();
+
+FToolMenuEntry CreateViewportSizingSubmenu();
+
+void ExtendCameraSubmenu(FName InCameraOptionsSubmenuName);
+} // namespace UE::LevelEditor
diff --git a/Engine/Source/Editor/LevelEditor/Private/ViewportToolbar/LevelEditorViewportToolbarSections.cpp b/Engine/Source/Editor/LevelEditor/Private/ViewportToolbar/LevelEditorViewportToolbarSections.cpp
deleted file mode 100644
diff --git a/Engine/Source/Editor/LevelEditor/Private/ViewportToolbar/LevelEditorViewportToolbarSections.h b/Engine/Source/Editor/LevelEditor/Private/ViewportToolbar/LevelEditorViewportToolbarSections.h
deleted file mode 100644
