[Viewport Toolbar] Transform section: remove Level editor only entries from other toolbars

diff --git a/Engine/Source/Editor/LevelEditor/Private/SLevelViewport.cpp b/Engine/Source/Editor/LevelEditor/Private/SLevelViewport.cpp
--- a/Engine/Source/Editor/LevelEditor/Private/SLevelViewport.cpp
+++ b/Engine/Source/Editor/LevelEditor/Private/SLevelViewport.cpp
@@ -1960,7 +1960,10 @@ TSharedPtr<SWidget> SLevelViewport::BuildViewportToolbar()
 				FToolMenuSection& LeftSection = ViewportToolbarMenu->FindOrAddSection("Left");
 
 				LeftSection.AddEntry(UE::UnrealEd::CreateTransformsSubmenu());
+				UE::LevelEditor::ExtendTransformSubmenu("LevelEditor.ViewportToolbar.Transform");
+
 				LeftSection.AddEntry(UE::UnrealEd::CreateSnappingSubmenu());
+				UE::LevelEditor::ExtendSnappingSubmenu("LevelEditor.ViewportToolbar.Snapping");
 			}
 
 			// Add the right-aligned part of the viewport toolbar.
diff --git a/Engine/Source/Editor/LevelEditor/Private/ViewportToolbar/LevelEditorSubmenus.cpp b/Engine/Source/Editor/LevelEditor/Private/ViewportToolbar/LevelEditorSubmenus.cpp
--- a/Engine/Source/Editor/LevelEditor/Private/ViewportToolbar/LevelEditorSubmenus.cpp
+++ b/Engine/Source/Editor/LevelEditor/Private/ViewportToolbar/LevelEditorSubmenus.cpp
@@ -26,6 +26,7 @@
 #include "ToolMenus.h"
 #include "ViewportToolbar/LevelViewportContext.h"
 #include "ViewportToolbar/UnrealEdViewportToolbar.h"
+#include "Widgets/Input/SNumericEntryBox.h"
 #include "Widgets/Input/SSpinBox.h"
 #include "Widgets/Input/SVolumeControl.h"
 #include "Widgets/SBoxPanel.h"
@@ -214,7 +215,8 @@ void AddClearBookmarkMenu(UToolMenu* InMenu, const TWeakPtr<::SLevelViewport>& I
 	{
 		FLevelEditorViewportClient& ViewportClient = LevelViewportPinned->GetLevelViewportClient();
 
-		const int32 NumberOfBookmarks = static_cast<int32>(IBookmarkTypeTools::Get().GetMaxNumberOfBookmarks(&ViewportClient));
+		const int32 NumberOfBookmarks =
+			static_cast<int32>(IBookmarkTypeTools::Get().GetMaxNumberOfBookmarks(&ViewportClient));
 		const int32 NumberOfMappedBookmarks = FMath::Min<int32>(AWorldSettings::NumMappedBookmarks, NumberOfBookmarks);
 
 		for (int32 BookmarkIndex = 0; BookmarkIndex < NumberOfMappedBookmarks; ++BookmarkIndex)
@@ -234,7 +236,9 @@ void AddClearBookmarkMenu(UToolMenu* InMenu, const TWeakPtr<::SLevelViewport>& I
 			if (IBookmarkTypeTools::Get().CheckBookmark(BookmarkIndex, &ViewportClient))
 			{
 				FUIAction Action;
-				Action.ExecuteAction.BindSP(LevelViewportPinned.ToSharedRef(), &::SLevelViewport::OnClearBookmark, BookmarkIndex);
+				Action.ExecuteAction.BindSP(
+					LevelViewportPinned.ToSharedRef(), &::SLevelViewport::OnClearBookmark, BookmarkIndex
+				);
 
 				Section.AddMenuEntry(
 					NAME_None,
@@ -419,6 +423,102 @@ FSlateIcon GetCameraSubmenuIconFromLevelViewport(const TWeakPtr<::SLevelViewport
 
 	return FSlateIcon(FAppStyle::GetAppStyleSetName(), "ClassIcon.CameraComponent");
 }
+
+FToolMenuEntry CreateActorSnapCheckboxMenu()
+{
+	FNewToolMenuDelegate MakeMenuDelegate = FNewToolMenuDelegate::CreateLambda(
+		[](UToolMenu* Submenu)
+		{
+			// Add "Actor snapping" widget.
+			{
+				FToolMenuSection& ActorSnappingSection =
+					Submenu->FindOrAddSection("ActorSnapping", LOCTEXT("ActorSnappingLabel", "Actor Snapping"));
+				const FText Label = LOCTEXT("ActorSnapDistanceLabel", "Snap Distance");
+					const FText Tooltip = LOCTEXT("ActorSnapDistanceTooltip", "The amount of offset to apply when snapping to surfaces");
+					const FMargin WidgetsMargin(2.0f, 0.0f, 3.0f, 0.0f);
+					FToolMenuEntry SnapDistance = FToolMenuEntry::InitMenuEntry(
+						"ActorSnapDistance",
+						FUIAction(
+							FExecuteAction(),
+							FCanExecuteAction::CreateLambda(
+								[]()
+								{
+									return !!GetDefault<ULevelEditorViewportSettings>()->bEnableActorSnap;
+								}
+							)
+						),
+						// clang-format off
+						SNew(SHorizontalBox)
+						+ SHorizontalBox::Slot()
+						.VAlign(VAlign_Center)
+						.Padding(WidgetsMargin)
+						.AutoWidth()
+						[
+							SNew(STextBlock)
+							.Text(Label)
+						]
+						+ SHorizontalBox::Slot()
+						.VAlign(VAlign_Center)
+						.Padding(WidgetsMargin)
+						.AutoWidth()
+						[
+							SNew(SBox).Padding(WidgetsMargin).MinDesiredWidth(100.0f)
+							[
+								// TODO: Check how to improve performance for this widget OnValueChanged.
+								// Same functionality in LevelEditorToolBar.cpp seems to have better performance
+								SNew(SNumericEntryBox<float>)
+								.ToolTipText(Tooltip)
+								.MinValue(0.0f)
+								.MaxValue(1.0f)
+								.MaxSliderValue(1.0f)
+								.AllowSpin(true)
+								.MaxFractionalDigits(1)
+								.Font(FAppStyle::GetFontStyle(TEXT("MenuItem.Font")))
+								.OnValueChanged_Static(&FLevelEditorActionCallbacks::SetActorSnapSetting)
+								.Value_Lambda([]()
+								{
+									return FLevelEditorActionCallbacks::GetActorSnapSetting();
+								})
+							]
+						]
+						// clang-format on
+					);
+				ActorSnappingSection.AddEntry(SnapDistance);
+			}
+		}
+	);
+	FToolUIAction CheckboxMenuAction;
+	{
+		CheckboxMenuAction.ExecuteAction = FToolMenuExecuteAction::CreateLambda(
+			[](const FToolMenuContext& InContext)
+			{
+				if (ULevelEditorViewportSettings* Settings = GetMutableDefault<ULevelEditorViewportSettings>())
+				{
+					Settings->bEnableActorSnap = !Settings->bEnableActorSnap;
+				}
+			}
+		);
+		CheckboxMenuAction.GetActionCheckState = FToolMenuGetActionCheckState::CreateLambda(
+			[](const FToolMenuContext& InContext)
+			{
+				return GetDefault<ULevelEditorViewportSettings>()->bEnableActorSnap ? ECheckBoxState::Checked
+																					: ECheckBoxState::Unchecked;
+			}
+		);
+	}
+	FToolMenuEntry ActorSnapping = FToolMenuEntry::InitSubMenu(
+		"ActorSnapping",
+		LOCTEXT("ActorSnapLabel", "Actor"),
+		FLevelEditorCommands::Get().EnableActorSnap->GetDescription(),
+		MakeMenuDelegate,
+		CheckboxMenuAction,
+		EUserInterfaceActionType::ToggleButton
+	);
+
+	return ActorSnapping;
+}
+
+
 } // namespace UE::LevelEditor::Private
 
 namespace UE::LevelEditor
@@ -662,7 +762,11 @@ void PopulateViewModesMenu(UToolMenu* InMenu)
 								if (TSharedPtr<::SLevelViewport> Viewport = WeakViewport.Pin())
 								{
 									FLevelEditorViewportClient& ViewportClient = Viewport->GetLevelViewportClient();
-									return ViewportClient.IsActorColorationVisualizationModeSelected(PrimitiveColorHandler.HandlerName) ? ECheckBoxState::Checked : ECheckBoxState::Unchecked;
+									return ViewportClient.IsActorColorationVisualizationModeSelected(
+											   PrimitiveColorHandler.HandlerName
+										   )
+											 ? ECheckBoxState::Checked
+											 : ECheckBoxState::Unchecked;
 								}
 
 								return ECheckBoxState::Unchecked;
@@ -1086,8 +1190,8 @@ FToolMenuEntry CreateShowHLODsSubmenu()
 									IWorldPartitionEditorModule* WorldPartitionEditorModule =
 										FModuleManager::GetModulePtr<IWorldPartitionEditorModule>("WorldPartitionEditor");
 									return WorldPartitionEditorModule
-										? WorldPartitionEditorModule->GetHLODInEditorMaxDrawDistance()
-										: 0;
+											 ? WorldPartitionEditorModule->GetHLODInEditorMaxDrawDistance()
+											 : 0;
 								};
 
 								// void SLevelViewportToolBar::OnHLODInEditorMaxDrawDistanceValueChanged(double NewValue) const
@@ -1112,8 +1216,7 @@ FToolMenuEntry CreateShowHLODsSubmenu()
 										.ToolTipText(
 											bHLODInEditorAllowed
 												? LOCTEXT(
-													  "HLODsInEditor_MaxDrawDistance_Tooltip",
-													  "Sets the maximum distance at which HLODs will be rendered (0.0 means infinite)"
+													  "HLODsInEditor_MaxDrawDistance_Tooltip", "Sets the maximum distance at which HLODs will be rendered (0.0 means infinite)"
 												  )
 												: HLODInEditorDisallowedReason
 										)
@@ -1534,37 +1637,45 @@ FToolMenuEntry CreateFeatureLevelPreviewSubmenu()
 			{
 				FToolMenuSection& Section =
 					InMenu->AddSection("EditorPreviewMode", LOCTEXT("EditorPreviewModePlatforms", "Preview Platforms"));
-				
+
 				if (FLevelEditorCommands::Get().DisablePlatformPreview.IsValid())
 				{
 					Section.AddMenuEntry(FLevelEditorCommands::Get().DisablePlatformPreview);
 				}
 				Section.AddSeparator("DisablePlatformPreviewSeparator");
 
-				for (auto Iter = FLevelEditorCommands::Get().PlatformToPreviewPlatformOverrides.CreateConstIterator(); Iter; ++Iter)
+				for (auto Iter = FLevelEditorCommands::Get().PlatformToPreviewPlatformOverrides.CreateConstIterator(); Iter;
+					 ++Iter)
 				{
 					FName PlatformName = Iter.Key();
 					const TArray<FLevelEditorCommands::PreviewPlatformCommand>& CommandList = Iter.Value();
-					const TArray<FLevelEditorCommands::PreviewPlatformCommand>* CommandListJson = FLevelEditorCommands::Get().PlatformToPreviewJsonPlatformOverrides.Find(PlatformName);
+					const TArray<FLevelEditorCommands::PreviewPlatformCommand>* CommandListJson =
+						FLevelEditorCommands::Get().PlatformToPreviewJsonPlatformOverrides.Find(PlatformName);
 
 					Section.AddSubMenu(
-						FName(PlatformName), 
-						FText::FromString(PlatformName.ToString()), 
+						FName(PlatformName),
+						FText::FromString(PlatformName.ToString()),
 						FText(),
 						FNewToolMenuDelegate::CreateLambda(
 							[CommandList, CommandListJson](UToolMenu* InSubMenu)
 							{
 								for (const FLevelEditorCommands::PreviewPlatformCommand& Command : CommandList)
 								{
-									FToolMenuSection& Section = InSubMenu->FindOrAddSection(Command.SectionName, FText::Format(LOCTEXT("PreviewJson", "{0}"), FText::FromName(Command.SectionName)));
+									FToolMenuSection& Section = InSubMenu->FindOrAddSection(
+										Command.SectionName,
+										FText::Format(LOCTEXT("PreviewJson", "{0}"), FText::FromName(Command.SectionName))
+									);
 									Section.AddMenuEntry(Command.CommandInfo);
 								}
 
 								if (CommandListJson != nullptr)
 								{
-									FToolMenuSection& SectionJson = InSubMenu->FindOrAddSection("PreviewWithJson", LOCTEXT("PreviewWithJsonLabel", "Preview With Json"));
+									FToolMenuSection& SectionJson = InSubMenu->FindOrAddSection(
+										"PreviewWithJson", LOCTEXT("PreviewWithJsonLabel", "Preview With Json")
+									);
 									TMap<FName, TArray<TSharedPtr<FUICommandInfo>>> SectionNameToCommandList;
-									for (const FLevelEditorCommands::PreviewPlatformCommand& PreviewJsonPlatform : *CommandListJson)
+									for (const FLevelEditorCommands::PreviewPlatformCommand& PreviewJsonPlatform :
+										 *CommandListJson)
 									{
 										if (PreviewJsonPlatform.bIsGeneratingJsonCommand)
 										{
@@ -1572,7 +1683,8 @@ FToolMenuEntry CreateFeatureLevelPreviewSubmenu()
 										}
 										else
 										{
-											SectionNameToCommandList.FindOrAdd(PreviewJsonPlatform.SectionName).Add(PreviewJsonPlatform.CommandInfo);
+											SectionNameToCommandList.FindOrAdd(PreviewJsonPlatform.SectionName)
+												.Add(PreviewJsonPlatform.CommandInfo);
 										}
 									}
 
@@ -1582,8 +1694,12 @@ FToolMenuEntry CreateFeatureLevelPreviewSubmenu()
 										const TArray<TSharedPtr<FUICommandInfo>>& CommandListValue = Iter.Value();
 										SectionJson.AddSubMenu(
 											SectionName,
-											FText::Format(LOCTEXT("PreviewJsonLabel", "Preview {0}"), FText::FromName(SectionName)),
-											FText::Format(LOCTEXT("PreviewJsonTooltip", "Preview {0}"), FText::FromName(SectionName)),
+											FText::Format(
+												LOCTEXT("PreviewJsonLabel", "Preview {0}"), FText::FromName(SectionName)
+											),
+											FText::Format(
+												LOCTEXT("PreviewJsonTooltip", "Preview {0}"), FText::FromName(SectionName)
+											),
 											FNewToolMenuDelegate::CreateLambda(
 												[CommandListValue](UToolMenu* InSubMenu)
 												{
@@ -1592,11 +1708,13 @@ FToolMenuEntry CreateFeatureLevelPreviewSubmenu()
 													{
 														Section.AddMenuEntry(Command);
 													}
-												})
+												}
+											)
 										);
 									}
 								}
-							})
+							}
+						)
 					);
 				}
 			}
@@ -2376,12 +2494,12 @@ void ExtendCameraSubmenu(FName InCameraOptionsSubmenuName)
 					}
 
 					UE::UnrealEd::GenerateViewportTypeMenu(InDynamicMenu);
-
 				}
 
 				// Create Section
 				{
-					FToolMenuSection& CreateSection = InDynamicMenu->FindOrAddSection("Create", LOCTEXT("CreateLabel", "Create"));
+					FToolMenuSection& CreateSection =
+						InDynamicMenu->FindOrAddSection("Create", LOCTEXT("CreateLabel", "Create"));
 
 					CreateSection.AddSubMenu(
 						"CreateCamera",
@@ -2492,8 +2610,7 @@ void ExtendCameraSubmenu(FName InCameraOptionsSubmenuName)
 
 					// This additional options section is used to force certain elements to appear after extensions
 					{
-						FToolMenuSection& AdditionalOptions =
-							InDynamicMenu->FindOrAddSection("AdditionalOptions");
+						FToolMenuSection& AdditionalOptions = InDynamicMenu->FindOrAddSection("AdditionalOptions");
 						AdditionalOptions.AddSeparator("AdditionalOptionsSeparator");
 
 						FToolMenuEntry HighResolutionScreenshot =
@@ -2507,6 +2624,68 @@ void ExtendCameraSubmenu(FName InCameraOptionsSubmenuName)
 	);
 }
 
+void ExtendTransformSubmenu(FName InTransformSubmenuName)
+{
+	UToolMenu* const Submenu = UToolMenus::Get()->ExtendMenu(InTransformSubmenuName);
+	if (!Submenu)
+	{
+		return;
+	}
+
+	// Gizmo
+	{
+		FToolMenuSection& GizmoSection = Submenu->FindOrAddSection("Gizmo", LOCTEXT("GizmoLabel", "Gizmo"));
+		FToolMenuEntry& ShowTransformGizmoEntry = GizmoSection.AddMenuEntry(
+			FLevelEditorCommands::Get().ShowTransformWidget, LOCTEXT("ShowTransformGizmoLabel", "Show Transform Gizmo")
+		);
+
+		ShowTransformGizmoEntry.InsertPosition.Position = EToolMenuInsertType::First;
+	}
+
+	// Selection
+	{
+		FToolMenuSection& SelectionSection = Submenu->FindOrAddSection("Selection", LOCTEXT("SelectionLabel", "Selection"));
+		SelectionSection.AddMenuEntry(
+			FLevelEditorCommands::Get().AllowTranslucentSelection,
+			LOCTEXT("SelectionAllowTranslucentSelectionLabel", "Translucent Objects")
+		);
+		SelectionSection.AddMenuEntry(
+			FLevelEditorCommands::Get().AllowGroupSelection, LOCTEXT("SelectionAllowGroupSelectionLabel", "Select Groups")
+		);
+		SelectionSection.AddMenuEntry(
+			FLevelEditorCommands::Get().StrictBoxSelect, LOCTEXT("SelectionStrictBoxSelectLabel", "Strict Marquee Selection")
+		);
+		SelectionSection.AddMenuEntry(
+			FLevelEditorCommands::Get().TransparentBoxSelect,
+			LOCTEXT("SelectionTransparentBoxSelectLabel", "Marquee Select Occluded")
+		);
+	}
+}
+
+void ExtendSnappingSubmenu(FName InSnappingSubmenuName)
+{
+	if (UToolMenu* const Submenu = UToolMenus::Get()->ExtendMenu(InSnappingSubmenuName))
+	{
+		FToolMenuSection& SnappingSection = Submenu->FindOrAddSection("Snapping");
+
+		// Actor Snapping
+		SnappingSection.AddEntry(UE::LevelEditor::Private::CreateActorSnapCheckboxMenu());
+
+		// Socket Snapping
+		FToolMenuEntry SocketSnapping =
+			FToolMenuEntry::InitMenuEntry(FLevelEditorCommands::Get().ToggleSocketSnapping);
+		SocketSnapping.UserInterfaceActionType = EUserInterfaceActionType::ToggleButton;
+		SocketSnapping.Label = LOCTEXT("SocketSnapLabel", "Socket");
+		SnappingSection.AddEntry(SocketSnapping);
+
+		// Vertex Snapping
+		FToolMenuEntry VertexSnapping = FToolMenuEntry::InitMenuEntry(FLevelEditorCommands::Get().EnableVertexSnap);
+		VertexSnapping.UserInterfaceActionType = EUserInterfaceActionType::ToggleButton;
+		VertexSnapping.Label = LOCTEXT("VertexSnapLabel", "Vertex");
+		SnappingSection.AddEntry(VertexSnapping);
+	}
+}
+
 FToolMenuEntry CreateToolbarCameraSubmenu()
 {
 	FToolMenuEntry Entry = FToolMenuEntry::InitDynamicEntry(
diff --git a/Engine/Source/Editor/LevelEditor/Private/ViewportToolbar/LevelEditorSubmenus.h b/Engine/Source/Editor/LevelEditor/Private/ViewportToolbar/LevelEditorSubmenus.h
--- a/Engine/Source/Editor/LevelEditor/Private/ViewportToolbar/LevelEditorSubmenus.h
+++ b/Engine/Source/Editor/LevelEditor/Private/ViewportToolbar/LevelEditorSubmenus.h
@@ -51,4 +51,6 @@ FToolMenuEntry CreateSettingsSubmenu();
 FToolMenuEntry CreateViewportSizingSubmenu();
 
 void ExtendCameraSubmenu(FName InCameraOptionsSubmenuName);
+void ExtendTransformSubmenu(FName InTransformSubmenuName);
+void ExtendSnappingSubmenu(FName InSnappingSubmenuName);
 } // namespace UE::LevelEditor
diff --git a/Engine/Source/Editor/UnrealEd/Private/ViewportToolbar/UnrealEdViewportToolbar.cpp b/Engine/Source/Editor/UnrealEd/Private/ViewportToolbar/UnrealEdViewportToolbar.cpp
--- a/Engine/Source/Editor/UnrealEd/Private/ViewportToolbar/UnrealEdViewportToolbar.cpp
+++ b/Engine/Source/Editor/UnrealEd/Private/ViewportToolbar/UnrealEdViewportToolbar.cpp
@@ -232,52 +232,6 @@ FToolMenuEntry CreateActorSnapDistanceEntry()
 	return SnapDistance;
 }
 
-FToolMenuEntry CreateActorSnapCheckboxMenu()
-{
-	FNewToolMenuDelegate MakeMenuDelegate = FNewToolMenuDelegate::CreateLambda(
-		[](UToolMenu* Submenu)
-		{
-			// Add "Actor snapping" widget.
-			{
-				FToolMenuSection& ActorSnappingSection =
-					Submenu->FindOrAddSection("ActorSnapping", LOCTEXT("ActorSnappingLabel", "Actor Snapping"));
-				ActorSnappingSection.AddEntry(CreateActorSnapDistanceEntry());
-			}
-		}
-	);
-
-	FToolUIAction CheckboxMenuAction;
-	{
-		CheckboxMenuAction.ExecuteAction = FToolMenuExecuteAction::CreateLambda(
-			[](const FToolMenuContext& InContext)
-			{
-				if (ULevelEditorViewportSettings* Settings = GetMutableDefault<ULevelEditorViewportSettings>())
-				{
-					Settings->bEnableActorSnap = !Settings->bEnableActorSnap;
-				}
-			}
-		);
-		CheckboxMenuAction.GetActionCheckState = FToolMenuGetActionCheckState::CreateLambda(
-			[](const FToolMenuContext& InContext)
-			{
-				return GetDefault<ULevelEditorViewportSettings>()->bEnableActorSnap ? ECheckBoxState::Checked
-																					: ECheckBoxState::Unchecked;
-			}
-		);
-	}
-
-	FToolMenuEntry ActorSnapping = FToolMenuEntry::InitSubMenu(
-		"ActorSnapping",
-		LOCTEXT("ActorSnapLabel", "Actor"),
-		FLevelEditorCommands::Get().EnableActorSnap->GetDescription(),
-		MakeMenuDelegate,
-		CheckboxMenuAction,
-		EUserInterfaceActionType::ToggleButton
-	);
-
-	return ActorSnapping;
-}
-
 FToolMenuEntry CreateLocationSnapCheckboxMenu()
 {
 	const FName LocationSnapName = "LocationSnap";
@@ -762,11 +716,6 @@ FToolMenuEntry CreateTransformsSubmenu()
 				{
 					FToolMenuSection& GizmoSection = Submenu->FindOrAddSection("Gizmo", LOCTEXT("GizmoLabel", "Gizmo"));
 
-					GizmoSection.AddMenuEntry(
-						FLevelEditorCommands::Get().ShowTransformWidget,
-						LOCTEXT("ShowTransformGizmoLabel", "Show Transform Gizmo")
-					);
-
 					TSharedRef<SWidget> GizmoScaleWidget =
 						// clang-format off
 						SNew(SHorizontalBox)
@@ -807,31 +756,6 @@ FToolMenuEntry CreateTransformsSubmenu()
 						"GizmoScale", GizmoScaleWidget, LOCTEXT("GizmoScaleLabel", "Gizmo Scale")
 					));
 				}
-
-				{
-					FToolMenuSection& SelectionSection =
-						Submenu->FindOrAddSection("Selection", LOCTEXT("SelectionLabel", "Selection"));
-
-					SelectionSection.AddMenuEntry(
-						FLevelEditorCommands::Get().AllowTranslucentSelection,
-						LOCTEXT("SelectionAllowTranslucentSelectionLabel", "Translucent Objects")
-					);
-
-					SelectionSection.AddMenuEntry(
-						FLevelEditorCommands::Get().AllowGroupSelection,
-						LOCTEXT("SelectionAllowGroupSelectionLabel", "Select Groups")
-					);
-
-					SelectionSection.AddMenuEntry(
-						FLevelEditorCommands::Get().StrictBoxSelect,
-						LOCTEXT("SelectionStrictBoxSelectLabel", "Strict Marquee Selection")
-					);
-
-					SelectionSection.AddMenuEntry(
-						FLevelEditorCommands::Get().TransparentBoxSelect,
-						LOCTEXT("SelectionTransparentBoxSelectLabel", "Marquee Select Occluded")
-					);
-				}
 			}
 		)
 	);
@@ -1153,18 +1077,6 @@ FToolMenuEntry CreateSnappingSubmenu()
 				SnappingSection.AddEntry(Private::CreateLocationSnapCheckboxMenu()).SetShowInToolbarTopLevel(true);
 				SnappingSection.AddEntry(Private::CreateRotationSnapCheckboxMenu()).SetShowInToolbarTopLevel(true);
 				SnappingSection.AddEntry(Private::CreateScaleSnapCheckboxMenu()).SetShowInToolbarTopLevel(true);
-				SnappingSection.AddEntry(Private::CreateActorSnapCheckboxMenu());
-
-				FToolMenuEntry SocketSnapping =
-					FToolMenuEntry::InitMenuEntry(FLevelEditorCommands::Get().ToggleSocketSnapping);
-				SocketSnapping.UserInterfaceActionType = EUserInterfaceActionType::ToggleButton;
-				SocketSnapping.Label = LOCTEXT("SocketSnapLabel", "Socket");
-				SnappingSection.AddEntry(SocketSnapping);
-
-				FToolMenuEntry VertexSnapping = FToolMenuEntry::InitMenuEntry(FLevelEditorCommands::Get().EnableVertexSnap);
-				VertexSnapping.UserInterfaceActionType = EUserInterfaceActionType::ToggleButton;
-				VertexSnapping.Label = LOCTEXT("VertexSnapLabel", "Vertex");
-				SnappingSection.AddEntry(VertexSnapping);
 			}
 		)
 	);
