Resolve issue where the names of SLevelViewport::GetToolbarVisibility and SLevelViewport::GetToolBarVisibility only differed by case by deprecating SLevelViewport::GetToolbarVisibility and replacing it with SLevelViewport::GetOptionsMenuVisibility

diff --git a/Engine/Source/Editor/LevelEditor/Private/SLevelViewport.cpp b/Engine/Source/Editor/LevelEditor/Private/SLevelViewport.cpp
--- a/Engine/Source/Editor/LevelEditor/Private/SLevelViewport.cpp
+++ b/Engine/Source/Editor/LevelEditor/Private/SLevelViewport.cpp
@@ -4967,10 +4967,14 @@ void SLevelViewport::OnFloatingButtonClicked()
 	LevelViewportClient->SetLastKeyViewport();
 }
 
-/** Get the visibility for viewport toolbar */
 EVisibility  SLevelViewport::GetToolbarVisibility() const
-{ 
-	return bShowToolbarAndControls ? EVisibility::Visible : EVisibility::Collapsed; 
+{
+	return GetOptionsMenuVisibility();
+}
+
+EVisibility SLevelViewport::GetOptionsMenuVisibility() const
+{
+	return bShowToolbarAndControls ? EVisibility::Visible : EVisibility::Collapsed;
 }
 
 void SLevelViewport::RemoveAllPreviews(const bool bRemoveFromDesktopViewport /**= true*/)
diff --git a/Engine/Source/Editor/LevelEditor/Private/SLevelViewportToolBar.cpp b/Engine/Source/Editor/LevelEditor/Private/SLevelViewportToolBar.cpp
--- a/Engine/Source/Editor/LevelEditor/Private/SLevelViewportToolBar.cpp
+++ b/Engine/Source/Editor/LevelEditor/Private/SLevelViewportToolBar.cpp
@@ -124,6 +124,7 @@ void SLevelViewportToolBar::Construct( const FArguments& InArgs )
 
 	const FMargin ToolbarSlotPadding(4.0f, 1.0f);
 
+	// clang-format off
 	ChildSlot
 	[
 		SNew(SBorder)
@@ -142,7 +143,7 @@ void SLevelViewportToolBar::Construct( const FArguments& InArgs )
 				[
 					SNew( SEditorViewportToolbarMenu )
 					.ParentToolBar( SharedThis( this ) )
-					.Visibility(Viewport.Pin().Get(), &SLevelViewport::GetToolbarVisibility)
+					.Visibility(Viewport.Pin().Get(), &SLevelViewport::GetOptionsMenuVisibility)
 					.Image("EditorViewportToolBar.OptionsDropdown")
 					.AddMetaData<FTagMetaData>(FTagMetaData(TEXT("EditorViewportToolBar.MenuDropdown")))
 					.OnGetMenuContent( this, &SLevelViewportToolBar::GenerateOptionsMenu )
@@ -333,7 +334,8 @@ void SLevelViewportToolBar::Construct( const FArguments& InArgs )
 			]
 		]
 	];
-	
+	// clang-format on
+
 	SViewportToolBar::Construct(SViewportToolBar::FArguments());
 }
 
diff --git a/Engine/Source/Editor/LevelEditor/Public/SLevelViewport.h b/Engine/Source/Editor/LevelEditor/Public/SLevelViewport.h
--- a/Engine/Source/Editor/LevelEditor/Public/SLevelViewport.h
+++ b/Engine/Source/Editor/LevelEditor/Public/SLevelViewport.h
@@ -401,9 +401,12 @@ public:
 	/** Actions to perform whenever the viewports floating buttons are pressed */
 	void OnFloatingButtonClicked();
 
-	/** Get the visibility for viewport toolbar */
+	UE_DEPRECATED(5.6, "Call GetOptionsMenuVisibility() instead.")
 	EVisibility GetToolbarVisibility() const;
 
+	/** Get the visibility for viewport toolbar's options menu. */
+	EVisibility GetOptionsMenuVisibility() const;
+
 	/** Get the visibility for items considered to be part of the 'full' viewport toolbar */
 	EVisibility GetFullToolbarVisibility() const { return (bShowToolbarAndControls && bShowFullToolbar) ? EVisibility::SelfHitTestInvisible : EVisibility::Collapsed; }
 
