[Level Editor] Clean up deprecated and not implemented functions

diff --git a/Engine/Source/Editor/LevelEditor/Private/SLevelEditor.cpp b/Engine/Source/Editor/LevelEditor/Private/SLevelEditor.cpp
--- a/Engine/Source/Editor/LevelEditor/Private/SLevelEditor.cpp
+++ b/Engine/Source/Editor/LevelEditor/Private/SLevelEditor.cpp
@@ -768,21 +768,6 @@ TSharedPtr<ISceneOutliner> SLevelEditor::GetMostRecentlyUsedSceneOutliner()
 	return SceneOutlinerPtr.Pin();
 }
 
-TSharedPtr<ISceneOutliner> SLevelEditor::GetSceneOutliner() const
-{
-	// Simply pick the first available outliner
-	for(const TPair<FName, TWeakPtr<ISceneOutliner>>& Outliner : SceneOutliners)
-	{
-		if(TSharedPtr<ISceneOutliner> OutlinerPin = Outliner.Value.Pin())
-		{
-			return OutlinerPin;
-		}
-	}
-
-	// No outliners are open currently
-	return nullptr;
-}
-
 TSharedRef<SDockTab> SLevelEditor::SummonDetailsPanel( FName TabIdentifier )
 {
 	TSharedRef<SActorDetails> ActorDetails = StaticCastSharedRef<SActorDetails>( CreateActorDetails( TabIdentifier ) );
@@ -1947,11 +1932,6 @@ TSharedPtr<SLevelViewport> SLevelEditor::GetActiveViewportInterface()
 	return GetActiveViewport();
 }
 
-TSharedPtr<FAssetThumbnailPool> SLevelEditor::GetThumbnailPool() const
-{
-	return UThumbnailManager::Get().GetSharedThumbnailPool();
-}
-
 void SLevelEditor::AppendCommands( const TSharedRef<FUICommandList>& InCommandsToAppend )
 {
 	LevelEditorCommands->Append(InCommandsToAppend);
diff --git a/Engine/Source/Editor/LevelEditor/Private/SLevelEditor.h b/Engine/Source/Editor/LevelEditor/Private/SLevelEditor.h
--- a/Engine/Source/Editor/LevelEditor/Private/SLevelEditor.h
+++ b/Engine/Source/Editor/LevelEditor/Private/SLevelEditor.h
@@ -77,9 +77,6 @@ public:
 	virtual const TArray< TSharedPtr< class IToolkit > >& GetHostedToolkits() const override;
 	virtual TArray< TSharedPtr< SLevelViewport > > GetViewports() const override;
 	virtual TSharedPtr<SLevelViewport> GetActiveViewportInterface() override;
-	PRAGMA_DISABLE_DEPRECATION_WARNINGS
-	virtual TSharedPtr< class FAssetThumbnailPool > GetThumbnailPool() const override;
-	PRAGMA_ENABLE_DEPRECATION_WARNINGS
 	virtual void AppendCommands( const TSharedRef<FUICommandList>& InCommandsToAppend ) override;
 	virtual void AddStandaloneLevelViewport( const TSharedRef<SLevelViewport>& LevelViewport ) override;
 
@@ -114,8 +111,6 @@ public:
 	 */
 	virtual FReply OnKeyDownInViewport( const FGeometry& MyGeometry, const FKeyEvent& InKeyEvent ) override;
 
-	bool CanCloseApp();
-
 	/** Returns the full action list for this level editor instance */
 	virtual const TSharedPtr< FUICommandList >& GetLevelEditorActions() const override
 	{
@@ -162,10 +157,6 @@ public:
 	
 	/** Return the most recently interacted with Outliner */
 	virtual TSharedPtr<ISceneOutliner> GetMostRecentlyUsedSceneOutliner() override;
-	
-	/** Return the most recently interacted with Outliner */
-	UE_DEPRECATED(5.1, "The Level Editor has multiple outliners, use GetAllSceneOutliners() or GetMostRecentlyUsedSceneOutliner() instead to avoid ambiguity")
-	virtual TSharedPtr<ISceneOutliner> GetSceneOutliner() const override;
 
 	virtual TWeakPtr<UE::WorldHierarchy::IWorldHierarchy> GetWorldHierarchy() override { return WorldHierarchy; }
 	
@@ -220,9 +211,6 @@ private:
 	 */
 	virtual FReply OnKeyDown( const FGeometry& MyGeometry, const FKeyEvent& InKeyEvent ) override;
 
-	/** Callback for when the property view changes */
-	void OnPropertyObjectArrayChanged(const FString& NewTitle, const TArray< UObject* >& UObjects );
-
 	/** Callback for when the level editor layout has changed */
 	void OnLayoutHasChanged();
 	
diff --git a/Engine/Source/Editor/LevelEditor/Public/ILevelEditor.h b/Engine/Source/Editor/LevelEditor/Public/ILevelEditor.h
--- a/Engine/Source/Editor/LevelEditor/Public/ILevelEditor.h
+++ b/Engine/Source/Editor/LevelEditor/Public/ILevelEditor.h
@@ -37,17 +37,6 @@ public:
 	/** Summons a context menu for this level editor at the mouse cursor's location */
 	virtual void SummonLevelViewportContextMenu(const FTypedElementHandle& HitProxyElement = FTypedElementHandle()) = 0;
 
-	UE_DEPRECATED(5.0, "Use the SummonLevelViewportContextMenu overload which takes a FTypedElementHandle.")
-	void SummonLevelViewportContextMenu(AActor* HitProxyActor)
-	{
-		FTypedElementHandle HitProxyElement;
-		if (HitProxyActor)
-		{
-			HitProxyElement = UEngineElementsLibrary::AcquireEditorActorElementHandle(HitProxyActor);
-		}
-		SummonLevelViewportContextMenu(HitProxyElement);
-	}
-
 	/** Gets the title for the context menu for this level editor */
 	virtual FText GetLevelViewportContextMenuTitle() const = 0;
 
@@ -63,10 +52,6 @@ public:
 	/** Gets the active level viewport for this level editor */
 	virtual TSharedPtr<SLevelViewport> GetActiveViewportInterface() = 0;
 
-	/** Get the thumbnail pool used by this level editor */
-	UE_DEPRECATED(5.0, "GetThumbnailPool has been replaced by UThumbnailManager::Get().GetSharedThumbnailPool().")
-	virtual TSharedPtr< class FAssetThumbnailPool > GetThumbnailPool() const = 0;
-
 	/** Access the level editor's action command list */
 	virtual const TSharedPtr< FUICommandList >& GetLevelEditorActions() const = 0;
 
@@ -95,10 +80,6 @@ public:
 	UE_DEPRECATED(5.5, "Use AddActorDetailsSCSEditorUICustomization/RemoveActorDetailsSCSEditorUICustomization instead.")
 	virtual void SetActorDetailsSCSEditorUICustomization(TSharedPtr<class ISCSEditorUICustomization> ActorDetailsSCSEditorUICustomization) {};
 
-	/** Return the most recently interacted with Outliner */
-	UE_DEPRECATED(5.1, "The Level Editor has multiple outliners, use GetAllSceneOutliners() or GetMostRecentlyUsedSceneOutliner() instead to avoid ambiguity")
-	virtual TSharedPtr<ISceneOutliner> GetSceneOutliner() const = 0;
-
 	/** Get an array containing weak pointers to all 4 Scene Outliners which could be potentially active */
 	virtual TArray<TWeakPtr<ISceneOutliner>> GetAllSceneOutliners() const = 0;
 
