Make the viewport warnings visible when the new viewport toolbar is enabled

This adds the SActionableMessageViewportWidget to the level editor viewport via an overlay slot in SLevelViewport instead of by adding it to the old viewport toolbar (SLevelViewportToolbar). This allows the viewport warnings (displayed by SActionableMessageViewportWidget) to be visible when the old viewport toolbar is hidden and the new one displayed (CVAR ToolMenusViewportToolbars set to 2).

The new placement still respects the "Show Toolbar" and "Hide Viewport UI" toggles.

diff --git a/Engine/Source/Editor/LevelEditor/Private/SLevelViewport.cpp b/Engine/Source/Editor/LevelEditor/Private/SLevelViewport.cpp
--- a/Engine/Source/Editor/LevelEditor/Private/SLevelViewport.cpp
+++ b/Engine/Source/Editor/LevelEditor/Private/SLevelViewport.cpp
@@ -103,6 +103,7 @@
 #include "EditorViewportTabContent.h"
 #include "EditorViewportCommands.h"
 #include "FunctionalUIScreenshotTest.h"
+#include "SActionableMessageViewportWidget.h"
 #include "SkeletalRenderPublic.h"
 #include "ViewportToolbar/LevelEditorSubmenus.h"
 #include "ViewportToolbar/LevelViewportContext.h"
@@ -376,6 +377,47 @@ void SLevelViewport::ConstructViewportOverlayContent()
 		++SlotIndex;
 #endif
 
+	// Add widget to show viewport warnings from the ActionableMessageSubsystem.
+	{
+		TAttribute<FMargin> Padding = TAttribute<FMargin>::CreateLambda(
+			[]()
+			{
+				const float TopPadding = UE::UnrealEd::ShowOldViewportToolbars() ? 37.0f : 8.0f;
+				return FMargin(0.f, TopPadding, 10.f, 0.f);
+			}
+		);
+
+		// clang-format off
+		ViewportOverlay->AddSlot(SlotIndex)
+		.Padding(Padding)
+		.VAlign(VAlign_Top)
+		.HAlign(HAlign_Right)
+		[
+			SNew(SBorder)
+			.BorderImage(FAppStyle::Get().GetBrush("ActionableMessage.Border"))
+			[
+				SAssignNew(ActionableMessageViewportWidget, SActionableMessageViewportWidget)
+				.Visibility_Lambda(
+					[this]()
+					{
+						if (!ShouldShowFullToolbar())
+						{
+							return EVisibility::Collapsed;
+						}
+						
+						if (GetToolBarVisibility() == EVisibility::Collapsed)
+						{
+							return EVisibility::Collapsed;
+						}
+
+						return ActionableMessageViewportWidget->GetVisibility();
+					}
+				)
+			]
+		];
+		// clang-format on
+	}
+
 	ViewportOverlay->AddSlot( SlotIndex )
 	.HAlign(HAlign_Right)
 	.Padding(0.0f, 0.0f, 0.0f, 35.0f)
diff --git a/Engine/Source/Editor/LevelEditor/Private/SLevelViewportToolBar.cpp b/Engine/Source/Editor/LevelEditor/Private/SLevelViewportToolBar.cpp
--- a/Engine/Source/Editor/LevelEditor/Private/SLevelViewportToolBar.cpp
+++ b/Engine/Source/Editor/LevelEditor/Private/SLevelViewportToolBar.cpp
@@ -55,7 +55,6 @@
 #include "SortHelper.h"
 #include "Interfaces/IMainFrameModule.h"
 #include "SCommonEditorViewportToolbarBase.h"
-#include "SActionableMessageViewportWidget.h"
 #include "Engine/SceneCapture.h"
 #include "ViewportToolbar/LevelEditorSubmenus.h"
 #include "ViewportToolbar/LevelViewportContext.h"
@@ -316,22 +315,6 @@ void SLevelViewportToolBar::Construct( const FArguments& InArgs )
 					]
 				]
 			]
-			+ SVerticalBox::Slot()
-			.AutoHeight()
-			.Padding(0.f, 8.f, 10.f, 0.f)
-			.VAlign(VAlign_Center)
-			.HAlign(HAlign_Right)
-			[
-				SNew(SBorder)
-				.BorderImage(FAppStyle::Get().GetBrush("ActionableMessage.Border"))
-				[
-					SAssignNew(ActionableMessageViewportWidget, SActionableMessageViewportWidget)
-					.Visibility_Lambda([this]()
-					{
-						return ActionableMessageViewportWidget->GetVisibility();
-					})
-				]
-			]
 		]
 	];
 	// clang-format on
diff --git a/Engine/Source/Editor/LevelEditor/Private/SLevelViewportToolBar.h b/Engine/Source/Editor/LevelEditor/Private/SLevelViewportToolBar.h
--- a/Engine/Source/Editor/LevelEditor/Private/SLevelViewportToolBar.h
+++ b/Engine/Source/Editor/LevelEditor/Private/SLevelViewportToolBar.h
@@ -11,7 +11,6 @@
 
 class ACameraActor;
 class SExtensionPanel;
-class SActionableMessageViewportWidget;
 class STransformViewportToolBar;
 class FMenuBuilder;
 class UToolMenu;
@@ -251,9 +250,6 @@ private:
 	/** STransformViewportToolBar menu */
 	TSharedPtr<STransformViewportToolBar> TransformToolbar;
 
-	/** Viewport widget for warning messages */
-	TSharedPtr<SActionableMessageViewportWidget> ActionableMessageViewportWidget;
-
 	/** The previous max STransformViewportToolBar width to allow deterministic size calculations */
 	mutable float TransformToolbar_CachedMaxWidth = 0.0f;
 };
diff --git a/Engine/Source/Editor/LevelEditor/Public/SLevelViewport.h b/Engine/Source/Editor/LevelEditor/Public/SLevelViewport.h
--- a/Engine/Source/Editor/LevelEditor/Public/SLevelViewport.h
+++ b/Engine/Source/Editor/LevelEditor/Public/SLevelViewport.h
@@ -26,6 +26,7 @@ class FLevelViewportLayout;
 class FSceneViewport;
 class FUICommandList;
 class ILevelEditor;
+class SActionableMessageViewportWidget;
 class SActorPreview;
 class SCaptureRegionWidget;
 class SGameLayerManager;
@@ -970,6 +971,9 @@ private:
 
 	TSharedPtr<SWorldPartitionViewportWidget> WorldPartitionViewportWidget;
 
+	/** Viewport widget for warning messages */
+	TSharedPtr<SActionableMessageViewportWidget> ActionableMessageViewportWidget;
+
 protected:
 	void LockActorInternal(AActor* NewActorToLock);
 
