Sequencer: Add a tooltip for the Sequencer tab

#rnx

diff --git a/Engine/Source/Editor/LevelEditor/Private/SLevelEditor.cpp b/Engine/Source/Editor/LevelEditor/Private/SLevelEditor.cpp
--- a/Engine/Source/Editor/LevelEditor/Private/SLevelEditor.cpp
+++ b/Engine/Source/Editor/LevelEditor/Private/SLevelEditor.cpp
@@ -1531,6 +1531,7 @@ TSharedRef<SWidget> SLevelEditor::RestoreContentArea( const TSharedRef<SDockTab>
 				FOnSpawnTab::CreateSP(this, &SLevelEditor::SpawnLevelEditorTab, LevelEditorTabIds::Sequencer, FString()),
 				FCanSpawnTab::CreateSP(this, &SLevelEditor::CanSpawnLevelEditorTab, LevelEditorTabIds::Sequencer))
 				.SetDisplayName(NSLOCTEXT("LevelEditorTabs", "Sequencer", "Sequencer"))
+				.SetTooltipText(NSLOCTEXT("LevelEditorTabs", "SequencerTooltipText", "Open the Sequencer tab to animate characters, cameras, and properties."))
 				.SetGroup( MenuStructure.GetLevelEditorCinematicsCategory() )
 				.SetIcon( FSlateIcon(FAppStyle::GetAppStyleSetName(), "LevelEditor.Tabs.Cinematics") );
 		}
