Disable preview on MacOS using SM6 due to issues with bindless/bindful

diff --git a/Engine/Source/Editor/LevelEditor/Private/LevelEditorActions.cpp b/Engine/Source/Editor/LevelEditor/Private/LevelEditorActions.cpp
--- a/Engine/Source/Editor/LevelEditor/Private/LevelEditorActions.cpp
+++ b/Engine/Source/Editor/LevelEditor/Private/LevelEditorActions.cpp
@@ -810,6 +810,12 @@ void FLevelEditorActionCallbacks::SetPreviewPlatform(FPreviewPlatformInfo NewPre
 
 bool FLevelEditorActionCallbacks::CanExecutePreviewPlatform(FPreviewPlatformInfo NewPreviewPlatform)
 {
+	// Temporary - We have disable SM6 preview for platforms < SM6 for now as it causes crashes due to a bindful/bindless mismatch
+	if (IsMetalSM6Platform(GMaxRHIShaderPlatform))
+	{
+		return false;
+	}
+	
 	if (NewPreviewPlatform.PreviewFeatureLevel > GMaxRHIFeatureLevel)
 	{
 		return false;
@@ -823,7 +829,6 @@ bool FLevelEditorActionCallbacks::CanExecutePreviewPlatform(FPreviewPlatformInfo
 		// for > 16 samplers we rely on the shader compiler being able to choose an appropriate profile for the
 		// preview feature level that supports > 16 samplers. On D3D12 SM5 the D3D shader compiler will use Dxc and
 		// sm6.0. Vulkan SM5 also appears to handle > 16 samplers fine.
-		// TODO: Look into support for Metal (MacOS).
 		if ((int32)FDataDrivenShaderPlatformInfo::GetMaxSamplers(PreviewShaderPlatform) > GRHIGlobals.MaxTextureSamplers)
 		{
 			return false;
