[Viewport Toolbar] Copy "level editor toolbar > Settings > Show subcomponents" to the viewport toolbar "Transform" menu

diff --git a/Engine/Source/Editor/LevelEditor/Private/ViewportToolbar/LevelEditorSubmenus.cpp b/Engine/Source/Editor/LevelEditor/Private/ViewportToolbar/LevelEditorSubmenus.cpp
--- a/Engine/Source/Editor/LevelEditor/Private/ViewportToolbar/LevelEditorSubmenus.cpp
+++ b/Engine/Source/Editor/LevelEditor/Private/ViewportToolbar/LevelEditorSubmenus.cpp
@@ -2631,6 +2631,7 @@ void ExtendTransformSubmenu(FName InTransformSubmenuName)
 		SelectionSection.AddMenuEntry(FLevelEditorCommands::Get().AllowGroupSelection);
 		SelectionSection.AddMenuEntry(FLevelEditorCommands::Get().StrictBoxSelect);
 		SelectionSection.AddMenuEntry(FLevelEditorCommands::Get().TransparentBoxSelect);
+		SelectionSection.AddMenuEntry(FLevelEditorCommands::Get().ShowSelectionSubcomponents);
 	}
 }
 
