World Partition - HLOD: Added the HLOD viewport options to the new viewport toolbar
* Can be found in the "Viewport-related settings" menu in the top right of the viewport toolbar

diff --git a/Engine/Source/Editor/LevelEditor/Private/ViewportToolbar/LevelEditorSubmenus.cpp b/Engine/Source/Editor/LevelEditor/Private/ViewportToolbar/LevelEditorSubmenus.cpp
--- a/Engine/Source/Editor/LevelEditor/Private/ViewportToolbar/LevelEditorSubmenus.cpp
+++ b/Engine/Source/Editor/LevelEditor/Private/ViewportToolbar/LevelEditorSubmenus.cpp
@@ -1980,10 +1980,14 @@ FToolMenuEntry CreateSettingsSubmenu()
 					FToolMenuSection& SettingsSection =
 						Submenu->FindOrAddSection("Settings", LOCTEXT("SettingsSectionLabel", "Settings"));
 
+					// Volume
 					SettingsSection.AddEntry(FToolMenuEntry::InitWidget(
 						"Volume", BuildVolumeControlCustomWidget(), LOCTEXT("VolumeControlLabel", "Volume")
 					));
 
+					// HLODs
+					SettingsSection.AddEntry(CreateShowHLODsSubmenu());
+
 					SettingsSection.AddSeparator("AdvancedSeparator");
 
 					{
