Fix a bug that caused "Show > HLODs" from not appearing in the new level editor viewport toolbar

diff --git a/Engine/Source/Editor/LevelEditor/Private/ViewportToolbar/LevelEditorSubmenus.cpp b/Engine/Source/Editor/LevelEditor/Private/ViewportToolbar/LevelEditorSubmenus.cpp
--- a/Engine/Source/Editor/LevelEditor/Private/ViewportToolbar/LevelEditorSubmenus.cpp
+++ b/Engine/Source/Editor/LevelEditor/Private/ViewportToolbar/LevelEditorSubmenus.cpp
@@ -1286,7 +1286,7 @@ FToolMenuEntry CreateShowLayersSubmenu()
 	// This is a dynamic entry so we can skip adding the submenu if the context
 	// indicates that the viewport's world is partitioned.
 	return FToolMenuEntry::InitDynamicEntry(
-		"ShowHLODsDynamic",
+		"ShowLayersDynamic",
 		FNewToolMenuSectionDelegate::CreateLambda(
 			[](FToolMenuSection& InDynamicSection) -> void
 			{
