Remove "Enable" prefix and "Snapping" suffix from Actor, Socket, and Vertex snapping command names

diff --git a/Engine/Source/Editor/LevelEditor/Private/LevelEditorActions.cpp b/Engine/Source/Editor/LevelEditor/Private/LevelEditorActions.cpp
--- a/Engine/Source/Editor/LevelEditor/Private/LevelEditorActions.cpp
+++ b/Engine/Source/Editor/LevelEditor/Private/LevelEditorActions.cpp
@@ -4091,14 +4091,14 @@ void FLevelEditorCommands::RegisterCommands()
 	
 	UI_COMMAND( DrawBrushMarkerPolys, "Draw Brush Polys", "Draws semi-transparent polygons around a brush when selected", EUserInterfaceActionType::ToggleButton, FInputChord() );
 	UI_COMMAND( OnlyLoadVisibleInPIE, "Only Load Visible Levels in Game Preview", "If enabled, when game preview starts, only visible levels will be loaded", EUserInterfaceActionType::ToggleButton, FInputChord() );
-	UI_COMMAND( ToggleSocketSnapping, "Enable Socket Snapping", "Enables or disables snapping to sockets", EUserInterfaceActionType::ToggleButton, FInputChord() ); 
+	UI_COMMAND( ToggleSocketSnapping, "Socket", "Enables or disables snapping to sockets", EUserInterfaceActionType::ToggleButton, FInputChord() );
 	UI_COMMAND( ToggleParticleSystemLOD, "Enable Particle System LOD Switching", "If enabled particle systems will use distance LOD switching in perspective viewports", EUserInterfaceActionType::ToggleButton, FInputChord() );
 	UI_COMMAND( ToggleFreezeParticleSimulation, "Freeze Particle Simulation", "If enabled particle systems will freeze their simulation state", EUserInterfaceActionType::ToggleButton, FInputChord() );
 	UI_COMMAND( ToggleParticleSystemHelpers, "Toggle Particle System Helpers", "Toggles showing particle system helper widgets in viewports", EUserInterfaceActionType::ToggleButton, FInputChord() );
 	UI_COMMAND( ToggleLODViewLocking, "Enable LOD View Locking", "If enabled viewports of the same type will use the same LOD", EUserInterfaceActionType::ToggleButton, FInputChord() );
 	UI_COMMAND( LevelStreamingVolumePrevis, "Enable Automatic Level Streaming", "If enabled, the viewport will stream in levels automatically when the camera is moved", EUserInterfaceActionType::ToggleButton, FInputChord() );
-	UI_COMMAND( EnableActorSnap, "Enable Actor Snapping", "If enabled, actors will snap to the location of other actors when they are within distance", EUserInterfaceActionType::ToggleButton, FInputChord(EModifierKey::Control|EModifierKey::Shift, EKeys::K) );
-	UI_COMMAND( EnableVertexSnap, "Enable Vertex Snapping","If enabled, actors will snap to the location of the nearest vertex on another actor in the direction of movement", EUserInterfaceActionType::ToggleButton, FInputChord() );
+	UI_COMMAND( EnableActorSnap, "Actor", "If enabled, actors will snap to the location of other actors when they are within distance", EUserInterfaceActionType::ToggleButton, FInputChord(EModifierKey::Control | EModifierKey::Shift, EKeys::K) );
+	UI_COMMAND( EnableVertexSnap, "Vertex", "If enabled, actors will snap to the location of the nearest vertex on another actor in the direction of " "movement", EUserInterfaceActionType::ToggleButton, FInputChord() );
 	UI_COMMAND( ToggleHideViewportUI, "Hide Viewport UI", "Toggles hidden viewport UI mode. Hides all overlaid viewport UI widgets", EUserInterfaceActionType::ToggleButton, FInputChord() );
 
 	//if (FParse::Param( FCommandLine::Get(), TEXT( "editortoolbox" ) ))
