Add CVAR LevelEditorToolbarSettings

The new CVAR controls whether the level editor toolbar "Settings" menu is shown. The default is to show it (value 1). Set it to 0 to hide the "Settings" menu.

diff --git a/Engine/Source/Editor/LevelEditor/Private/LevelEditorToolBar.cpp b/Engine/Source/Editor/LevelEditor/Private/LevelEditorToolBar.cpp
--- a/Engine/Source/Editor/LevelEditor/Private/LevelEditorToolBar.cpp
+++ b/Engine/Source/Editor/LevelEditor/Private/LevelEditorToolBar.cpp
@@ -65,6 +65,21 @@
 #include "DataDrivenShaderPlatformInfo.h"
 #include "ViewportToolbar/LevelEditorSubmenus.h"
 
+namespace UE::Editor::Private
+{
+
+int32 CVarLevelEditorToolbarSettingsValue = 1;
+
+}
+
+static FAutoConsoleVariableRef CVarToolMenusViewportToolbars(
+	TEXT("LevelEditorToolbarSettings"),
+	UE::Editor::Private::CVarLevelEditorToolbarSettingsValue,
+	TEXT("Control whether the level editor Settings menu is visible in the toolbar. Set to 1 (default) to show it and "
+		 "0 to hide it."),
+	ECVF_Default
+);
+
 FName FLevelEditorToolBar::SecondaryModeToolbarName("LevelEditor.SecondaryToolbar");
 
 namespace PreviewModeFunctionality
@@ -1515,16 +1530,27 @@ void FLevelEditorToolBar::RegisterLevelEditorToolBar( const TSharedRef<FUIComman
 	SettingsToolbar->StyleName = "AssetEditorToolbar";
 	{
 		FToolMenuSection& SettingsSection = SettingsToolbar->AddSection("ProjectSettings");
-		FToolMenuEntry SettingsEntry =
-			FToolMenuEntry::InitComboButton(
-				"LevelToolbarQuickSettings",
-				FUIAction(),
-				FOnGetContent::CreateStatic(&FLevelEditorToolBar::GenerateQuickSettingsMenu, InCommandList, TWeakPtr<SLevelEditor>(InLevelEditor)),
-				LOCTEXT("QuickSettingsCombo", "Settings"),
-				LOCTEXT("QuickSettingsCombo_ToolTip", "Project and Editor settings"),
-				FSlateIcon(FAppStyle::GetAppStyleSetName(), "LevelEditor.GameSettings"),
-				false,
-				"LevelToolbarQuickSettings");
+
+		FUIAction Action;
+		Action.IsActionVisibleDelegate = FIsActionButtonVisible::CreateLambda(
+			[]() -> bool
+			{
+				return UE::Editor::Private::CVarLevelEditorToolbarSettingsValue != 0;
+			}
+		);
+
+		FToolMenuEntry SettingsEntry = FToolMenuEntry::InitComboButton(
+			"LevelToolbarQuickSettings",
+			Action,
+			FOnGetContent::CreateStatic(
+				&FLevelEditorToolBar::GenerateQuickSettingsMenu, InCommandList, TWeakPtr<SLevelEditor>(InLevelEditor)
+			),
+			LOCTEXT("QuickSettingsCombo", "Settings"),
+			LOCTEXT("QuickSettingsCombo_ToolTip", "Project and Editor settings"),
+			FSlateIcon(FAppStyle::GetAppStyleSetName(), "LevelEditor.GameSettings"),
+			false,
+			"LevelToolbarQuickSettings"
+		);
 		SettingsEntry.StyleNameOverride = "CalloutToolbar";
 
 		SettingsSection.AddEntry(SettingsEntry);
