Set icon on "Camera > Cameras" menu in the new viewport toolbar

You see this menu once you have 10 or more cameras in your level.

diff --git a/Engine/Source/Editor/LevelEditor/Private/ViewportToolbar/LevelEditorSubmenus.cpp b/Engine/Source/Editor/LevelEditor/Private/ViewportToolbar/LevelEditorSubmenus.cpp
--- a/Engine/Source/Editor/LevelEditor/Private/ViewportToolbar/LevelEditorSubmenus.cpp
+++ b/Engine/Source/Editor/LevelEditor/Private/ViewportToolbar/LevelEditorSubmenus.cpp
@@ -2373,7 +2373,7 @@ void AddCameraActorSelectSection(UToolMenu* InMenu)
 	constexpr uint32 MaxCamerasInTopLevelMenu = 10;
 	if (LookThroughActors.Num() > MaxCamerasInTopLevelMenu)
 	{
-		Section.AddSubMenu(
+		FToolMenuEntry& Entry = Section.AddSubMenu(
 			"CameraActors",
 			CameraActorsHeading,
 			LOCTEXT("LookThroughPlacedCameras_ToolTip", "Look through and pilot placed cameras"),
@@ -2388,6 +2388,7 @@ void AddCameraActorSelectSection(UToolMenu* InMenu)
 				}
 			)
 		);
+		Entry.Icon = FSlateIconFinder::FindIconForClass(ACameraActor::StaticClass());
 	}
 	else if (!LookThroughActors.IsEmpty())
 	{
