[Viewport Toolbar] Copy "level editor toolbar > Settings > Hide Viewport UI" to the main menu's "Window" menu

diff --git a/Engine/Source/Editor/LevelEditor/Private/LevelEditorMenu.cpp b/Engine/Source/Editor/LevelEditor/Private/LevelEditorMenu.cpp
--- a/Engine/Source/Editor/LevelEditor/Private/LevelEditorMenu.cpp
+++ b/Engine/Source/Editor/LevelEditor/Private/LevelEditorMenu.cpp
@@ -243,6 +243,25 @@ void FLevelEditorMenu::RegisterLevelEditorMenus()
 			}
 		}
 
+		static void ExtendWindowMenu()
+		{
+			if (UToolMenu* Menu = UToolMenus::Get()->ExtendMenu("MainFrame.MainMenu.Window"))
+			{
+				// This entry needs to be placed after the fullscreen separator in the Window menu.
+				const FName FullscreenSeparatorName = "FullscreenSeparator";
+
+				FToolMenuSection& LayoutSection = Menu->FindOrAddSection("WindowLayout");
+
+				// Making sure the separator exist - the one in Window menu might be missing if not on a Win build.
+				// See FMainMenu::RegisterWindowMenu().
+				LayoutSection.AddSeparator(FullscreenSeparatorName);
+
+				FToolMenuEntry& ToggleUIEntry = LayoutSection.AddMenuEntry(FLevelEditorCommands::Get().ToggleHideViewportUI);
+				ToggleUIEntry.InsertPosition.Position = EToolMenuInsertType::After;
+				ToggleUIEntry.InsertPosition.Name = FullscreenSeparatorName;
+			}
+		}
+
 		static void ExtendMenuBar()
 		{
 			UToolMenu* Menu = UToolMenus::Get()->ExtendMenu("LevelEditor.MainMenu");
@@ -324,6 +343,8 @@ void FLevelEditorMenu::RegisterLevelEditorMenus()
 	// Extend the Help menu
 	Local::ExtendHelpMenu();
 
+	Local::ExtendWindowMenu();
+
 	// Extend the Build menu
 	RegisterBuildMenu();
 
diff --git a/Engine/Source/Editor/MainFrame/Private/Menus/MainMenu.cpp b/Engine/Source/Editor/MainFrame/Private/Menus/MainMenu.cpp
--- a/Engine/Source/Editor/MainFrame/Private/Menus/MainMenu.cpp
+++ b/Engine/Source/Editor/MainFrame/Private/Menus/MainMenu.cpp
@@ -370,10 +370,15 @@ void FMainMenu::RegisterWindowMenu()
 
 		// Enable Fullscreen section
 #if !PLATFORM_MAC && !PLATFORM_LINUX // On Mac/Linux windowed fullscreen mode in the editor is currently unavailable
+
+		const FName FullscreenSeparatorName = "FullscreenSeparator";
 		// Separator
-		Section.AddSeparator("FullscreenSeparator");
+		Section.AddSeparator(FullscreenSeparatorName);
+
 		// Fullscreen
-		Section.AddMenuEntry(FMainFrameCommands::Get().ToggleFullscreen);
+		FToolMenuEntry& FullscreenEntry = Section.AddMenuEntry(FMainFrameCommands::Get().ToggleFullscreen);
+		FullscreenEntry.InsertPosition.Position = EToolMenuInsertType::After;
+		FullscreenEntry.InsertPosition.Name = FullscreenSeparatorName;
 #endif
 	}
 }
