[Viewport Toolbar] Level Editor Camera menu- move select piloted actor to piloting section

diff --git a/Engine/Source/Editor/LevelEditor/Private/ViewportToolbar/LevelEditorSubmenus.cpp b/Engine/Source/Editor/LevelEditor/Private/ViewportToolbar/LevelEditorSubmenus.cpp
--- a/Engine/Source/Editor/LevelEditor/Private/ViewportToolbar/LevelEditorSubmenus.cpp
+++ b/Engine/Source/Editor/LevelEditor/Private/ViewportToolbar/LevelEditorSubmenus.cpp
@@ -2419,13 +2419,6 @@ void ExtendCameraSubmenu(FName InCameraOptionsSubmenuName)
 				ULevelViewportContext* LevelViewportContext = InDynamicMenu->FindContext<ULevelViewportContext>();
 				TWeakPtr<::SLevelViewport> LevelViewportWeak = LevelViewportContext->LevelViewport;
 
-				FToolMenuSection& SelectPilotedActorSection = InDynamicMenu->AddSection("SelectPilotedActor");
-
-				SelectPilotedActorSection.InsertPosition =
-					FToolMenuInsert("LevelViewportCameraType_Perspective", EToolMenuInsertType::Before);
-
-				SelectPilotedActorSection.AddMenuEntry(FLevelViewportCommands::Get().SelectPilotedActor);
-
 				// Camera Selection elements
 				{
 					AddCameraActorSelectSection(InDynamicMenu);
@@ -2439,6 +2432,8 @@ void ExtendCameraSubmenu(FName InCameraOptionsSubmenuName)
 						FToolMenuInsert("LevelViewportCameraType_Ortho", EToolMenuInsertType::After);
 					PilotSection.AddSeparator("PilotSectionSeparator");
 
+					PilotSection.AddMenuEntry(FLevelViewportCommands::Get().SelectPilotedActor);
+
 					bool bShowPilotSelectedActorEntry = false;
 
 					AActor* SelectedActor = nullptr;
