[Viewport Toolbar] "Hide Viewport UI" entry should appear before the Fullscreen one in Window Menu

diff --git a/Engine/Source/Editor/LevelEditor/Private/LevelEditorMenu.cpp b/Engine/Source/Editor/LevelEditor/Private/LevelEditorMenu.cpp
--- a/Engine/Source/Editor/LevelEditor/Private/LevelEditorMenu.cpp
+++ b/Engine/Source/Editor/LevelEditor/Private/LevelEditorMenu.cpp
@@ -247,18 +247,17 @@ void FLevelEditorMenu::RegisterLevelEditorMenus()
 		{
 			if (UToolMenu* Menu = UToolMenus::Get()->ExtendMenu("MainFrame.MainMenu.Window"))
 			{
-				// This entry needs to be placed after the fullscreen separator in the Window menu.
-				const FName FullscreenSeparatorName = "FullscreenSeparator";
-
 				FToolMenuSection& LayoutSection = Menu->FindOrAddSection("WindowLayout");
 
+				// This entry needs to be placed after the fullscreen separator in the Window menu.
 				// Making sure the separator exist - the one in Window menu might be missing if not on a Win build.
 				// See FMainMenu::RegisterWindowMenu().
-				LayoutSection.AddSeparator(FullscreenSeparatorName);
+				LayoutSection.AddSeparator("FullscreenSeparator");
 
+				// Hide Viewport UI should appear before the ToggleFullscreen entry
 				FToolMenuEntry& ToggleUIEntry = LayoutSection.AddMenuEntry(FLevelEditorCommands::Get().ToggleHideViewportUI);
-				ToggleUIEntry.InsertPosition.Position = EToolMenuInsertType::After;
-				ToggleUIEntry.InsertPosition.Name = FullscreenSeparatorName;
+				ToggleUIEntry.InsertPosition.Position = EToolMenuInsertType::Before;
+				ToggleUIEntry.InsertPosition.Name = "ToggleFullscreen";
 			}
 		}
 
