[Viewport Toolbar] Copy Cascade entries from Settings menu to Level Viewport Toolbar "Cog" menu

diff --git a/Engine/Source/Editor/LevelEditor/Private/ViewportToolbar/LevelEditorSubmenus.cpp b/Engine/Source/Editor/LevelEditor/Private/ViewportToolbar/LevelEditorSubmenus.cpp
--- a/Engine/Source/Editor/LevelEditor/Private/ViewportToolbar/LevelEditorSubmenus.cpp
+++ b/Engine/Source/Editor/LevelEditor/Private/ViewportToolbar/LevelEditorSubmenus.cpp
@@ -1990,6 +1990,29 @@ FToolMenuEntry CreateSettingsSubmenu()
 						const FLevelViewportCommands& LevelViewportActions = FLevelViewportCommands::Get();
 						SettingsSection.AddMenuEntry(LevelViewportActions.AdvancedSettings);
 					}
+
+					FToolMenuSection& CascadeSection =
+						Submenu->FindOrAddSection("Cascade", LOCTEXT("CascadeSectionLabel", "Cascade"));
+
+					constexpr bool bInOpenSubMenuOnClick = false;
+
+					CascadeSection.AddSubMenu(
+						"CascadeSubmenu",
+						LOCTEXT("CascadeLabel", "Cascade"),
+						LOCTEXT("CascadeTooltip", "Cascade Options"),
+						FNewToolMenuDelegate::CreateLambda(
+							[](UToolMenu* InMenu)
+							{
+								FToolMenuSection& Section =
+									InMenu->AddSection("Cascade", LOCTEXT("CascadeLabel", "Cascade"));
+								Section.AddMenuEntry(FLevelEditorCommands::Get().ToggleParticleSystemLOD);
+								Section.AddMenuEntry(FLevelEditorCommands::Get().ToggleParticleSystemHelpers);
+								Section.AddMenuEntry(FLevelEditorCommands::Get().ToggleFreezeParticleSimulation);
+							}
+						),
+						bInOpenSubMenuOnClick,
+						FSlateIcon(FSlateIcon(FAppStyle::GetAppStyleSetName(), "ClassIcon.ParticleSystemComponent"))
+					);
 				}
 			}
 		)
