Fix "Copy Actor File Path" scene outliner action to work with entity proxy actors.

diff --git a/Engine/Source/Editor/LevelEditor/Private/LevelEditorActions.cpp b/Engine/Source/Editor/LevelEditor/Private/LevelEditorActions.cpp
--- a/Engine/Source/Editor/LevelEditor/Private/LevelEditorActions.cpp
+++ b/Engine/Source/Editor/LevelEditor/Private/LevelEditorActions.cpp
@@ -3104,19 +3104,27 @@ void FLevelEditorActionCallbacks::CopyActorFilePathtoClipboard_Clicked()
 			Result.Append(LINE_TERMINATOR);
 		}
 
-		const UPackage* Package = Actor->GetPackage();
-		const FString LocalFullPath(Package->GetLoadedPath().GetLocalFullPath());
+		const UPackage* SceneOutlinerPackage = Actor->GetSceneOutlinerTopParentPackage();
+		const FString LocalFullPath(SceneOutlinerPackage->GetLoadedPath().GetLocalFullPath());
+
 		if (SelectedActors.Num() > 1)
 		{
-			const FString& ActorLabel = Actor->GetActorLabel(false);
-			if (ActorLabel.Len())
+			const UPackage* ActorPackage = Actor->GetPackage();
+
+			// If the packages differ, it's probably an actor proxy, so we can't rely on the actor's name
+			if (SceneOutlinerPackage == ActorPackage)
 			{
-				Result.Append(ActorLabel);
+				const FString& ActorLabel = Actor->GetActorLabel(false);
+				if (ActorLabel.Len())
+				{
+					Result.Append(ActorLabel);
+				}
+				Result.Append(TEXT("("));
+				Result.Append(Actor->GetName());
+				Result.Append(TEXT("): "));
 			}
-			Result.Append(TEXT("("));
-			Result.Append(Actor->GetName());
-			Result.Append(TEXT("): "));
 		}
+
 		Result.Append(FPaths::ConvertRelativePathToFull(LocalFullPath));
 	}
 
diff --git a/Engine/Source/Editor/SceneOutliner/Private/ActorTreeItem.cpp b/Engine/Source/Editor/SceneOutliner/Private/ActorTreeItem.cpp
--- a/Engine/Source/Editor/SceneOutliner/Private/ActorTreeItem.cpp
+++ b/Engine/Source/Editor/SceneOutliner/Private/ActorTreeItem.cpp
@@ -410,7 +410,7 @@ FString FActorTreeItem::GetPackageName() const
 {
 	if (const AActor* ActorPtr = Actor.Get())
 	{
-		if (UPackage* ActorPackage = ActorPtr->GetSceneOutlinerPackage())
+		if (UPackage* ActorPackage = ActorPtr->GetSceneOutlinerItemPackage())
 		{
 			return ActorPackage->GetName();
 		}
diff --git a/Engine/Source/Runtime/Engine/Classes/GameFramework/Actor.h b/Engine/Source/Runtime/Engine/Classes/GameFramework/Actor.h
--- a/Engine/Source/Runtime/Engine/Classes/GameFramework/Actor.h
+++ b/Engine/Source/Runtime/Engine/Classes/GameFramework/Actor.h
@@ -2657,8 +2657,15 @@ public:
 	ENGINE_API virtual AActor* GetSceneOutlinerParent() const;
 
 	/** Returns the actor package that should be used in the scene outliner */
+	UE_DEPRECATED(5.6, "Please use GetSceneOutlinerItemPackage or GetSceneOutlinerTopParentPackage instead.")
 	ENGINE_API virtual UPackage* GetSceneOutlinerPackage() const;
 
+	/** Returns the actor item package that should be used in the scene outliner */
+	ENGINE_API virtual UPackage* GetSceneOutlinerItemPackage() const;
+
+	/** Returns the actor top parent package that should be used in the scene outliner */
+	ENGINE_API virtual UPackage* GetSceneOutlinerTopParentPackage() const;
+
 	/** Called before editor copy, true allow export */
 	virtual bool ShouldExport() { return true; }
 
diff --git a/Engine/Source/Runtime/Engine/Private/ActorEditor.cpp b/Engine/Source/Runtime/Engine/Private/ActorEditor.cpp
--- a/Engine/Source/Runtime/Engine/Private/ActorEditor.cpp
+++ b/Engine/Source/Runtime/Engine/Private/ActorEditor.cpp
@@ -1078,7 +1078,19 @@ AActor* AActor::GetSceneOutlinerParent() const
 	return GetAttachParentActor();
 }
 
+PRAGMA_DISABLE_DEPRECATION_WARNINGS
 UPackage* AActor::GetSceneOutlinerPackage() const
+{
+	return GetSceneOutlinerItemPackage();
+}
+PRAGMA_ENABLE_DEPRECATION_WARNINGS
+
+UPackage* AActor::GetSceneOutlinerItemPackage() const
+{
+	return GetExternalPackage();
+}
+
+UPackage* AActor::GetSceneOutlinerTopParentPackage() const
 {
 	return GetExternalPackage();
 }
