Remove keyboard shortcut for DumpGPU command.

- Users were accidently hitting this keyboard shorcut.
- Leaving the UI command in case we want to hook this up to a UI menu, etc.

diff --git a/Engine/Source/Editor/LevelEditor/Private/LevelEditorActions.cpp b/Engine/Source/Editor/LevelEditor/Private/LevelEditorActions.cpp
--- a/Engine/Source/Editor/LevelEditor/Private/LevelEditorActions.cpp
+++ b/Engine/Source/Editor/LevelEditor/Private/LevelEditorActions.cpp
@@ -4123,7 +4123,7 @@ void FLevelEditorCommands::RegisterCommands()
 
 	UI_COMMAND( RecompileShaders, "Recompile Changed Shaders", "Recompiles shaders which are out of date", EUserInterfaceActionType::Button, FInputChord( EModifierKey::Shift|EModifierKey::Control, EKeys::Period ) );
 	UI_COMMAND( ProfileGPU, "Profile GPU", "Profiles the GPU for the next frame and opens a window with profiled data", EUserInterfaceActionType::Button, FInputChord( EModifierKey::Shift|EModifierKey::Control, EKeys::Comma ) );
-	UI_COMMAND( DumpGPU, "Dump GPU", "Dump the GPU intermediary resources for the next frame and opens explorer", EUserInterfaceActionType::Button, FInputChord( EModifierKey::Shift|EModifierKey::Control, EKeys::Slash ) );
+	UI_COMMAND( DumpGPU, "Dump GPU", "Dump the GPU intermediary resources for the next frame and opens explorer", EUserInterfaceActionType::Button, FInputChord());
 
 	UI_COMMAND( ResetAllParticleSystems, "Reset All Particle Systems", "Resets all particle system emitters (removes all active particles and restarts them)", EUserInterfaceActionType::Button, FInputChord( EModifierKey::Shift, EKeys::Slash ) );
 	UI_COMMAND( ResetSelectedParticleSystem, "Resets Selected Particle Systems" , "Resets selected particle system emitters (removes all active particles and restarts them)", EUserInterfaceActionType::Button, FInputChord( EKeys::Slash ) );
