Fix bug where the Camera menu of the old level editor viewport toolbar was empty

diff --git a/Engine/Source/Editor/LevelEditor/Private/SLevelViewportToolBar.cpp b/Engine/Source/Editor/LevelEditor/Private/SLevelViewportToolBar.cpp
--- a/Engine/Source/Editor/LevelEditor/Private/SLevelViewportToolBar.cpp
+++ b/Engine/Source/Editor/LevelEditor/Private/SLevelViewportToolBar.cpp
@@ -747,14 +747,33 @@ void SLevelViewportToolBar::GenerateCameraSpawnMenu(UToolMenu* Menu) const
 TSharedRef<SWidget> SLevelViewportToolBar::GenerateCameraMenu() const
 {
 	static const FName MenuName("LevelEditor.LevelViewportToolBar.Camera");
-	if (!UToolMenus::Get()->IsMenuRegistered(MenuName))
+
+	// Use a static bool to track whether or not this menu is registered. Bool instead of checking the registered
+	// state with ToolMenus because we want the new viewport toolbar to be able to create this menu without breaking
+	// this code. Static because this code can be called multiple times using different instances of this class.
+	static bool bDidRegisterMenu = false;
+	if (!bDidRegisterMenu)
 	{
-		UToolMenu* Menu = UToolMenus::Get()->RegisterMenu(MenuName);
-		Menu->AddDynamicSection("DynamicSection", FNewToolMenuDelegate::CreateLambda([](UToolMenu* InMenu)
-		{
-			ULevelViewportToolBarContext* Context = InMenu->FindContext<ULevelViewportToolBarContext>();
-			Context->LevelViewportToolBarWidget.Pin()->FillCameraMenu(InMenu);
-		}));
+		bDidRegisterMenu = true;
+
+		constexpr bool bWarnIfAlreadyRegistered = false;
+		UToolMenu* Menu =
+			UToolMenus::Get()->RegisterMenu(MenuName, NAME_None, EMultiBoxType::Menu, bWarnIfAlreadyRegistered);
+		Menu->AddDynamicSection(
+			"DynamicSection",
+			FNewToolMenuDelegate::CreateLambda(
+				[](UToolMenu* InMenu)
+				{
+					if (ULevelViewportToolBarContext* Context = InMenu->FindContext<ULevelViewportToolBarContext>())
+					{
+						if (TSharedPtr<SLevelViewportToolBar> ViewportToolbar = Context->LevelViewportToolBarWidget.Pin())
+						{
+							ViewportToolbar->FillCameraMenu(InMenu);
+						}
+					}
+				}
+			)
+		);
 	}
 
 	Viewport.Pin()->OnFloatingButtonClicked();
