OFPE: Route "Save Selected Actors" and "Show File History" for entities.
Can also be executed on child entities, which will be run on the top parent entity.

diff --git a/Engine/Source/Editor/LevelEditor/Private/LevelEditorActions.cpp b/Engine/Source/Editor/LevelEditor/Private/LevelEditorActions.cpp
--- a/Engine/Source/Editor/LevelEditor/Private/LevelEditorActions.cpp
+++ b/Engine/Source/Editor/LevelEditor/Private/LevelEditorActions.cpp
@@ -3140,9 +3140,9 @@ void FLevelEditorActionCallbacks::SaveActor_Clicked()
 	for (FSelectionIterator It(GEditor->GetSelectedActorIterator()); It; ++It)
 	{
 		AActor* Actor = Cast<AActor>(*It);
-		if (Actor->IsPackageExternal())
+		if (UPackage* ActorPackage = Actor->GetSceneOutlinerTopParentPackage())
 		{
-			PackagesToSave.Add(Actor->GetPackage());
+			PackagesToSave.Add(ActorPackage);
 		}
 	}
 
@@ -3163,7 +3163,7 @@ bool FLevelEditorActionCallbacks::SaveActor_CanExecute()
 	for (FSelectionIterator It(GEditor->GetSelectedActorIterator()); It; ++It)
 	{
 		AActor* Actor = Cast<AActor>(*It);
-		if (Actor->IsPackageExternal())
+		if (UPackage* ActorPackage = Actor->GetSceneOutlinerTopParentPackage())
 		{
 			return true;
 		}
@@ -3177,28 +3177,18 @@ static TArray<FString> GetSelectedActorsPackageFullpath()
 
 	for (FSelectionIterator It(GEditor->GetSelectedActorIterator()); It; ++It)
 	{
-		const AActor* Actor = Cast<AActor>(*It);
-
-		if (Actor == nullptr)
+		if (const AActor* Actor = Cast<AActor>(*It))
 		{
-			continue;
-		}
-
-		const UPackage* Package = Actor->GetPackage();
-
-		if (Package == nullptr)
-		{
-			continue;
-		}
-
-		const FString LocalFullPath(Package->GetLoadedPath().GetLocalFullPath());
+			if (const UPackage* Package = Actor->GetSceneOutlinerTopParentPackage())
+			{
+				const FString LocalFullPath(Package->GetLoadedPath().GetLocalFullPath());
 
-		if (LocalFullPath.IsEmpty())
-		{
-			continue;
+				if (!LocalFullPath.IsEmpty())
+				{
+					PackageFullpaths.Add(FPaths::ConvertRelativePathToFull(LocalFullPath));
+				}
+			}
 		}
-
-		PackageFullpaths.Add(FPaths::ConvertRelativePathToFull(LocalFullPath));
 	}
 
 	return PackageFullpaths;
