Only show actionable messages in the active level viewport.
- Refactor a handful of duplicate "is active" module checks into using a single `IsActiveLevelViewport()` function.

diff --git a/Engine/Source/Editor/LevelEditor/Private/SLevelViewport.cpp b/Engine/Source/Editor/LevelEditor/Private/SLevelViewport.cpp
--- a/Engine/Source/Editor/LevelEditor/Private/SLevelViewport.cpp
+++ b/Engine/Source/Editor/LevelEditor/Private/SLevelViewport.cpp
@@ -410,6 +410,11 @@ void SLevelViewport::ConstructViewportOverlayContent()
 							return EVisibility::Collapsed;
 						}
 
+						if (!IsActiveLevelViewport())
+						{
+							return EVisibility::Collapsed;
+						}
+
 						return ActionableMessageViewportWidget->GetVisibility();
 					}
 				)
@@ -1904,12 +1909,8 @@ EVisibility SLevelViewport::GetCloseImmersiveButtonVisibility() const
 }
 
 EVisibility SLevelViewport::GetTransformToolbarVisibility() const
-{
-	FLevelEditorModule& LevelEditorModule = FModuleManager::GetModuleChecked<FLevelEditorModule>( LevelEditorName );
-	TSharedPtr<IAssetViewport> ActiveLevelViewport = LevelEditorModule.GetFirstActiveViewport();
-
-	// Am I the ActiveLevelViewport? 
-	if( ActiveLevelViewport.Get() == this )
+{ 
+	if (IsActiveLevelViewport())
 	{
 		// Only return visible if we are/were the active viewport. 
 		return EVisibility::Visible;
@@ -2271,6 +2272,12 @@ void SLevelViewport::OnCreateCameraActor(UClass* InClass)
 	GEditor->RedrawAllViewports();
 }
 
+bool SLevelViewport::IsActiveLevelViewport() const
+{
+	FLevelEditorModule& LevelEditorModule = FModuleManager::GetModuleChecked<FLevelEditorModule>(LevelEditorName);
+	return LevelEditorModule.GetFirstActiveViewport().Get() == this;
+}
+
 bool SLevelViewport::IsPerspectiveViewport() const
 {
 	bool bIsPerspective = false;
@@ -4943,15 +4950,7 @@ void SLevelViewport::RemoveOverlayWidget(TSharedRef<SWidget> OverlaidWidget)
 
 bool SLevelViewport::CanProduceActionForCommand(const TSharedRef<const FUICommandInfo>& Command) const
 {
-
-	FLevelEditorModule& LevelEditorModule = FModuleManager::GetModuleChecked<FLevelEditorModule>( LevelEditorName );
-	TSharedPtr<IAssetViewport> ActiveLevelViewport = LevelEditorModule.GetFirstActiveViewport();
-	if ( ActiveLevelViewport.IsValid() )
-	{
-		return ActiveLevelViewport == SharedThis(this);
-	}
-
-	return false;
+	return IsActiveLevelViewport();
 }
 
 void SLevelViewport::LockActorInternal(AActor* NewActorToLock)
diff --git a/Engine/Source/Editor/LevelEditor/Public/SLevelViewport.h b/Engine/Source/Editor/LevelEditor/Public/SLevelViewport.h
--- a/Engine/Source/Editor/LevelEditor/Public/SLevelViewport.h
+++ b/Engine/Source/Editor/LevelEditor/Public/SLevelViewport.h
@@ -520,6 +520,11 @@ private:
 	 */
 	void OnTakeHighResScreenshot();
 
+	/**
+	 * Returns whether this level viewport is the active level viewport 
+	 */
+	bool IsActiveLevelViewport() const;
+
 	/**
 	 * Called to check currently selected editor viewport is a perspective one
 	 */
