[Viewport Toolbar] Render Doc capture button is missing in the new level editor viewport toolbar

diff --git a/Engine/Source/Editor/LevelEditor/Private/SLevelViewport.cpp b/Engine/Source/Editor/LevelEditor/Private/SLevelViewport.cpp
--- a/Engine/Source/Editor/LevelEditor/Private/SLevelViewport.cpp
+++ b/Engine/Source/Editor/LevelEditor/Private/SLevelViewport.cpp
@@ -2087,6 +2087,22 @@ TSharedPtr<SWidget> SLevelViewport::BuildViewportToolbar()
 				RightSection.AddEntry(UE::LevelEditor::CreateSettingsSubmenu());
 				RightSection.AddEntry(UE::LevelEditor::CreateViewportSizingSubmenu());
 			}
+
+			// Sections added for compatibility with old toolbar transform menu extensions
+			{
+				// Existing code might add entries via FToolBarBuilder Extenders,
+				// which seem to only work with Toolbar main sections.
+				// The following sections match the ones added in STransformViewportToolBar::MakeTransformToolBar(...),
+				// which makes use of a FToolBarBuilder. Like so, code extending the Old Transform Toolbar using these
+				// sections names as hooks should also work in the new Level Viewport Toolbar.
+
+				ViewportToolbarMenu->FindOrAddSection("Transform").Alignment = EToolMenuSectionAlign::Last;
+				ViewportToolbarMenu->FindOrAddSection("LocationGridSnap").Alignment = EToolMenuSectionAlign::Last;
+				ViewportToolbarMenu->FindOrAddSection("RotationGridSnap").Alignment = EToolMenuSectionAlign::Last;
+				ViewportToolbarMenu->FindOrAddSection("Layer2DSnap").Alignment = EToolMenuSectionAlign::Last;
+				ViewportToolbarMenu->FindOrAddSection("ScaleGridSnap").Alignment = EToolMenuSectionAlign::Last;
+				ViewportToolbarMenu->FindOrAddSection("CameraSpeed").Alignment = EToolMenuSectionAlign::Last;
+			}
 		}
 	}
 
@@ -2130,6 +2146,17 @@ TSharedPtr<SWidget> SLevelViewport::BuildViewportToolbar()
 					ViewportToolbarContext.AddExtender(ViewMenuExtenders);
 				}
 			}
+
+			// Stay backward-compatible with legacy toolbar menu extenders
+			{
+				TSharedPtr<FExtender> OldToolbarExtenders =
+					LevelEditorModule.GetToolBarExtensibilityManager()->GetAllExtenders();
+
+				if (OldToolbarExtenders.IsValid())
+				{
+					ViewportToolbarContext.AddExtender(OldToolbarExtenders);
+				}
+			}
 		}
 
 		// Add the UnrealEd viewport toolbar context.
