[Truncation Warnings] Fix warnings in LevelEditor module

#rnx

diff --git a/Engine/Source/Editor/LevelEditor/Private/DlgDeltaTransform.cpp b/Engine/Source/Editor/LevelEditor/Private/DlgDeltaTransform.cpp
--- a/Engine/Source/Editor/LevelEditor/Private/DlgDeltaTransform.cpp
+++ b/Engine/Source/Editor/LevelEditor/Private/DlgDeltaTransform.cpp
@@ -3,6 +3,7 @@
 #include "DlgDeltaTransform.h"
 #include "InputCoreTypes.h"
 #include "Input/Reply.h"
+#include "Math/MathFwd.h"
 #include "Widgets/DeclarativeSyntaxSupport.h"
 #include "Widgets/SCompoundWidget.h"
 #include "Widgets/SBoxPanel.h"
@@ -44,6 +45,8 @@ public:
 
 		DeltaTransform.Set(FVector::ZeroVector);
 
+		using SVectorEntryBox = SNumericVectorInputBox<FVector::FReal, FVector, 3>;
+
 		ChildSlot
 		[
 			SNew(SBorder)
@@ -66,7 +69,7 @@ public:
 						.FillWidth(1)
 						.VAlign( VAlign_Center )
 						[
-							SNew( SVectorInputBox )
+							SNew( SVectorEntryBox )
 							.X( this, &SDlgDeltaTransform::GetDeltaX )
 							.Y( this, &SDlgDeltaTransform::GetDeltaY )
 							.Z( this, &SDlgDeltaTransform::GetDeltaZ )
@@ -173,9 +176,9 @@ private:
 			return X.IsSet() && Y.IsSet() && Z.IsSet();
 		}
 
-		TOptional<float> X;
-		TOptional<float> Y;
-		TOptional<float> Z;
+		TOptional<FVector::FReal> X;
+		TOptional<FVector::FReal> Y;
+		TOptional<FVector::FReal> Z;
 	};
 
 	/** 
@@ -198,22 +201,22 @@ private:
 		return FReply::Handled();
 	}
 
-	TOptional<float> GetDeltaX() const
+	TOptional<FVector::FReal> GetDeltaX() const
 	{
 		return DeltaTransform.X;
 	}
 
-	TOptional<float> GetDeltaY() const
+	TOptional<FVector::FReal> GetDeltaY() const
 	{
 		return DeltaTransform.Y;
 	}
 
-	TOptional<float> GetDeltaZ() const
+	TOptional<FVector::FReal> GetDeltaZ() const
 	{
 		return DeltaTransform.Z;
 	}
 
-	void OnSetDelta( float NewValue, ETextCommit::Type CommitInfo, int32 Axis )
+	void OnSetDelta(FVector::FReal NewValue, ETextCommit::Type CommitInfo, int32 Axis )
 	{
 		switch (Axis)
 		{
diff --git a/Engine/Source/Editor/LevelEditor/Private/Elements/Component/ComponentElementLevelEditorViewportInteractionCustomization.cpp b/Engine/Source/Editor/LevelEditor/Private/Elements/Component/ComponentElementLevelEditorViewportInteractionCustomization.cpp
--- a/Engine/Source/Editor/LevelEditor/Private/Elements/Component/ComponentElementLevelEditorViewportInteractionCustomization.cpp
+++ b/Engine/Source/Editor/LevelEditor/Private/Elements/Component/ComponentElementLevelEditorViewportInteractionCustomization.cpp
@@ -146,9 +146,9 @@ void FComponentElementLevelEditorViewportInteractionCustomization::ValidateScale
 	};
 
 	/** Check scale criteria to see if this is allowed, returns modified absolute scale*/
-	auto CheckScaleValue = [](float ScaleDeltaToCheck, float CurrentScaleFactor, float CurrentExtent, bool bCheckSmallExtent, bool bSnap) -> float
+	auto CheckScaleValue = [](FVector::FReal ScaleDeltaToCheck, FVector::FReal CurrentScaleFactor, FVector::FReal CurrentExtent, bool bCheckSmallExtent, bool bSnap) -> FVector::FReal
 	{
-		float AbsoluteScaleValue = ScaleDeltaToCheck + CurrentScaleFactor;
+		FVector::FReal AbsoluteScaleValue = ScaleDeltaToCheck + CurrentScaleFactor;
 		if (bSnap)
 		{
 			AbsoluteScaleValue = FMath::GridSnap(AbsoluteScaleValue, GEditor->GetScaleGridSize());
@@ -157,8 +157,8 @@ void FComponentElementLevelEditorViewportInteractionCustomization::ValidateScale
 		if (FMath::Abs(CurrentExtent) < KINDA_SMALL_NUMBER) {
 			return AbsoluteScaleValue;
 		}
-		float UnscaledExtent = CurrentExtent / CurrentScaleFactor;
-		float ScaledExtent = UnscaledExtent * AbsoluteScaleValue;
+		FVector::FReal UnscaledExtent = CurrentExtent / CurrentScaleFactor;
+		FVector::FReal ScaledExtent = UnscaledExtent * AbsoluteScaleValue;
 
 		if ((FMath::Square(ScaledExtent)) > BIG_NUMBER)	// cant get too big...
 		{
@@ -168,7 +168,7 @@ void FComponentElementLevelEditorViewportInteractionCustomization::ValidateScale
 			(FMath::Abs(ScaledExtent) < MIN_ACTOR_BOUNDS_EXTENT * 0.5f ||		// ...or too small (apply sign in this case)...
 				(CurrentScaleFactor < 0.0f) != (AbsoluteScaleValue < 0.0f)))	// ...also cant cross the zero boundary
 		{
-			return ((MIN_ACTOR_BOUNDS_EXTENT * 0.5) / UnscaledExtent) * (CurrentScaleFactor < 0.0f ? -1.0f : 1.0f);
+			return ((MIN_ACTOR_BOUNDS_EXTENT * 0.5f) / UnscaledExtent) * (CurrentScaleFactor < 0.0f ? -1.0f : 1.0f);
 		}
 
 		return AbsoluteScaleValue;
@@ -189,7 +189,7 @@ void FComponentElementLevelEditorViewportInteractionCustomization::ValidateScale
 		int ActiveAxisCount = 0;
 		bool CurrentValueSameSign = true;
 		bool FirstSignPositive = true;
-		float MaxComponentSum = -1.0f;
+		FVector::FReal MaxComponentSum = -1.0f;
 		int32 MaxAxisIndex = -1;
 		const ULevelEditorViewportSettings* ViewportSettings = GetDefault<ULevelEditorViewportSettings>();
 		bool SnapScaleAfter = ViewportSettings->SnapScaleEnabled;
@@ -288,7 +288,7 @@ void FComponentElementLevelEditorViewportInteractionCustomization::ValidateScale
 	{
 		if (bActiveAxes[Axis])
 		{
-			float ModifiedScaleAbsolute = CheckScaleValue(InOutScaleDelta[Axis], InCurrentScale[Axis], InBoxExtent[Axis], bInCheckSmallExtent, bSnapAxes);
+			FVector::FReal ModifiedScaleAbsolute = CheckScaleValue(InOutScaleDelta[Axis], InCurrentScale[Axis], InBoxExtent[Axis], bInCheckSmallExtent, bSnapAxes);
 			InOutScaleDelta[Axis] = ModifiedScaleAbsolute - InCurrentScale[Axis];
 		}
 		else
diff --git a/Engine/Source/Editor/LevelEditor/Private/HighResScreenshotUI.cpp b/Engine/Source/Editor/LevelEditor/Private/HighResScreenshotUI.cpp
--- a/Engine/Source/Editor/LevelEditor/Private/HighResScreenshotUI.cpp
+++ b/Engine/Source/Editor/LevelEditor/Private/HighResScreenshotUI.cpp
@@ -399,14 +399,14 @@ FReply SHighResScreenshotDialog::OnCaptureClicked()
 		auto ConfigViewport = Config.TargetViewport.Pin();
 		if (ConfigViewport.IsValid())
 		{
-			GScreenshotResolutionX = ConfigViewport->GetRenderTargetTextureSizeXY().X * Config.ResolutionMultiplier;
-			GScreenshotResolutionY = ConfigViewport->GetRenderTargetTextureSizeXY().Y * Config.ResolutionMultiplier;
+			GScreenshotResolutionX = FMath::AsUInt(ConfigViewport->GetRenderTargetTextureSizeXY().X * Config.ResolutionMultiplier);
+			GScreenshotResolutionY = FMath::AsUInt(ConfigViewport->GetRenderTargetTextureSizeXY().Y * Config.ResolutionMultiplier);
 			FIntRect ScaledCaptureRegion = Config.UnscaledCaptureRegion;
 
 			if (ScaledCaptureRegion.Area() > 0)
 			{
 				ScaledCaptureRegion.Clip(FIntRect(FIntPoint::ZeroValue, ConfigViewport->GetRenderTargetTextureSizeXY()));
-				ScaledCaptureRegion *= Config.ResolutionMultiplier;
+				ScaledCaptureRegion *= FMath::TruncToInt32(Config.ResolutionMultiplier);
 			}
 
 			Config.CaptureRegion = ScaledCaptureRegion;
diff --git a/Engine/Source/Editor/LevelEditor/Private/LevelEditor.cpp b/Engine/Source/Editor/LevelEditor/Private/LevelEditor.cpp
--- a/Engine/Source/Editor/LevelEditor/Private/LevelEditor.cpp
+++ b/Engine/Source/Editor/LevelEditor/Private/LevelEditor.cpp
@@ -592,9 +592,9 @@ void FLevelEditorModule::GoImmersiveWithActiveLevelViewport( const bool bForceGa
 				const bool bWantImmersive = true;
 				const bool bAllowAnimation = false;
 				ActiveLevelViewport->MakeImmersive( bWantImmersive, bAllowAnimation );
-				FVector2D WindowSize = Window->GetSizeInScreen();
+				FVector2f WindowSize = Window->GetSizeInScreen();
 				// Set the initial size of the viewport to be the size of the window. This must be done because Slate has not ticked yet so the viewport will have no initial size
-				ActiveLevelViewport->GetActiveViewport()->SetInitialSize( FIntPoint( FMath::TruncToInt( WindowSize.X ), FMath::TruncToInt( WindowSize.Y ) ) );
+				ActiveLevelViewport->GetActiveViewport()->SetInitialSize( FIntPoint( FMath::TruncToInt32( WindowSize.X ), FMath::TruncToInt32( WindowSize.Y ) ) );
 			}
 		}
 	}
diff --git a/Engine/Source/Editor/LevelEditor/Private/LevelEditorContextMenu.cpp b/Engine/Source/Editor/LevelEditor/Private/LevelEditorContextMenu.cpp
--- a/Engine/Source/Editor/LevelEditor/Private/LevelEditorContextMenu.cpp
+++ b/Engine/Source/Editor/LevelEditor/Private/LevelEditorContextMenu.cpp
@@ -1163,7 +1163,7 @@ class SOnlyExpandsBox : public SBox
 protected:
 	virtual FVector2D ComputeDesiredSize(float LayoutScaleMultiplier) const override
 	{
-		FVector2D RequestedDesiredSize = SBox::ComputeDesiredSize(LayoutScaleMultiplier);
+		UE::Slate::FDeprecateVector2DResult RequestedDesiredSize = UE::Slate::CastToVector2f(SBox::ComputeDesiredSize(LayoutScaleMultiplier));
 		if (RequestedDesiredSize.X > MaxPreviousWidth)
 		{
 			MaxPreviousWidth = RequestedDesiredSize.X;
diff --git a/Engine/Source/Editor/LevelEditor/Private/LightmapResRatioAdjust.cpp b/Engine/Source/Editor/LevelEditor/Private/LightmapResRatioAdjust.cpp
--- a/Engine/Source/Editor/LevelEditor/Private/LightmapResRatioAdjust.cpp
+++ b/Engine/Source/Editor/LevelEditor/Private/LightmapResRatioAdjust.cpp
@@ -216,7 +216,7 @@ bool FLightmapResRatioAdjustSettings::ApplyRatioAdjustment()
 				if (bConvertIt == true)
 				{
 					// Process it!
-					int32 CurrentResolution = Surf.LightMapScale;
+					int32 CurrentResolution = FMath::TruncToInt32(Surf.LightMapScale);
 					float Scalar = 1.0f / Settings.Ratio;
 					if ((Scalar < 1.0f) && (CurrentResolution <= Settings.Min_BSPSurfaces))
 					{
diff --git a/Engine/Source/Editor/LevelEditor/Private/SLevelViewport.cpp b/Engine/Source/Editor/LevelEditor/Private/SLevelViewport.cpp
--- a/Engine/Source/Editor/LevelEditor/Private/SLevelViewport.cpp
+++ b/Engine/Source/Editor/LevelEditor/Private/SLevelViewport.cpp
@@ -545,7 +545,7 @@ void SLevelViewport::ConstructViewportOverlayContent()
 	ViewportOverlay->AddSlot(SlotIndex)
 	.VAlign( VAlign_Fill )
 	.HAlign( HAlign_Fill )
-	.Padding( 0 )
+	.Padding( 0.f )
 	[
 		SAssignNew(CaptureRegionWidget, SCaptureRegionWidget)
 	];
@@ -907,9 +907,8 @@ bool SLevelViewport::HandleDragObjects(const FGeometry& MyGeometry, const FDragD
 		LevelViewportClient->GetViewportDimensions(ViewportOrigin, ViewportSize);
 
 		// Save off the local mouse position from the drop point for potential use later (with Drag Drop context menu)
-		CachedOnDropLocalMousePos = MyGeometry.AbsoluteToLocal( DragDropEvent.GetScreenSpacePosition() ) * MyGeometry.Scale;
-		CachedOnDropLocalMousePos.X -= ViewportOrigin.X;
-		CachedOnDropLocalMousePos.Y -= ViewportOrigin.Y;
+		CachedOnDropLocalMousePos = FVector2f(MyGeometry.AbsoluteToLocal( DragDropEvent.GetScreenSpacePosition() ) * MyGeometry.Scale).IntPoint();
+		CachedOnDropLocalMousePos -= ViewportOrigin;
 	}
 
 	// Update the currently dragged actor if it exists
@@ -3337,7 +3336,7 @@ void SActorPreview::Construct( const FArguments& InArgs )
 	this->ChildSlot
 	[
 		SNew(SBorder)
-		.Padding(0)
+		.Padding(0.f)
 
 		.Visibility(EVisibility::SelfHitTestInvisible)
 
@@ -3398,7 +3397,7 @@ void SActorPreview::Construct( const FArguments& InArgs )
 			[
 				// Create a button to pin/unpin this viewport
 				SNew(SButton)
-					.ContentPadding(0)
+					.ContentPadding(0.f)
 					.ForegroundColor(FSlateColor::UseForeground())
 					.ButtonStyle(FAppStyle::Get(), "ToggleButton")
 
@@ -3419,15 +3418,15 @@ void SActorPreview::Construct( const FArguments& InArgs )
 			+SOverlay::Slot()
 			.HAlign(HAlign_Left)
 			.VAlign(VAlign_Bottom)
-			.Padding( 0 )
+			.Padding( 0.f )
 			[
 				SNew(SBox)
-				.WidthOverride(45)
-				.HeightOverride(45)
+				.WidthOverride(45.f)
+				.HeightOverride(45.f)
 				[
 					// Create a button to pin/unpin this viewport
 					SNew( SButton )
-						.ContentPadding(0)
+						.ContentPadding(0.f)
 						.ForegroundColor( FSlateColor::UseForeground() )
 						.ButtonStyle( FAppStyle::Get(), "ToggleButton" )
 
@@ -3449,15 +3448,15 @@ void SActorPreview::Construct( const FArguments& InArgs )
 			+ SOverlay::Slot()
 			.HAlign(HAlign_Right)
 			.VAlign(VAlign_Bottom)
-			.Padding(0)
+			.Padding(0.f)
 			[
 				SNew(SBox)
-				.WidthOverride(45)
-				.HeightOverride(45)
+				.WidthOverride(45.f)
+				.HeightOverride(45.f)
 				[
 					// Create a button to attach/detach this viewport
 					SNew(SButton)
-						.ContentPadding(0)
+						.ContentPadding(0.f)
 						.ForegroundColor(FSlateColor::UseForeground())
 						.ButtonStyle(FAppStyle::Get(), "ToggleButton")
 						.IsFocusable(false)
@@ -3785,11 +3784,11 @@ FOptionalSize SActorPreview::OnReadHeight() const
 {
 	const float MinimumHeight = 32;
 	// Also used as parent height in case valid parent viewport is not set
-	const float MaximumHeight = 428;
+	const FIntPoint::IntType MaximumHeight = 428;
 	// Used to make sure default viewport scale * parent viewport height = roughly same size as original windows
 	const float PreviewScalingFactor = 0.06308f;
 
-	float ParentHeight = MaximumHeight;
+	FIntPoint::IntType ParentHeight = MaximumHeight;
 	if ( ParentViewport.IsValid() )
 	{
 		ParentHeight = ParentViewport.Pin()->GetActiveViewport()->GetSizeXY().Y;
diff --git a/Engine/Source/Editor/LevelEditor/Private/SSurfaceProperties.cpp b/Engine/Source/Editor/LevelEditor/Private/SSurfaceProperties.cpp
--- a/Engine/Source/Editor/LevelEditor/Private/SSurfaceProperties.cpp
+++ b/Engine/Source/Editor/LevelEditor/Private/SSurfaceProperties.cpp
@@ -40,8 +40,8 @@ void SSurfaceProperties::Construct( const FArguments& InArgs )
 		FBspSurf* Surf = *It;
 		UModel* Model = It.GetModel();
 
-		const FVector TextureU(Model->Vectors[Surf->vTextureU]);
-		const FVector TextureV(Model->Vectors[Surf->vTextureV]);
+		const FVector3f TextureU(Model->Vectors[Surf->vTextureU]);
+		const FVector3f TextureV(Model->Vectors[Surf->vTextureV]);
 
 		const float TextureUSize = TextureU.Size();
 		const float TextureVSize = TextureV.Size();
@@ -80,12 +80,12 @@ void SSurfaceProperties::Construct( const FArguments& InArgs )
 
 		+SVerticalBox::Slot()
 		.AutoHeight()
-		.Padding(0,0,0,5)
+		.Padding(0.f,0.f,0.f,5.f)
 		[
 			SNew(SBorder)
 			.BorderBackgroundColor( BorderColor )
 			.BorderImage( BorderStyle )
-			.Padding(10)
+			.Padding(10.f)
 			.AddMetaData<FTagMetaData>(TEXT("DetailsView.TexturePan"))
 			[
 				ConstructTexturePan()
@@ -94,29 +94,29 @@ void SSurfaceProperties::Construct( const FArguments& InArgs )
 
 		+SVerticalBox::Slot()
 		.AutoHeight()
-		.Padding(0,0,0,5)
+		.Padding(0.f,0.f,0.f,5.f)
 		[
 			SNew(SHorizontalBox)
 			+ SHorizontalBox::Slot()
 			.FillWidth(3)
-			.Padding(0,0,5,0)
+			.Padding(0.f,0.f,5.f,0.f)
 			[
 				SNew(SBorder)
 				.BorderBackgroundColor( BorderColor )
 				.BorderImage( BorderStyle )
-				.Padding(10)
+				.Padding(10.f)
 				.AddMetaData<FTagMetaData>(TEXT("DetailsView.TextureRotate"))
 				[
 					ConstructTextureRotate()
 				]
 			]
 			+ SHorizontalBox::Slot()
-			.FillWidth(2)
+			.FillWidth(2.f)
 			[
 				SNew(SBorder)
 				.BorderBackgroundColor( BorderColor )
 				.BorderImage( BorderStyle )
-				.Padding(10)
+				.Padding(10.f)
 				.AddMetaData<FTagMetaData>(TEXT("DetailsView.TextureFlip"))
 				[
 					ConstructTextureFlip()
@@ -126,12 +126,12 @@ void SSurfaceProperties::Construct( const FArguments& InArgs )
 
 		+SVerticalBox::Slot()
 		.AutoHeight()
-		.Padding(0,0,0,5)
+		.Padding(0.f,0.f,0.f,5.f)
 		[
 			SNew(SBorder)
 			.BorderBackgroundColor( BorderColor )
 			.BorderImage( BorderStyle )
-			.Padding(10)
+			.Padding(10.f)
 			.AddMetaData<FTagMetaData>(TEXT("DetailsView.TextureScale"))
 			[
 				ConstructTextureScale()
@@ -140,12 +140,12 @@ void SSurfaceProperties::Construct( const FArguments& InArgs )
 			
 		+SVerticalBox::Slot()
 		.AutoHeight()
-		.Padding(0,0,0,5)
+		.Padding(0.f,0.f,0.f,5.f)
 		[
 			SNew(SBorder)
 			.BorderBackgroundColor( BorderColor )
 			.BorderImage( BorderStyle )
-			.Padding(10)
+			.Padding(10.f)
 			.AddMetaData<FTagMetaData>(TEXT("DetailsView.ConstructLighting"))
 			[
 				ConstructLighting()
@@ -161,7 +161,7 @@ TSharedRef<SWidget> SSurfaceProperties::ConstructTexturePan()
 	
 	ParentBox->AddSlot()
 	.AutoHeight()
-	.Padding(0,0,0,5)
+	.Padding(0.f,0.f,0.f,5.f)
 	[
 		SNew(STextBlock)
 		.Text( LOCTEXT("Pan", "Pan:") )
@@ -189,7 +189,7 @@ TSharedRef<SWidget> SSurfaceProperties::ConstructTexturePan()
 		.FillHeight(1)
 		.VAlign(VAlign_Center)
 		.HAlign(HAlign_Fill)
-		.Padding(5)
+		.Padding(5.f)
 		[
 			SNew( SCheckBox )
 			.IsChecked( this, &SSurfaceProperties::IsUsingNegativePanning, Channels[i] )
@@ -206,7 +206,7 @@ TSharedRef<SWidget> SSurfaceProperties::ConstructTexturePan()
 	
 	TSharedPtr<SUniformGridPanel> GridBox;	
 	HorizontalBox->AddSlot()
-	.FillWidth(1)
+	.FillWidth(1.f)
 	[
 		SAssignNew(GridBox,SUniformGridPanel)
 	];
@@ -291,7 +291,7 @@ TSharedRef<SWidget> SSurfaceProperties::ConstructTextureRotate()
 	
 	Parent->AddSlot()
 	.AutoHeight()
-	.Padding(0,0,5,5)
+	.Padding(0.f,0.f,5.f,5.f)
 	[
 		SNew (STextBlock)
 		.Text(LOCTEXT("RotateTitle", "Rotate:"))
@@ -308,7 +308,7 @@ TSharedRef<SWidget> SSurfaceProperties::ConstructTextureRotate()
 	.AutoWidth()
 	.HAlign(HAlign_Center)
 	.VAlign(VAlign_Center)
-	.Padding(5)
+	.Padding(5.f)
 	[
 		SNew( SCheckBox )
 		.IsChecked( this, &SSurfaceProperties::IsUsingNegativeRotation )
@@ -335,7 +335,7 @@ TSharedRef<SWidget> SSurfaceProperties::ConstructTextureRotate()
 			TSharedPtr<SComboButton> CurrentButton;
 			TSharedPtr<SWidget> NumberBox;
 			RotateBox->AddSlot()
-			.FillWidth(1)
+			.FillWidth(1.f)
 			.VAlign(VAlign_Fill)
 			.HAlign(HAlign_Fill)
 			[
@@ -370,7 +370,7 @@ TSharedRef<SWidget> SSurfaceProperties::ConstructTextureRotate()
 		else
 		{
 			RotateBox->AddSlot()
-			.FillWidth(1)
+			.FillWidth(1.f)
 			.VAlign(VAlign_Fill)
 			.HAlign(HAlign_Fill)
 			[	
@@ -395,7 +395,7 @@ TSharedRef<SWidget> SSurfaceProperties::ConstructTextureFlip()
 	TSharedRef<SVerticalBox> Parent = SNew(SVerticalBox);
 	Parent->AddSlot()
 	.AutoHeight()
-	.Padding(0,0,5,5)
+	.Padding(0.f,0.f,5.f,5.f)
 	[
 		SNew (STextBlock)
 		.Text(LOCTEXT("FlipTitle", "Flip:"))
@@ -422,7 +422,7 @@ TSharedRef<SWidget> SSurfaceProperties::ConstructTextureFlip()
 			SNew(SButton)
 			.VAlign(VAlign_Center)
 			.HAlign(HAlign_Center)
-			.ContentPadding(FMargin(0,5))
+			.ContentPadding(FMargin(0.f, 5.f))
 			.Text(ButtonFields[Idx])
 			.ToolTipText(LOCTEXT("FlipToolTip", "Flip texture"))
 			.OnClicked(this,&SSurfaceProperties::OnFlipTexture,TextureCoordinateChannels[Idx])
@@ -439,7 +439,7 @@ TSharedRef<SWidget> SSurfaceProperties::ConstructTextureScale()
 	
 	Parent->AddSlot()
 	.AutoHeight()
-	.Padding(0,0,0,5)
+	.Padding(0.f,0.f,0.f,5.f)
 	[
 		SNew(SHorizontalBox)
 
@@ -472,10 +472,10 @@ TSharedRef<SWidget> SSurfaceProperties::ConstructTextureScale()
 	{
 		Controls->AddSlot()
 		.HAlign(HAlign_Fill)
-		.FillWidth(1)
+		.FillWidth(1.f)
 		[
 			SAssignNew(NewComboButton, SComboButton)
-			.ContentPadding(0)
+			.ContentPadding(0.f)
 			.HAlign(HAlign_Fill)
 			.ButtonContent()
 			[
@@ -520,7 +520,7 @@ TSharedRef<SWidget> SSurfaceProperties::ConstructTextureScale()
 	];
 
 	Controls->AddSlot()
-	.FillWidth(1);
+	.FillWidth(1.f);
 
 	Controls->AddSlot()
 	.AutoWidth()
@@ -583,7 +583,7 @@ TSharedRef<SWidget> SSurfaceProperties::ConstructLighting()
 
 	Parent->AddSlot()
 	.AutoHeight()
-	.Padding(0,0,0,5)
+	.Padding(0.f,0.f,0.f,5.f)
 	[
 		SNew(STextBlock)
 		.Text(LOCTEXT("LightingTitle", "Lighting:"))
@@ -595,7 +595,7 @@ TSharedRef<SWidget> SSurfaceProperties::ConstructLighting()
 		SNew(SHorizontalBox)
 		+SHorizontalBox::Slot()
 		.AutoWidth()
-		.Padding(0,0,10,5)
+		.Padding(0.f,0.f,10.f,5.f)
 		[
 			SNew(STextBlock)
 			.Text(LOCTEXT("LightingLightMapResolution", "Lightmap Resolution:"))
@@ -603,8 +603,8 @@ TSharedRef<SWidget> SSurfaceProperties::ConstructLighting()
 
 		+SHorizontalBox::Slot()
 		.HAlign(HAlign_Fill)
-		.FillWidth(1)
-		.Padding(0,0,0,5)
+		.FillWidth(1.f)
+		.Padding(0.f,0.f,0.f,5.f)
 		[
 			SNew(SNumericEntryBox<float>)
 			.OnValueCommitted(this, &SSurfaceProperties::OnLightmapResolutionCommitted)
@@ -614,7 +614,7 @@ TSharedRef<SWidget> SSurfaceProperties::ConstructLighting()
 
 		+SHorizontalBox::Slot()
 		.HAlign(HAlign_Fill)
-		.FillWidth(1)
+		.FillWidth(1.f)
 	];
 
 	FDetailsViewArgs Args;
@@ -659,14 +659,14 @@ FReply SSurfaceProperties::OnPanTexture(int32 PanAmount, TextureCoordChannel Cha
 
 FReply SSurfaceProperties::OnRotateTexture(int32 RotationAmount, RotationAction Action)
 {
-	const float Mod =  bUseNegativeRotation ? -1 : 1;
+	const float Mod =  bUseNegativeRotation ? -1.f : 1.f;
 
 	const float RotateRadians = RotationAmount / 180.0f * PI;
 
-	const float UU = cos(RotateRadians);
+	const float UU = cosf(RotateRadians);
 	const float VV = UU;
-	const float UV = -sin(RotateRadians) * Mod;
-	const float VU = sin(RotateRadians) * Mod;
+	const float UV = -sinf(RotateRadians) * Mod;
+	const float VU = sinf(RotateRadians) * Mod;
 	GUnrealEd->Exec( GetWorld(), *FString::Printf( TEXT("POLY TEXMULT UU=%f VV=%f UV=%f VU=%f"), UU, VV, UV, VU ) );
 
 	return FReply::Handled();
@@ -772,7 +772,7 @@ TOptional<float> SSurfaceProperties::GetLightmapResolutionValue() const
 
 void SSurfaceProperties::OnLightmapResolutionCommitted(float NewValue, ETextCommit::Type CommitInfo)
 {
-	const float LightMapScale = FMath::Clamp<float>(NewValue, 0.1f, 65536.0);
+	const float LightMapScale = FMath::Clamp<float>(NewValue, 0.1f, 65536.0f);
 
 	bool bSurfacesDirty = false;
 	for ( int32 LevelIndex = 0 ; LevelIndex < GetWorld()->GetNumLevels(); ++LevelIndex )
diff --git a/Engine/Source/Editor/LevelEditor/Private/ViewportToolbar/LevelEditorSubmenus.cpp b/Engine/Source/Editor/LevelEditor/Private/ViewportToolbar/LevelEditorSubmenus.cpp
--- a/Engine/Source/Editor/LevelEditor/Private/ViewportToolbar/LevelEditorSubmenus.cpp
+++ b/Engine/Source/Editor/LevelEditor/Private/ViewportToolbar/LevelEditorSubmenus.cpp
@@ -38,12 +38,12 @@
 namespace UE::LevelEditor::Private
 {
 
-bool IsLandscapeLODSettingChecked(FLevelEditorViewportClient& ViewportClient, int32 Value)
+bool IsLandscapeLODSettingChecked(FLevelEditorViewportClient& ViewportClient, int8 Value)
 {
 	return ViewportClient.LandscapeLODOverride == Value;
 }
 
-void OnLandscapeLODChanged(FLevelEditorViewportClient& ViewportClient, int32 NewValue)
+void OnLandscapeLODChanged(FLevelEditorViewportClient& ViewportClient, int8 NewValue)
 {
 	ViewportClient.LandscapeLODOverride = NewValue;
 	ViewportClient.Invalidate();
@@ -814,7 +814,7 @@ void PopulateViewModesMenu(UToolMenu* InMenu)
 			FToolMenuSection& SubMenuSection =
 				InMenu->AddSection("LevelViewportLandScapeLOD", LOCTEXT("LandscapeLODHeader", "Landscape LOD"));
 
-			auto CreateLandscapeLODAction = [WeakViewport](int32 LODValue)
+			auto CreateLandscapeLODAction = [WeakViewport](int8 LODValue)
 			{
 				FUIAction LandscapeLODAction;
 				LandscapeLODAction.ExecuteAction = FExecuteAction::CreateLambda(
@@ -855,7 +855,7 @@ void PopulateViewModesMenu(UToolMenu* InMenu)
 			SubMenuSection.AddSeparator("LandscapeLODSeparator");
 
 			static const FText FormatString = LOCTEXT("LandscapeLODFixed", "Fixed at {0}");
-			for (int32 i = 0; i < 8; ++i)
+			for (int8 i = 0; i < 8; ++i)
 			{
 				SubMenuSection.AddMenuEntry(
 					NAME_None,
@@ -1180,7 +1180,7 @@ FToolMenuEntry CreateShowHLODsSubmenu()
 											}
 										)
 										.OnEndSliderMovement_Lambda(
-											[](float)
+											[](double)
 											{
 												FSlateThrottleManager::Get().DisableThrottle(false);
 											}
@@ -1230,7 +1230,7 @@ FToolMenuEntry CreateShowHLODsSubmenu()
 											}
 										)
 										.OnEndSliderMovement_Lambda(
-											[](float)
+											[](double)
 											{
 												FSlateThrottleManager::Get().DisableThrottle(false);
 											}
diff --git a/Engine/Source/Editor/LevelEditor/Public/SLevelViewport.h b/Engine/Source/Editor/LevelEditor/Public/SLevelViewport.h
--- a/Engine/Source/Editor/LevelEditor/Public/SLevelViewport.h
+++ b/Engine/Source/Editor/LevelEditor/Public/SLevelViewport.h
@@ -855,7 +855,7 @@ private:
 	TArray<UObject*> DroppedObjects;
 
 	/** Caching off of the DragDropEvent Local Mouse Position grabbed from OnDrop */
-	FVector2D CachedOnDropLocalMousePos;
+	FIntPoint CachedOnDropLocalMousePos;
 
 	/** Weak pointer to the highres screenshot dialog if it's open. Will become invalid if UI is closed whilst the viewport is open */
 	TWeakPtr<class SWindow> HighResScreenshotDialog;
