[Viewport Toolbar] Camera menu is too long - need to reorganize

- Moves Piloting and Camera Movement entries to dedicated submenus
- Renames snap entries label to match new design

diff --git a/Engine/Source/Editor/LevelEditor/Private/LevelEditorActions.cpp b/Engine/Source/Editor/LevelEditor/Private/LevelEditorActions.cpp
--- a/Engine/Source/Editor/LevelEditor/Private/LevelEditorActions.cpp
+++ b/Engine/Source/Editor/LevelEditor/Private/LevelEditorActions.cpp
@@ -3911,8 +3911,8 @@ void FLevelEditorCommands::RegisterCommands()
 
 	UI_COMMAND( GoHere, "Go Here", "Moves the camera to the current mouse position", EUserInterfaceActionType::Button, FInputChord() );
 
-	UI_COMMAND( SnapCameraToObject, "Snap View to Object", "Snaps the view to the selected object", EUserInterfaceActionType::Button, FInputChord() );
-	UI_COMMAND( SnapObjectToCamera, "Snap Object to View", "Snaps the selected object to the view", EUserInterfaceActionType::Button, FInputChord() );
+	UI_COMMAND( SnapCameraToObject, "Move Camera to Object", "Move the current camera to match the location and rotation of the selected object.", EUserInterfaceActionType::Button, FInputChord() );
+	UI_COMMAND( SnapObjectToCamera, "Move Object to Camera", "Move the selected object to match the location and rotation of the current camera.", EUserInterfaceActionType::Button, FInputChord() );
 
 	UI_COMMAND(CopyActorFilePathtoClipboard, "Copy Selected Actor(s) File Path", "Copy the file path of the selected actors", EUserInterfaceActionType::Button, FInputChord());
 	UI_COMMAND(SaveActor, "Save Selected Actor(s)", "Save the selected actors", EUserInterfaceActionType::Button, FInputChord());
diff --git a/Engine/Source/Editor/LevelEditor/Private/ViewportToolbar/LevelEditorSubmenus.cpp b/Engine/Source/Editor/LevelEditor/Private/ViewportToolbar/LevelEditorSubmenus.cpp
--- a/Engine/Source/Editor/LevelEditor/Private/ViewportToolbar/LevelEditorSubmenus.cpp
+++ b/Engine/Source/Editor/LevelEditor/Private/ViewportToolbar/LevelEditorSubmenus.cpp
@@ -520,6 +520,135 @@ FToolMenuEntry CreateActorSnapCheckboxMenu()
 	return ActorSnapping;
 }
 
+void AddPilotSubmenu(UToolMenu* InMenu)
+{
+	if (!InMenu)
+	{
+		return;
+	}
+
+	ULevelViewportContext* LevelViewportContext = InMenu->FindContext<ULevelViewportContext>();
+	if (!LevelViewportContext)
+	{
+		return;
+	}
+
+	TWeakPtr<::SLevelViewport> LevelViewportWeak = LevelViewportContext->LevelViewport;
+
+	FToolMenuInsert InsertPosition("LevelViewportCameraType_Ortho", EToolMenuInsertType::After);
+
+	FToolMenuSection& PostOrthoSection = InMenu->FindOrAddSection("CameraSubmenuPostOrtho", FText(), InsertPosition);
+	PostOrthoSection.AddSeparator("CameraSubmenuSeparator");
+
+	PostOrthoSection.AddSubMenu(
+		"PilotingSubmenu",
+		LOCTEXT("PilotingSubmenu", "Pilot"),
+		LOCTEXT("PilotingSubmenu_ToolTip", "Piloting cameras and actors"),
+		FNewToolMenuDelegate::CreateLambda(
+			[LevelViewportWeak](UToolMenu* InMenu)
+			{
+				FToolMenuSection& PilotSection = InMenu->FindOrAddSection("Pilot");
+
+				PilotSection.AddMenuEntry(FLevelViewportCommands::Get().SelectPilotedActor);
+
+				bool bShowPilotSelectedActorEntry = false;
+
+				AActor* SelectedActor = nullptr;
+				if (TSharedPtr<::SLevelViewport> LevelViewport = LevelViewportWeak.Pin())
+				{
+					TArray<AActor*> SelectedActors;
+					GEditor->GetSelectedActors()->GetSelectedObjects<AActor>(SelectedActors);
+
+					if (!SelectedActors.IsEmpty() && !LevelViewport->IsSelectedActorLocked())
+					{
+						SelectedActor = SelectedActors[0];
+						const FLevelEditorViewportClient& ViewportClient = LevelViewport->GetLevelViewportClient();
+
+						bShowPilotSelectedActorEntry = SelectedActor && ViewportClient.IsPerspective()
+													&& !ViewportClient.IsLockedToCinematic();
+					}
+				}
+
+				if (bShowPilotSelectedActorEntry)
+				{
+					// Pilot Selected Actor Entry
+					PilotSection.AddMenuEntry(
+						FLevelViewportCommands::Get().PilotSelectedActor,
+						FText::Format(LOCTEXT("PilotActor", "Pilot '{0}'"), FText::FromString(SelectedActor->GetActorLabel()))
+					);
+				}
+
+				// Stop Piloting Entry
+				PilotSection.AddEntry(UE::LevelEditor::Private::CreateEjectActorPilotEntry());
+
+				// Exact Camera View Entry
+				{
+					FToolMenuEntry& ToggleCameraView =
+						PilotSection.AddMenuEntry(FLevelViewportCommands::Get().ToggleActorPilotCameraView);
+					ToggleCameraView.Label = LOCTEXT("ToggleCameraViewLabel", "Exact Camera View");
+					ToggleCameraView.SetShowInToolbarTopLevel(
+						TAttribute<bool>::CreateLambda(
+							[LevelViewportWeak]()
+							{
+								if (TSharedPtr<::SLevelViewport> EditorViewport = LevelViewportWeak.Pin())
+								{
+									return EditorViewport->IsAnyActorLocked();
+								}
+								return false;
+							}
+						)
+					);
+				}
+			}
+		),
+		false,
+		FSlateIcon(FAppStyle::Get().GetStyleSetName(), "LevelViewport.PilotSelectedActor")
+	);
+}
+
+void AddCameraMovementSubmenu(UToolMenu* InMenu)
+{
+	if (!InMenu)
+	{
+		return;
+	}
+
+	FToolMenuInsert InsertPosition("PilotingSubmenu", EToolMenuInsertType::After);
+
+	FToolMenuSection& PostOrthoSection = InMenu->FindOrAddSection("CameraSubmenuPostOrtho", FText(), InsertPosition);
+	PostOrthoSection.AddSeparator("CameraSubmenuSeparator");
+	PostOrthoSection.AddSubMenu(
+		"CameraMovementSubmenu",
+		LOCTEXT("CameraMovementSubmenu", "Camera Movement"),
+		LOCTEXT("CameraMovementSubmenu_ToolTip", "Camera movement options"),
+		FNewToolMenuDelegate::CreateLambda(
+			[](UToolMenu* InMenu)
+			{
+				FToolMenuSection& CameraMovementSection =
+					InMenu->FindOrAddSection("CameraMovement", LOCTEXT("CameraMovementLabel", "Camera Movement"));
+
+				CameraMovementSection.AddEntry(UE::UnrealEd::CreateCameraSpeedMenu());
+
+				// Frame Selection
+				{
+					FToolMenuEntry FocusViewportToSelection =
+						FToolMenuEntry::InitMenuEntry(FEditorViewportCommands::Get().FocusViewportToSelection);
+					FocusViewportToSelection.UserInterfaceActionType = EUserInterfaceActionType::ToggleButton;
+					FocusViewportToSelection.Icon = FSlateIcon(FAppStyle::GetAppStyleSetName(), "Icons.FrameActor");
+					CameraMovementSection.AddEntry(FocusViewportToSelection);
+				}
+
+				// Move Camera/Object
+				{
+					CameraMovementSection.AddMenuEntry(FLevelEditorCommands::Get().SnapCameraToObject);
+					CameraMovementSection.AddMenuEntry(FLevelEditorCommands::Get().SnapObjectToCamera);
+				}
+			}
+		),
+		false,
+		FSlateIcon(FSlateIcon(FAppStyle::GetAppStyleSetName(), "ClassIcon.CameraComponent"))
+	);
+}
 
 } // namespace UE::LevelEditor::Private
 
@@ -2447,68 +2576,15 @@ void ExtendCameraSubmenu(FName InCameraOptionsSubmenuName)
 					AddCameraActorSelectSection(InDynamicMenu);
 				}
 
-				// Pilot Section
+				// Post Ortho Viewport
+				// A section of entries shown right after orthographic view selection
 				{
-					FToolMenuSection& PilotSection = InDynamicMenu->FindOrAddSection("Pilot");
-
-					PilotSection.InsertPosition =
-						FToolMenuInsert("LevelViewportCameraType_Ortho", EToolMenuInsertType::After);
-					PilotSection.AddSeparator("PilotSectionSeparator");
-
-					PilotSection.AddMenuEntry(FLevelViewportCommands::Get().SelectPilotedActor);
-
-					bool bShowPilotSelectedActorEntry = false;
-
-					AActor* SelectedActor = nullptr;
-					if (TSharedPtr<::SLevelViewport> LevelViewport = LevelViewportWeak.Pin())
-					{
-						TArray<AActor*> SelectedActors;
-						GEditor->GetSelectedActors()->GetSelectedObjects<AActor>(SelectedActors);
-
-						if (!SelectedActors.IsEmpty() && !LevelViewport->IsSelectedActorLocked())
-						{
-							SelectedActor = SelectedActors[0];
-							const FLevelEditorViewportClient& ViewportClient = LevelViewport->GetLevelViewportClient();
-
-							bShowPilotSelectedActorEntry = SelectedActor && ViewportClient.IsPerspective()
-														&& !ViewportClient.IsLockedToCinematic();
-						}
-					}
-
-					if (bShowPilotSelectedActorEntry)
-					{
-						// Pilot Selected Actor Entry
-						PilotSection.AddMenuEntry(
-							FLevelViewportCommands::Get().PilotSelectedActor,
-							FText::Format(
-								LOCTEXT("PilotActor", "Pilot '{0}'"), FText::FromString(SelectedActor->GetActorLabel())
-							)
-						);
-					}
-
-					// Stop Piloting Entry
-					PilotSection.AddEntry(UE::LevelEditor::Private::CreateEjectActorPilotEntry());
-
-					// Exact Camera View Entry
-					{
-						FToolMenuEntry& ToggleCameraView =
-							PilotSection.AddMenuEntry(FLevelViewportCommands::Get().ToggleActorPilotCameraView);
-						ToggleCameraView.Label = LOCTEXT("ToggleCameraViewLabel", "Exact Camera View");
-						ToggleCameraView.SetShowInToolbarTopLevel(TAttribute<bool>::CreateLambda(
-							[LevelViewportWeak]()
-							{
-								if (TSharedPtr<::SLevelViewport> EditorViewport = LevelViewportWeak.Pin())
-								{
-									return EditorViewport->IsAnyActorLocked();
-								}
-								return false;
-							}
-						));
-					}
-
-					UE::UnrealEd::GenerateViewportTypeMenu(InDynamicMenu);
+					UE::LevelEditor::Private::AddPilotSubmenu(InDynamicMenu);
+					UE::LevelEditor::Private::AddCameraMovementSubmenu(InDynamicMenu);
 				}
 
+				UE::UnrealEd::GenerateViewportTypeMenu(InDynamicMenu);
+
 				// Create Section
 				{
 					FToolMenuSection& CreateSection =
@@ -2543,33 +2619,6 @@ void ExtendCameraSubmenu(FName InCameraOptionsSubmenuName)
 					);
 				}
 
-				// Positioning Section
-				{
-					FToolMenuSection& PositioningSection =
-						InDynamicMenu->FindOrAddSection("Positioning", LOCTEXT("PositioningLabel", "Positioning"));
-
-					PositioningSection.AddEntry(UE::UnrealEd::CreateCameraSpeedMenu());
-
-					PositioningSection.AddSeparator("PositioningSeparator_1");
-
-					// Frame Selection
-					{
-						FToolMenuEntry FocusViewportToSelection =
-							FToolMenuEntry::InitMenuEntry(FEditorViewportCommands::Get().FocusViewportToSelection);
-						FocusViewportToSelection.UserInterfaceActionType = EUserInterfaceActionType::ToggleButton;
-						FocusViewportToSelection.Icon = FSlateIcon(FAppStyle::GetAppStyleSetName(), "Icons.FrameActor");
-						PositioningSection.AddEntry(FocusViewportToSelection);
-					}
-
-					PositioningSection.AddSeparator("PositioningSeparator_2");
-
-					// Move Camera/Object
-					{
-						PositioningSection.AddMenuEntry(FLevelEditorCommands::Get().SnapCameraToObject);
-						PositioningSection.AddMenuEntry(FLevelEditorCommands::Get().SnapObjectToCamera);
-					}
-				}
-
 				// Options Section
 				{
 					FToolMenuSection& OptionsSection =
