Fix HighResScreenshot tool.

diff --git a/Engine/Source/Editor/LevelEditor/Private/HighResScreenshotUI.cpp b/Engine/Source/Editor/LevelEditor/Private/HighResScreenshotUI.cpp
--- a/Engine/Source/Editor/LevelEditor/Private/HighResScreenshotUI.cpp
+++ b/Engine/Source/Editor/LevelEditor/Private/HighResScreenshotUI.cpp
@@ -399,8 +399,8 @@ FReply SHighResScreenshotDialog::OnCaptureClicked()
 		auto ConfigViewport = Config.TargetViewport.Pin();
 		if (ConfigViewport.IsValid())
 		{
-			GScreenshotResolutionX = FMath::AsUInt(ConfigViewport->GetRenderTargetTextureSizeXY().X * Config.ResolutionMultiplier);
-			GScreenshotResolutionY = FMath::AsUInt(ConfigViewport->GetRenderTargetTextureSizeXY().Y * Config.ResolutionMultiplier);
+			GScreenshotResolutionX = uint32(ConfigViewport->GetRenderTargetTextureSizeXY().X * Config.ResolutionMultiplier);
+			GScreenshotResolutionY = uint32(ConfigViewport->GetRenderTargetTextureSizeXY().Y * Config.ResolutionMultiplier);
 			FIntRect ScaledCaptureRegion = Config.UnscaledCaptureRegion;
 
 			if (ScaledCaptureRegion.Area() > 0)
