[Viewport Toolbar] Expose Editor Preferences in the "Camera" menu of the new level editor viewport toolbar

diff --git a/Engine/Source/Editor/LevelEditor/Private/LevelEditor.cpp b/Engine/Source/Editor/LevelEditor/Private/LevelEditor.cpp
--- a/Engine/Source/Editor/LevelEditor/Private/LevelEditor.cpp
+++ b/Engine/Source/Editor/LevelEditor/Private/LevelEditor.cpp
@@ -1819,6 +1819,26 @@ void FLevelEditorModule::BindGlobalLevelEditorCommands()
 		FIsActionChecked::CreateStatic(&FLevelEditorActionCallbacks::IsFeatureLevelPreviewActive),
 		FIsActionButtonVisible::CreateStatic(FLevelEditorActionCallbacks::IsPreviewModeButtonVisible));
 
+	ActionList.MapAction(
+		Commands.OrbitCameraAroundSelection,
+		FExecuteAction::CreateStatic(&FLevelEditorActionCallbacks::ToggleOrbitCameraAroundSelection),
+		FCanExecuteAction(),
+		FIsActionChecked::CreateStatic(&FLevelEditorActionCallbacks::IsOrbitCameraAroundSelectionChecked));
+
+	ActionList.MapAction(
+		Commands.LinkOrthographicViewports,
+		FExecuteAction::CreateStatic(&FLevelEditorActionCallbacks::ToggleLinkOrthographicViewports),
+		FCanExecuteAction(),
+		FIsActionChecked::CreateStatic(&FLevelEditorActionCallbacks::IsLinkOrthographicViewportsChecked));
+
+	ActionList.MapAction(
+		Commands.OrthoZoomToCursor,
+		FExecuteAction::CreateStatic(&FLevelEditorActionCallbacks::ToggleOrthoZoomToCursor),
+		FCanExecuteAction(),
+		FIsActionChecked::CreateStatic(&FLevelEditorActionCallbacks::IsOrthoZoomToCursorChecked));
+
+	
+
 	const TArray<FPreviewPlatformMenuItem>& MenuItems = FDataDrivenPlatformInfoRegistry::GetAllPreviewPlatformMenuItems();
 	// We need one extra slot for the Disable Preview option
 	check(MenuItems.Num() == Commands.PreviewPlatformOverrides.Num());
diff --git a/Engine/Source/Editor/LevelEditor/Private/LevelEditorActions.cpp b/Engine/Source/Editor/LevelEditor/Private/LevelEditorActions.cpp
--- a/Engine/Source/Editor/LevelEditor/Private/LevelEditorActions.cpp
+++ b/Engine/Source/Editor/LevelEditor/Private/LevelEditorActions.cpp
@@ -3774,6 +3774,83 @@ bool FLevelEditorActionCallbacks::GeometryCollection_IsChecked()
 	return true;
 }
 
+void FLevelEditorActionCallbacks::TogglePreviewSelectedCameras(const TWeakPtr<::SLevelViewport>& InLevelViewportWeak)
+{
+	if (ULevelEditorViewportSettings* ViewportSettings = GetMutableDefault<ULevelEditorViewportSettings>())
+	{
+		ViewportSettings->bPreviewSelectedCameras = !ViewportSettings->bPreviewSelectedCameras;
+
+		if (TSharedPtr<::SLevelViewport> LevelViewport = InLevelViewportWeak.Pin())
+		{
+			LevelViewport->OnPreviewSelectedCamerasChange();
+		}
+	}
+}
+
+bool FLevelEditorActionCallbacks::IsPreviewSelectedCamerasChecked()
+{
+	if (const ULevelEditorViewportSettings* const ViewportSettings = GetDefault<ULevelEditorViewportSettings>())
+	{
+		return ViewportSettings->bPreviewSelectedCameras;
+	}
+
+	return false;
+}
+
+void FLevelEditorActionCallbacks::ToggleOrbitCameraAroundSelection()
+{
+	if (ULevelEditorViewportSettings* ViewportSettings = GetMutableDefault<ULevelEditorViewportSettings>())
+	{
+		ViewportSettings->bOrbitCameraAroundSelection = !ViewportSettings->bOrbitCameraAroundSelection;
+	}
+}
+
+bool FLevelEditorActionCallbacks::IsOrbitCameraAroundSelectionChecked()
+{
+	if (const ULevelEditorViewportSettings* const ViewportSettings = GetDefault<ULevelEditorViewportSettings>())
+	{
+		return ViewportSettings->bOrbitCameraAroundSelection;
+	}
+
+	return false;
+}
+
+void FLevelEditorActionCallbacks::ToggleLinkOrthographicViewports()
+{
+	if (ULevelEditorViewportSettings* ViewportSettings = GetMutableDefault<ULevelEditorViewportSettings>())
+	{
+		ViewportSettings->bUseLinkedOrthographicViewports = !ViewportSettings->bUseLinkedOrthographicViewports;
+	}
+}
+
+bool FLevelEditorActionCallbacks::IsLinkOrthographicViewportsChecked()
+{
+	if (const ULevelEditorViewportSettings* const ViewportSettings = GetDefault<ULevelEditorViewportSettings>())
+	{
+		return ViewportSettings->bUseLinkedOrthographicViewports;
+	}
+
+	return false;
+}
+
+void FLevelEditorActionCallbacks::ToggleOrthoZoomToCursor()
+{
+	if (ULevelEditorViewportSettings* ViewportSettings = GetMutableDefault<ULevelEditorViewportSettings>())
+	{
+		ViewportSettings->bCenterZoomAroundCursor = !ViewportSettings->bCenterZoomAroundCursor;
+	}
+}
+
+bool FLevelEditorActionCallbacks::IsOrthoZoomToCursorChecked()
+{
+	if (const ULevelEditorViewportSettings* const ViewportSettings = GetDefault<ULevelEditorViewportSettings>())
+	{
+		return ViewportSettings->bCenterZoomAroundCursor;
+	}
+
+	return false;
+}
+
 class UWorld* FLevelEditorActionCallbacks::GetWorld()
 {
 	return GEditor->GetEditorWorldContext().World();
@@ -4130,6 +4207,15 @@ void FLevelEditorCommands::RegisterCommands()
 
 	UI_COMMAND(ToggleFeatureLevelPreview, "Preview Mode Toggle", "Toggles the Preview Mode on or off for the currently selected Preview target", EUserInterfaceActionType::ToggleButton, FInputChord());
 
+	// Camera Preferences
+	UI_COMMAND(OrbitCameraAroundSelection, "Orbit Around Selection", "If enabled, the camera will orbit around the current selection in the viewport", EUserInterfaceActionType::ToggleButton, FInputChord());
+	UI_COMMAND(LinkOrthographicViewports, "Link Ortho Camera Movement", "If checked all orthographic view ports are linked to the same position and move together.", EUserInterfaceActionType::ToggleButton, FInputChord());
+	UI_COMMAND(OrthoZoomToCursor, "Ortho Zoom to Cursor", "If checked, in orthographic viewports zooming will center on the mouse position.  If unchecked, the zoom is around the center of the viewport.", EUserInterfaceActionType::ToggleButton, FInputChord());
+
+	/*TSharedPtr<FUICommandInfo> LinkOrthographicViewports;
+	TSharedPtr<FUICommandInfo> OrthoZoomToCursor;
+	TSharedPtr<FUICommandInfo> DistanceBasedCameraSpeed;*/
+	
 	// Add preview platforms
 	TSet<FName> PreviewShaderPlatformNames;
 	for (const FPreviewPlatformMenuItem& Item : FDataDrivenPlatformInfoRegistry::GetAllPreviewPlatformMenuItems())
diff --git a/Engine/Source/Editor/LevelEditor/Private/ViewportToolbar/LevelEditorSubmenus.cpp b/Engine/Source/Editor/LevelEditor/Private/ViewportToolbar/LevelEditorSubmenus.cpp
--- a/Engine/Source/Editor/LevelEditor/Private/ViewportToolbar/LevelEditorSubmenus.cpp
+++ b/Engine/Source/Editor/LevelEditor/Private/ViewportToolbar/LevelEditorSubmenus.cpp
@@ -624,24 +624,30 @@ void AddCameraMovementSubmenu(UToolMenu* InMenu)
 		FNewToolMenuDelegate::CreateLambda(
 			[](UToolMenu* InMenu)
 			{
-				FToolMenuSection& CameraMovementSection =
-					InMenu->FindOrAddSection("CameraMovement", LOCTEXT("CameraMovementLabel", "Camera Movement"));
+				{
+					FToolMenuSection& CameraMovementSection =
+						InMenu->FindOrAddSection("CameraMovement", LOCTEXT("CameraMovementLabel", "Camera Movement"));
 
-				CameraMovementSection.AddEntry(UE::UnrealEd::CreateCameraSpeedMenu());
+					CameraMovementSection.AddEntry(UE::UnrealEd::CreateCameraSpeedMenu());
 
-				// Frame Selection
-				{
-					FToolMenuEntry FocusViewportToSelection =
-						FToolMenuEntry::InitMenuEntry(FEditorViewportCommands::Get().FocusViewportToSelection);
-					FocusViewportToSelection.UserInterfaceActionType = EUserInterfaceActionType::ToggleButton;
-					FocusViewportToSelection.Icon = FSlateIcon(FAppStyle::GetAppStyleSetName(), "Icons.FrameActor");
-					CameraMovementSection.AddEntry(FocusViewportToSelection);
-				}
+					// Frame Selection
+					{
+						FToolMenuEntry FocusViewportToSelection =
+							FToolMenuEntry::InitMenuEntry(FEditorViewportCommands::Get().FocusViewportToSelection);
+						FocusViewportToSelection.UserInterfaceActionType = EUserInterfaceActionType::ToggleButton;
+						FocusViewportToSelection.Icon = FSlateIcon(FAppStyle::GetAppStyleSetName(), "Icons.FrameActor");
+						CameraMovementSection.AddEntry(FocusViewportToSelection);
+					}
 
-				// Move Camera/Object
-				{
 					CameraMovementSection.AddMenuEntry(FLevelEditorCommands::Get().SnapCameraToObject);
 					CameraMovementSection.AddMenuEntry(FLevelEditorCommands::Get().SnapObjectToCamera);
+					CameraMovementSection.AddMenuEntry(FLevelEditorCommands::Get().OrbitCameraAroundSelection);
+
+					FToolMenuSection& OrthoSection =
+						InMenu->FindOrAddSection("Orthographic", LOCTEXT("OrthographicSectionLabel", "Orthographic"));
+
+					OrthoSection.AddMenuEntry(FLevelEditorCommands::Get().LinkOrthographicViewports);
+					OrthoSection.AddMenuEntry(FLevelEditorCommands::Get().OrthoZoomToCursor);
 				}
 			}
 		),
@@ -650,6 +656,97 @@ void AddCameraMovementSubmenu(UToolMenu* InMenu)
 	);
 }
 
+FToolMenuEntry CreatePreviewSelectedCamerasCheckBoxSubmenu()
+{
+	FNewToolMenuDelegate PreviewSelectedCamerasMenuDelegate = FNewToolMenuDelegate::CreateLambda(
+		[](UToolMenu* Submenu)
+		{
+			FToolMenuSection& CameraPreviewSection =
+				Submenu->FindOrAddSection("Camera Preview", LOCTEXT("CameraPreviewLabel", "Camera Preview"));
+
+			constexpr float PreviewSizeMin = 1.0f;
+			constexpr float PreviewSizeMax = 10.0f;
+
+			FToolMenuEntry PreviewSizeEntry = UnrealEd::CreateNumericEntry(
+				"PreviewSize",
+				LOCTEXT("PreviewSizeLabel", "Preview Size"),
+				LOCTEXT("PreviewSizeTooltip", "Affects the size of 'picture in picture' previews if they are enabled"),
+				FCanExecuteAction(),
+				UnrealEd::FNumericEntryExecuteActionDelegate::CreateLambda(
+					[Context = Submenu->Context](float InValue)
+					{
+						if (ULevelEditorViewportSettings* const ViewportSettings =
+								GetMutableDefault<ULevelEditorViewportSettings>())
+						{
+							ViewportSettings->CameraPreviewSize = InValue;
+
+							// If preview is not on, we assume the user wants to turn it on as they are editing preview size, so let's toggle it
+							if (!FLevelEditorActionCallbacks::IsPreviewSelectedCamerasChecked())
+							{
+								TWeakPtr<::SLevelViewport> LevelViewportWeak;
+								if (ULevelViewportContext* const LevelViewportContext =
+										Context.FindContext<ULevelViewportContext>())
+								{
+									LevelViewportWeak = LevelViewportContext->LevelViewport;
+								}
+								FLevelEditorActionCallbacks::TogglePreviewSelectedCameras(LevelViewportWeak);
+							}
+						}
+					}
+				),
+				TAttribute<float>::CreateLambda(
+					[]()
+					{
+						if (const ULevelEditorViewportSettings* const ViewportSettings =
+								GetDefault<ULevelEditorViewportSettings>())
+						{
+							return ViewportSettings->CameraPreviewSize;
+						}
+						return PreviewSizeMin;
+					}
+				),
+				PreviewSizeMin,
+				PreviewSizeMax
+			);
+
+			CameraPreviewSection.AddEntry(PreviewSizeEntry);
+		}
+	);
+
+	FText PreviewSelectedCamerasTooltip = LOCTEXT(
+		"CameraPreviewWindowTooltip", "When enabled, selecting a camera actor will display a live 'picture in picture' preview from the camera's perspective within the current editor view port.  This can be used to easily tweak camera positioning, post-processing and other settings without having to possess the camera itself.  This feature may reduce application performance when enabled."
+	);
+
+	FToolMenuEntry Entry = UnrealEd::CreateCheckboxSubmenu(
+		"PreviewSelectedCameras",
+		LOCTEXT("PreviewSelectedCamerasLabel", "Preview Selected Cameras"),
+		PreviewSelectedCamerasTooltip,
+		FToolMenuExecuteAction::CreateLambda(
+			[](const FToolMenuContext& InContext)
+			{
+				TWeakPtr<::SLevelViewport> LevelViewportWeak;
+				if (ULevelViewportContext* const LevelViewportContext = InContext.FindContext<ULevelViewportContext>())
+				{
+					LevelViewportWeak = LevelViewportContext->LevelViewport;
+				}
+
+				FLevelEditorActionCallbacks::TogglePreviewSelectedCameras(LevelViewportWeak);
+			}
+		),
+		FToolMenuCanExecuteAction(),
+		FToolMenuGetActionCheckState::CreateLambda(
+			[](const FToolMenuContext& InContext)
+			{
+				return FLevelEditorActionCallbacks::IsPreviewSelectedCamerasChecked() ? ECheckBoxState::Checked
+																					  : ECheckBoxState::Unchecked;
+			}
+		),
+		PreviewSelectedCamerasMenuDelegate
+	);
+
+	return Entry;
+}
+
 } // namespace UE::LevelEditor::Private
 
 namespace UE::LevelEditor
@@ -2640,6 +2737,7 @@ void ExtendCameraSubmenu(FName InCameraOptionsSubmenuName)
 					// This additional options section is used to force certain elements to appear after extensions
 					{
 						FToolMenuSection& AdditionalOptions = InDynamicMenu->FindOrAddSection("AdditionalOptions");
+						AdditionalOptions.AddEntry(UE::LevelEditor::Private::CreatePreviewSelectedCamerasCheckBoxSubmenu());
 						AdditionalOptions.AddSeparator("AdditionalOptionsSeparator");
 
 						FToolMenuEntry HighResolutionScreenshot =
diff --git a/Engine/Source/Editor/LevelEditor/Public/LevelEditorActions.h b/Engine/Source/Editor/LevelEditor/Public/LevelEditorActions.h
--- a/Engine/Source/Editor/LevelEditor/Public/LevelEditorActions.h
+++ b/Engine/Source/Editor/LevelEditor/Public/LevelEditorActions.h
@@ -17,6 +17,7 @@
 #include "Elements/Framework/TypedElementHandle.h"
 #include "Elements/Interfaces/TypedElementWorldInterface.h"
 
+class SLevelViewport;
 class FLightingBuildOptions;
 class SLevelEditor;
 class UActorFactory;
@@ -629,6 +630,13 @@ public:
 	TMap<FName, TArray<PreviewPlatformCommand>> PlatformToPreviewPlatformOverrides;
 	TMap<FName, TArray<PreviewPlatformCommand>> PlatformToPreviewJsonPlatformOverrides;
 
+	/**
+	 * Camera Preferences
+	 */
+	TSharedPtr<FUICommandInfo> OrbitCameraAroundSelection;
+	TSharedPtr<FUICommandInfo> LinkOrthographicViewports;
+	TSharedPtr<FUICommandInfo> OrthoZoomToCursor;
+
 	///**
 	// * Mode Commands                   
 	// */
@@ -867,6 +875,18 @@ public:
 	static void GeometryCollection_SelectInverseGeometry();
 	static bool GeometryCollection_IsChecked();
 
+	static void TogglePreviewSelectedCameras(const TWeakPtr<SLevelViewport>& InLevelViewportWeak = {});
+	static bool IsPreviewSelectedCamerasChecked();
+
+	static void ToggleOrbitCameraAroundSelection();
+	static bool IsOrbitCameraAroundSelectionChecked();
+
+	static void ToggleLinkOrthographicViewports();
+	static bool IsLinkOrthographicViewportsChecked();
+
+	static void ToggleOrthoZoomToCursor();
+	static bool IsOrthoZoomToCursorChecked();
+
 	/**
 	 * Called when the Scene Stats button is clicked.  Invokes the Primitive Stats dialog.
 	 */
diff --git a/Engine/Source/Editor/UnrealEd/Private/EditorViewportCommands.cpp b/Engine/Source/Editor/UnrealEd/Private/EditorViewportCommands.cpp
--- a/Engine/Source/Editor/UnrealEd/Private/EditorViewportCommands.cpp
+++ b/Engine/Source/Editor/UnrealEd/Private/EditorViewportCommands.cpp
@@ -209,6 +209,8 @@ void FEditorViewportCommands::RegisterCommands()
 	UI_COMMAND(OpenEditorPerformanceEditorPreferences, "Editor Preferences...", "Opens the editor preferences for viewport performance", EUserInterfaceActionType::Button, FInputChord());
 
 	UI_COMMAND(ToggleInViewportContextMenu, "In-Viewport Context Menu", "Shows a contextual menu of key properties and actions for the selected items in the viewport.", EUserInterfaceActionType::ToggleButton, FInputChord(EKeys::Tab));
+
+	UI_COMMAND(ToggleDistanceBasedCameraSpeed, "Distance Based Camera Speed", "If enabled, scale the perspective camera speed based on the distance between the camera and its look-at position.", EUserInterfaceActionType::ToggleButton, FInputChord());
 }
 
 #undef LOCTEXT_NAMESPACE
diff --git a/Engine/Source/Editor/UnrealEd/Private/SEditorViewport.cpp b/Engine/Source/Editor/UnrealEd/Private/SEditorViewport.cpp
--- a/Engine/Source/Editor/UnrealEd/Private/SEditorViewport.cpp
+++ b/Engine/Source/Editor/UnrealEd/Private/SEditorViewport.cpp
@@ -497,6 +497,12 @@ void SEditorViewport::BindCommands()
 		Commands.OpenEditorPerformanceEditorPreferences,
 		FExecuteAction::CreateSP(this, &SEditorViewport::OnOpenViewportPerformanceEditorPreferences));
 
+	CommandListRef.MapAction(
+		Commands.ToggleDistanceBasedCameraSpeed,
+		FExecuteAction::CreateStatic(&SEditorViewport::OnToggleDistanceBasedCameraSpeed),
+		FCanExecuteAction(),
+		FIsActionChecked::CreateStatic(&SEditorViewport::IsDistanceBasedCameraSpeedEnabled));
+
 	// Simple macro for binding many view mode UI commands
 
 #define MAP_VIEWMODEPARAM_ACTION( ViewModeCommand, ViewModeParam ) \
@@ -977,6 +983,24 @@ bool SEditorViewport::IsRotateToSurfaceNormalEnabled()
 	return Settings.bSnapRotation;
 }
 
+void SEditorViewport::OnToggleDistanceBasedCameraSpeed()
+{
+	if (ULevelEditorViewportSettings* ViewportSettings = GetMutableDefault<ULevelEditorViewportSettings>())
+	{
+		ViewportSettings->bUseDistanceScaledCameraSpeed = !ViewportSettings->bUseDistanceScaledCameraSpeed;
+	}
+}
+
+bool SEditorViewport::IsDistanceBasedCameraSpeedEnabled()
+{
+	if (const ULevelEditorViewportSettings* const ViewportSettings = GetDefault<ULevelEditorViewportSettings>())
+	{
+		return ViewportSettings->bUseDistanceScaledCameraSpeed;
+	}
+
+	return false;
+}
+
 TSharedPtr<SWidget> SEditorViewport::CreateViewportIndicatorWidget(const TAttribute<EVisibility>& InVisibility)
 {
 	// This makes a gradient that displays whether a viewport is active
diff --git a/Engine/Source/Editor/UnrealEd/Private/ViewportToolbar/UnrealEdViewportToolbar.cpp b/Engine/Source/Editor/UnrealEd/Private/ViewportToolbar/UnrealEdViewportToolbar.cpp
--- a/Engine/Source/Editor/UnrealEd/Private/ViewportToolbar/UnrealEdViewportToolbar.cpp
+++ b/Engine/Source/Editor/UnrealEd/Private/ViewportToolbar/UnrealEdViewportToolbar.cpp
@@ -3812,6 +3812,7 @@ FToolMenuEntry CreateCameraSpeedMenu()
 							{
 								CameraSpeedSection.AddEntry(CreateCameraSpeedSlider(Context->Viewport));
 								CameraSpeedSection.AddEntry(CreateCameraSpeedScalarSlider(Context->Viewport));
+								CameraSpeedSection.AddMenuEntry(FEditorViewportCommands::Get().ToggleDistanceBasedCameraSpeed);
 							}
 						}
 					),
diff --git a/Engine/Source/Editor/UnrealEd/Public/EditorViewportCommands.h b/Engine/Source/Editor/UnrealEd/Public/EditorViewportCommands.h
--- a/Engine/Source/Editor/UnrealEd/Public/EditorViewportCommands.h
+++ b/Engine/Source/Editor/UnrealEd/Public/EditorViewportCommands.h
@@ -261,6 +261,9 @@ public:
 	/** Enables or disables rotating to surface normal when snapping to surfaces */
 	TSharedPtr< FUICommandInfo > RotateToSurfaceNormal;
 
+	/** Enables or disables scaling the perspective camera speed, based on the distance between the camera and its look-at position */
+	TSharedPtr<FUICommandInfo> ToggleDistanceBasedCameraSpeed;
+
 public:
 	/** Registers our commands with the binding system */
 	virtual void RegisterCommands() override;
diff --git a/Engine/Source/Editor/UnrealEd/Public/SEditorViewport.h b/Engine/Source/Editor/UnrealEd/Public/SEditorViewport.h
--- a/Engine/Source/Editor/UnrealEd/Public/SEditorViewport.h
+++ b/Engine/Source/Editor/UnrealEd/Public/SEditorViewport.h
@@ -281,6 +281,16 @@ protected:
 	 */
 	static UNREALED_API bool IsRotateToSurfaceNormalEnabled();
 
+	/**
+	 * Called to enable/disable distance-based perspective camera speed scaling
+	 */
+	static UNREALED_API void OnToggleDistanceBasedCameraSpeed();
+
+	/**
+	 * Called to test whether distance-based perspective camera speed scaling is enabled
+	 */
+	static UNREALED_API bool IsDistanceBasedCameraSpeedEnabled();
+
 	/**
 	 * Creates a widget showing whether the viewport is active or not
 	 * @param InVisibility visibility attribute, can be used to dynamically hide the widget
