[Truncation Warnings] Enable truncation warnings in level editor module

#rnx

diff --git a/Engine/Source/Editor/LevelEditor/LevelEditor.Build.cs b/Engine/Source/Editor/LevelEditor/LevelEditor.Build.cs
--- a/Engine/Source/Editor/LevelEditor/LevelEditor.Build.cs
+++ b/Engine/Source/Editor/LevelEditor/LevelEditor.Build.cs
@@ -6,6 +6,9 @@ public class LevelEditor : ModuleRules
 {
 	public LevelEditor(ReadOnlyTargetRules Target) : base(Target)
 	{
+		// Enable truncation warnings in this module
+		CppCompileWarningSettings.UnsafeTypeCastWarningLevel = WarningLevel.Warning;
+
 		PrivateIncludePathModuleNames.AddRange(
 			new string[] {
 				"AssetTools",
diff --git a/Engine/Source/Editor/LevelEditor/Private/HighResScreenshotUI.cpp b/Engine/Source/Editor/LevelEditor/Private/HighResScreenshotUI.cpp
--- a/Engine/Source/Editor/LevelEditor/Private/HighResScreenshotUI.cpp
+++ b/Engine/Source/Editor/LevelEditor/Private/HighResScreenshotUI.cpp
@@ -399,8 +399,8 @@ FReply SHighResScreenshotDialog::OnCaptureClicked()
 		auto ConfigViewport = Config.TargetViewport.Pin();
 		if (ConfigViewport.IsValid())
 		{
-			GScreenshotResolutionX = uint32(ConfigViewport->GetRenderTargetTextureSizeXY().X * Config.ResolutionMultiplier);
-			GScreenshotResolutionY = uint32(ConfigViewport->GetRenderTargetTextureSizeXY().Y * Config.ResolutionMultiplier);
+			GScreenshotResolutionX = uint32(static_cast<float>(ConfigViewport->GetRenderTargetTextureSizeXY().X) * Config.ResolutionMultiplier);
+			GScreenshotResolutionY = uint32(static_cast<float>(ConfigViewport->GetRenderTargetTextureSizeXY().Y) * Config.ResolutionMultiplier);
 			FIntRect ScaledCaptureRegion = Config.UnscaledCaptureRegion;
 
 			if (ScaledCaptureRegion.Area() > 0)
diff --git a/Engine/Source/Editor/LevelEditor/Private/LevelEditorActions.cpp b/Engine/Source/Editor/LevelEditor/Private/LevelEditorActions.cpp
--- a/Engine/Source/Editor/LevelEditor/Private/LevelEditorActions.cpp
+++ b/Engine/Source/Editor/LevelEditor/Private/LevelEditorActions.cpp
@@ -1246,7 +1246,7 @@ int32 FLevelEditorActionCallbacks::GetLightingResolutionRatio()
 void FLevelEditorActionCallbacks::SetLightingResolutionRatio( int32 Value )
 {
 	FLightmapResRatioAdjustSettings& Settings = FLightmapResRatioAdjustSettings::Get();
-	const float NewValue = Value / 100.0f;
+	const float NewValue = static_cast<float>(Value) / 100.0f;
 	if ( Settings.Ratio != NewValue )
 	{
 		Settings.Ratio = NewValue;
diff --git a/Engine/Source/Editor/LevelEditor/Private/LevelEditorCreateActorMenu.cpp b/Engine/Source/Editor/LevelEditor/Private/LevelEditorCreateActorMenu.cpp
--- a/Engine/Source/Editor/LevelEditor/Private/LevelEditorCreateActorMenu.cpp
+++ b/Engine/Source/Editor/LevelEditor/Private/LevelEditorCreateActorMenu.cpp
@@ -154,7 +154,7 @@ class SAssetMenuEntry : public SCompoundWidget
 			ActorTypeDisplayName.IsEmpty()
 			? SNullWidget::NullWidget
 			: SNew(SBox)
-				.Padding(FMargin(0, 8, 0, 0))
+				.Padding(FMargin(0.f, 8.f, 0.f, 0.f))
 				[
 					SNew(STextBlock)
 					.Text(ActorTypeDisplayName)
@@ -163,26 +163,26 @@ class SAssetMenuEntry : public SCompoundWidget
 				];
 
 		ChildSlot
-		.Padding( FMargin(0, 0, 8, 0) )
+		.Padding( FMargin(0.f, 0.f, 8.f, 0.f) )
 		[
 			SNew( SBorder )
-			.Padding( 0 )
+			.Padding( 0.f )
 			.BorderImage( FAppStyle::Get().GetBrush("LevelViewportContextMenu.AssetTileItem.NameAreaBackground") )
 			[
 				SNew( SHorizontalBox )
 				+SHorizontalBox::Slot()
-				.Padding( 0, 0, 0, 0 )
+				.Padding( 0.f, 0.f, 0.f, 0.f)
 				.VAlign( VAlign_Center )
 				.AutoWidth()
 				[
 					SNew( SBorder )
-					//.Padding( FMargin(4, 0, 0, 0) ) // prevent contents from overlapping the rounded corners
-					.Padding(0)
+					//.Padding( FMargin(4.f, 0.f, 0.f, 0.f) ) // prevent contents from overlapping the rounded corners
+					.Padding(0.f)
 					.BorderImage( FAppStyle::Get().GetBrush("LevelViewportContextMenu.AssetTileItem.ThumbnailAreaBackground") )
 					[
 						SNew( SBox )
-						.WidthOverride( 48 )
-						.HeightOverride( 48 )
+						.WidthOverride( FOptionalSize(48.f) )
+						.HeightOverride( FOptionalSize(48.f) )
 						[
 							SNew( SMenuThumbnail, Asset )
 							.Width(48)
@@ -193,7 +193,7 @@ class SAssetMenuEntry : public SCompoundWidget
 
 				+SHorizontalBox::Slot()
 				.VAlign(VAlign_Center)
-				.Padding(5, 0, 10, 0)
+				.Padding(5.f, 0.f, 10.f, 0.f)
 				[
 					SNew( SVerticalBox )
 					+SVerticalBox::Slot()
diff --git a/Engine/Source/Editor/LevelEditor/Private/LevelEditorSubsystem.cpp b/Engine/Source/Editor/LevelEditor/Private/LevelEditorSubsystem.cpp
--- a/Engine/Source/Editor/LevelEditor/Private/LevelEditorSubsystem.cpp
+++ b/Engine/Source/Editor/LevelEditor/Private/LevelEditorSubsystem.cpp
@@ -803,8 +803,8 @@ public:
 					.Padding(1.f, 1.f, 1.f, 1.f)
 					[
 						SNew(SBox)
-						.WidthOverride(16)
-						.HeightOverride(16)
+						.WidthOverride(16.f)
+						.HeightOverride(16.f)
 						[
 							SNew(SImage)
 							.Image_Lambda([this]() { return FClassIconFinder::FindIconForActor(GetEditingLevelInstance()); })
diff --git a/Engine/Source/Editor/LevelEditor/Private/LevelEditorToolBar.cpp b/Engine/Source/Editor/LevelEditor/Private/LevelEditorToolBar.cpp
--- a/Engine/Source/Editor/LevelEditor/Private/LevelEditorToolBar.cpp
+++ b/Engine/Source/Editor/LevelEditor/Private/LevelEditorToolBar.cpp
@@ -1573,7 +1573,7 @@ TSharedRef< SWidget > FLevelEditorToolBar::MakeLevelEditorSecondaryModeToolbar(
 	}
 
 	return SNew(SBorder)
-	.Padding(0)
+	.Padding(0.f)
 	.BorderImage(FAppStyle::Get().GetBrush("NoBorder"))
 	.IsEnabled(FSlateApplication::Get().GetNormalExecutionAttribute())
 	[
@@ -1618,7 +1618,7 @@ TSharedRef< SWidget > FLevelEditorToolBar::MakeLevelEditorToolBar( const TShared
 			.AutoWidth()
 			[
 				SNew(SBorder)
-				.Padding(0)
+				.Padding(0.f)
 				.BorderImage(FAppStyle::Get().GetBrush("NoBorder"))
 				.IsEnabled(FSlateApplication::Get().GetNormalExecutionAttribute())
 				[
@@ -1629,7 +1629,7 @@ TSharedRef< SWidget > FLevelEditorToolBar::MakeLevelEditorToolBar( const TShared
 			.AutoWidth()
 			[
 				SNew(SBorder)
-				.Padding(0)
+				.Padding(0.f)
 				.BorderImage(FAppStyle::Get().GetBrush("NoBorder"))
 				.IsEnabled(FSlateApplication::Get().GetNormalExecutionAttribute())
 				[
@@ -1642,7 +1642,7 @@ TSharedRef< SWidget > FLevelEditorToolBar::MakeLevelEditorToolBar( const TShared
 			.AutoWidth()
 			[
 				SNew(SBorder)
-				.Padding(0)
+				.Padding(0.f)
 				.BorderImage(FAppStyle::Get().GetBrush("NoBorder"))
 				[
 					UToolMenus::Get()->GenerateWidget("LevelEditor.LevelEditorToolBar.PlayToolBar", MenuContext) // Always enabled
@@ -1653,7 +1653,7 @@ TSharedRef< SWidget > FLevelEditorToolBar::MakeLevelEditorToolBar( const TShared
 			.HAlign(HAlign_Left)
 			[
 				SNew(SBorder)
-				.Padding(0)
+				.Padding(0.f)
 				.BorderImage(FAppStyle::Get().GetBrush("NoBorder"))
 				.IsEnabled(FSlateApplication::Get().GetNormalExecutionAttribute())
 				[
@@ -1665,7 +1665,7 @@ TSharedRef< SWidget > FLevelEditorToolBar::MakeLevelEditorToolBar( const TShared
 			.Padding(0.0f, 0.0f, 7.0f, 0.0f)
 			[
 				SNew(SBorder)
-				.Padding(0)
+				.Padding(0.f)
 				.BorderImage(FAppStyle::Get().GetBrush("NoBorder"))
 				.IsEnabled(FSlateApplication::Get().GetNormalExecutionAttribute())
 				[
@@ -2108,8 +2108,8 @@ void FLevelEditorToolBar::RegisterCinematicsMenu()
 			InitOptions.bShowSearchBox = true;
 			InitOptions.bShowCreateNewFolder = false;
 
-			InitOptions.ColumnMap.Add(FSceneOutlinerBuiltInColumnTypes::Label(), FSceneOutlinerColumnInfo(ESceneOutlinerColumnVisibility::Visible, 0, FCreateSceneOutlinerColumn(), false, 0.6));
-			InitOptions.ColumnMap.Add(FSceneOutlinerBuiltInColumnTypes::ActorInfo(), FSceneOutlinerColumnInfo(ESceneOutlinerColumnVisibility::Visible, 10, FCreateSceneOutlinerColumn(), false, 0.4));
+			InitOptions.ColumnMap.Add(FSceneOutlinerBuiltInColumnTypes::Label(), FSceneOutlinerColumnInfo(ESceneOutlinerColumnVisibility::Visible, 0, FCreateSceneOutlinerColumn(), false, 0.6f));
+			InitOptions.ColumnMap.Add(FSceneOutlinerBuiltInColumnTypes::ActorInfo(), FSceneOutlinerColumnInfo(ESceneOutlinerColumnVisibility::Visible, 10, FCreateSceneOutlinerColumn(), false, 0.4f));
 
 			// Only display MovieScene actors
 			auto ActorFilter = [&](const AActor* Actor) {
@@ -2127,7 +2127,7 @@ void FLevelEditorToolBar::RegisterCinematicsMenu()
 			.MaxHeight(400.0f)
 			[
 				SNew(SBox)
-				.WidthOverride(360)
+				.WidthOverride(360.f)
 				[
 					SceneOutlinerModule.CreateActorPicker(
 						InitOptions,
diff --git a/Engine/Source/Editor/LevelEditor/Private/SActorDetails.cpp b/Engine/Source/Editor/LevelEditor/Private/SActorDetails.cpp
--- a/Engine/Source/Editor/LevelEditor/Private/SActorDetails.cpp
+++ b/Engine/Source/Editor/LevelEditor/Private/SActorDetails.cpp
@@ -69,7 +69,7 @@ namespace UE::LevelEditor::Private
 					[
 						SAssignNew(PromoteElementButton, SButton)
 						.ButtonStyle(FAppStyle::Get(), "DetailsView.NameAreaButton")
-						.ContentPadding(0)
+						.ContentPadding(0.f)
 						.HAlign(HAlign_Center)
 						.VAlign(VAlign_Center)
 						.AddMetaData<FTagMetaData>(FTagMetaData(TEXT("Element.PromoteElement")))
@@ -195,14 +195,14 @@ public:
 				.AutoWidth()
 				.HAlign(HAlign_Center)
 				.VAlign(VAlign_Center)
-				.Padding(2)
+				.Padding(2.f)
 				[
 					SNew(SImage)
 					.Image(FAppStyle::Get().GetBrush("Icons.Warning"))
 				]
 				+ SHorizontalBox::Slot()
 					.VAlign(VAlign_Center)
-					.Padding(2)
+					.Padding(2.f)
 					[
 						SNew(SRichTextBlock)
 						.DecoratorStyleSet(&FAppStyle::Get())
@@ -350,7 +350,7 @@ void SActorDetails::Construct(const FArguments& InArgs, UTypedElementSelectionSe
 				SNew( SVerticalBox )
 				+SVerticalBox::Slot()
 				.AutoHeight()
-				.Padding( FMargin(0,0,0,1) )
+				.Padding( FMargin(0.f,0.f,0.f,1.f) )
 				[
 					SNew(SActorDetailsUneditableComponentWarning)
 					.Visibility(this, &SActorDetails::GetUCSComponentWarningVisibility)
@@ -359,7 +359,7 @@ void SActorDetails::Construct(const FArguments& InArgs, UTypedElementSelectionSe
 				]
 				+SVerticalBox::Slot()
 				.AutoHeight()
-				.Padding( FMargin(0,0,0,1) )
+				.Padding( FMargin(0.f,0.f,0.f,1.f) )
 				[
 					SNew(SActorDetailsUneditableComponentWarning)
 					.Visibility(this, &SActorDetails::GetInheritedBlueprintComponentWarningVisibility)
@@ -368,7 +368,7 @@ void SActorDetails::Construct(const FArguments& InArgs, UTypedElementSelectionSe
 				]
 				+SVerticalBox::Slot()
 				.AutoHeight()
-				.Padding( FMargin(0,0,0,1) )
+				.Padding( FMargin(0.f,0.f,0.f,1.f) )
 				[
 					SNew(SActorDetailsUneditableComponentWarning)
 					.Visibility(this, &SActorDetails::GetNativeComponentWarningVisibility)
diff --git a/Engine/Source/Editor/LevelEditor/Private/SInViewportDetails.cpp b/Engine/Source/Editor/LevelEditor/Private/SInViewportDetails.cpp
--- a/Engine/Source/Editor/LevelEditor/Private/SInViewportDetails.cpp
+++ b/Engine/Source/Editor/LevelEditor/Private/SInViewportDetails.cpp
@@ -229,7 +229,7 @@ void SInViewportDetails::GenerateWidget()
 								.AutoWidth()
 								.HAlign(HAlign_Left)
 								.VAlign(VAlign_Center)
-								.Padding(0)
+								.Padding(0.f)
 								[
 									SNew(SImage)
 									.Image(ActorIcon)
diff --git a/Engine/Source/Editor/LevelEditor/Private/SLevelEditor.cpp b/Engine/Source/Editor/LevelEditor/Private/SLevelEditor.cpp
--- a/Engine/Source/Editor/LevelEditor/Private/SLevelEditor.cpp
+++ b/Engine/Source/Editor/LevelEditor/Private/SLevelEditor.cpp
@@ -336,7 +336,7 @@ void SLevelEditor::Initialize( const TSharedRef<SDockTab>& OwnerTab, const TShar
 
 	SecondaryModeToolbarWidget =
 		SNew(SBorder)
-		.Padding(0)
+		.Padding(0.f)
 		.BorderImage(FAppStyle::Get().GetBrush( "NoBorder" ));
 	
 	SecondaryModeToolbarWidget->SetContent(FLevelEditorToolBar::MakeLevelEditorSecondaryModeToolbar(LevelEditorCommands.ToSharedRef(), ModeUILayers));
@@ -1003,7 +1003,7 @@ TSharedRef<SDockTab> SLevelEditor::SpawnLevelEditorTab( const FSpawnTabArgs& Arg
 			.OnExtendContextMenu(this, &SLevelEditor::OnExtendSceneOutlinerTabContextMenu)
 			[
 				SNew(SBorder)
-				.Padding( 0 )
+				.Padding( 0.f )
 				.BorderImage( FAppStyle::GetBrush("ToolPanel.GroupBorder") )
 				.AddMetaData<FTutorialMetaData>(FTutorialMetaData(TEXT("SceneOutliner"), TEXT("LevelEditorSceneOutliner")))
 				[
@@ -1023,7 +1023,7 @@ TSharedRef<SDockTab> SLevelEditor::SpawnLevelEditorTab( const FSpawnTabArgs& Arg
 			.Label( NSLOCTEXT("LevelEditor", "LayersTabTitle", "Layers") )
 			[
 				SNew(SBorder)
-				.Padding( 0 )
+				.Padding( 0.f )
 				.BorderImage( FAppStyle::GetBrush("ToolPanel.GroupBorder") )
 				.AddMetaData<FTutorialMetaData>(FTutorialMetaData(TEXT("LayerBrowser"), TEXT("LevelEditorLayerBrowser")))
 				[
@@ -1079,7 +1079,7 @@ TSharedRef<SDockTab> SLevelEditor::SpawnLevelEditorTab( const FSpawnTabArgs& Arg
 			.Label(NSLOCTEXT("LevelEditor", "DataLayersTabTitle", "Data Layers"))
 			[
 				SNew(SBorder)
-				.Padding(0)
+				.Padding(0.f)
 				.BorderImage(FAppStyle::GetBrush("ToolPanel.GroupBorder"))
 				.AddMetaData<FTutorialMetaData>(FTutorialMetaData(TEXT("DataLayerBrowser"), TEXT("LevelEditorDataLayerBrowser")))
 				[
@@ -1640,7 +1640,7 @@ TSharedRef<SWidget> SLevelEditor::RestoreContentArea( const TSharedRef<SDockTab>
 				->Split
 				(
 					FTabManager::NewStack()
-					->SetSizeCoefficient(.4)
+					->SetSizeCoefficient(.4f)
 					->SetHideTabWell(false)
 					->AddTab("ContentBrowserTab1", ETabState::ClosedTab)
 					->AddTab(LevelEditorTabIds::Sequencer, ETabState::ClosedTab)
@@ -2263,8 +2263,8 @@ FVector2f SLevelEditor::GetActiveViewportWidgetSize() const
 		// SceneView is deleted with the ViewFamily
 		FSceneView* SceneView = ActiveViewport->GetViewportClient()->CalcSceneView(&ViewFamily);
 		const float InvDpiScale = 1.0f / ActiveViewport->GetViewportClient()->GetDPIScale();
-		const float MaxX = SceneView->UnscaledViewRect.Width() * InvDpiScale;
-		const float MaxY = SceneView->UnscaledViewRect.Height() * InvDpiScale;
+		const float MaxX = static_cast<float>(SceneView->UnscaledViewRect.Width()) * InvDpiScale;
+		const float MaxY = static_cast<float>(SceneView->UnscaledViewRect.Height()) * InvDpiScale;
 		return FVector2f(MaxX, MaxY);
 	}
 	return FVector2f::ZeroVector;
diff --git a/Engine/Source/Editor/LevelEditor/Private/SLevelEditorBuildAndSubmit.cpp b/Engine/Source/Editor/LevelEditor/Private/SLevelEditorBuildAndSubmit.cpp
--- a/Engine/Source/Editor/LevelEditor/Private/SLevelEditorBuildAndSubmit.cpp
+++ b/Engine/Source/Editor/LevelEditor/Private/SLevelEditorBuildAndSubmit.cpp
@@ -55,7 +55,7 @@ public:
 				+SHorizontalBox::Slot()
 				.AutoWidth()
 				.VAlign(VAlign_Center)
-				.Padding(4, 1)
+				.Padding(4.f, 1.f)
 				[
 					SNew(SCheckBox)
 					.OnCheckStateChanged(this, &SPackageItem::OnCheckStateChanged)
@@ -94,7 +94,7 @@ public:
 				+SHorizontalBox::Slot()
 				.AutoWidth()
 				.VAlign(VAlign_Center)
-				.Padding(5, 0)
+				.Padding(5.f, 0.f)
 				[
 					SNew(STextBlock)
 					.Font(FAppStyle::GetFontStyle("BuildAndSubmit.SmallFont"))
@@ -174,11 +174,11 @@ void SLevelEditorBuildAndSubmit::Construct( const FArguments& InArgs, const TSha
 		SNew(SVerticalBox)
 		+SVerticalBox::Slot()
 		.AutoHeight()
-		.Padding(5)
+		.Padding(5.f)
 		[
 			SNew(SExpandableArea)
 			.AreaTitle(LOCTEXT("DescriptionSectionTitle", "Submission Description"))
-			.Padding(2)
+			.Padding(2.f)
 			.BodyContent()
 			[
 				SAssignNew(DescriptionBox, SEditableTextBox)
@@ -191,9 +191,9 @@ void SLevelEditorBuildAndSubmit::Construct( const FArguments& InArgs, const TSha
 		[
 			SNew(SExpandableArea )
 			.AreaTitle(LOCTEXT("ExtraPackagesSectionTitle", "Additional Files to Submit"))
-			.Padding(2)
+			.Padding(2.f)
 			.InitiallyCollapsed(true)
-			.MaxHeight(300)
+			.MaxHeight(300.f)
 			.OnAreaExpansionChanged(this, &SLevelEditorBuildAndSubmit::OnShowHideExtraPackagesSection)
 			.BodyContent()
 			[
@@ -206,17 +206,17 @@ void SLevelEditorBuildAndSubmit::Construct( const FArguments& InArgs, const TSha
 		]
 		+SVerticalBox::Slot()
 		.AutoHeight()
-		.Padding(5)
+		.Padding(5.f)
 		[
 			SNew( SExpandableArea )
 			.AreaTitle(LOCTEXT("BuildOptionsSectionTitle", "Build Options"))
-			.Padding(2)
+			.Padding(2.f)
 			.BodyContent()
 			[
 				SNew(SVerticalBox)
 				+SVerticalBox::Slot()
 				.AutoHeight()
-				.Padding(2)
+				.Padding(2.f)
 				[
 					SAssignNew(NoSubmitOnMapErrorBox, SCheckBox)
 					[
@@ -226,7 +226,7 @@ void SLevelEditorBuildAndSubmit::Construct( const FArguments& InArgs, const TSha
 				]
 				+SVerticalBox::Slot()
 				.AutoHeight()
-				.Padding(2)
+				.Padding(2.f)
 				[
 					SAssignNew(NoSubmitOnSaveErrorBox, SCheckBox)
 					[
@@ -236,7 +236,7 @@ void SLevelEditorBuildAndSubmit::Construct( const FArguments& InArgs, const TSha
 				]
 				+SVerticalBox::Slot()
 				.AutoHeight()
-				.Padding(2)
+				.Padding(2.f)
 				[
 					SAssignNew(ShowPackagesNotInSCBox, SCheckBox)
 					.OnCheckStateChanged(this, &SLevelEditorBuildAndSubmit::OnShowPackagesNotInSCBoxChanged)
@@ -247,7 +247,7 @@ void SLevelEditorBuildAndSubmit::Construct( const FArguments& InArgs, const TSha
 				]
 				+SVerticalBox::Slot()
 				.AutoHeight()
-				.Padding(2)
+				.Padding(2.f)
 				[
 					SAssignNew(AddFilesToSCBox, SCheckBox)
 					.IsChecked(ECheckBoxState::Checked)
@@ -260,13 +260,13 @@ void SLevelEditorBuildAndSubmit::Construct( const FArguments& InArgs, const TSha
 		]
 		+SVerticalBox::Slot()
 		.AutoHeight()
-		.Padding(5)
+		.Padding(5.f)
 		.HAlign(HAlign_Right)
 		[
 			SNew(SHorizontalBox)
 			+SHorizontalBox::Slot()
 			.AutoWidth()
-			.Padding(0, 0, 10, 0)
+			.Padding(0.f, 0.f, 10.f, 0.f)
 			[
 				SNew(SButton)
 				.OnClicked(this, &SLevelEditorBuildAndSubmit::OnBuildAndCloseClicked)
diff --git a/Engine/Source/Editor/LevelEditor/Private/SLevelViewport.cpp b/Engine/Source/Editor/LevelEditor/Private/SLevelViewport.cpp
--- a/Engine/Source/Editor/LevelEditor/Private/SLevelViewport.cpp
+++ b/Engine/Source/Editor/LevelEditor/Private/SLevelViewport.cpp
@@ -3793,7 +3793,7 @@ FOptionalSize SActorPreview::OnReadHeight() const
 	{
 		ParentHeight = ParentViewport.Pin()->GetActiveViewport()->GetSizeXY().Y;
 	}
-	return FMath::Clamp( GetDefault<ULevelEditorViewportSettings>()->CameraPreviewSize * ParentHeight * PreviewScalingFactor, MinimumHeight, MaximumHeight );
+	return FMath::Clamp( GetDefault<ULevelEditorViewportSettings>()->CameraPreviewSize * static_cast<float>(ParentHeight) * PreviewScalingFactor, MinimumHeight, MaximumHeight );
 }
 
 float SActorPreview::OnReadTextWidth() const
@@ -4852,15 +4852,15 @@ void SLevelViewport::ToggleInViewportContextMenu()
 					SceneView->WorldToPixel(SelectedActor->GetTransform().GetLocation(), ScreenPos);
 					ScreenPos *= InvDpiScale;
 					const float EdgeFactor = 0.85f;
-					const float MinX = SceneView->UnscaledViewRect.Width() * InvDpiScale * (1 - EdgeFactor);
-					const float MinY = SceneView->UnscaledViewRect.Height() * InvDpiScale * (1 - EdgeFactor);
-					const float MaxX = SceneView->UnscaledViewRect.Width() * InvDpiScale * EdgeFactor;
-					const float MaxY = (SceneView->UnscaledViewRect.Height() * InvDpiScale * EdgeFactor);
+					const float MinX = static_cast<float>(SceneView->UnscaledViewRect.Width()) * InvDpiScale * (1.f - EdgeFactor);
+					const float MinY = static_cast<float>(SceneView->UnscaledViewRect.Height()) * InvDpiScale * (1.f - EdgeFactor);
+					const float MaxX = static_cast<float>(SceneView->UnscaledViewRect.Width()) * InvDpiScale * EdgeFactor;
+					const float MaxY = (static_cast<float>(SceneView->UnscaledViewRect.Height()) * InvDpiScale * EdgeFactor);
 					const bool bOutside = ScreenPos.X < MinX || ScreenPos.X > MaxX || ScreenPos.Y < MinY || ScreenPos.Y > MaxY;
 					if (bOutside)
 					{
-						ScreenPos.X = (SceneView->UnscaledViewRect.Width() * InvDpiScale) / 2.0f;
-						ScreenPos.Y = (SceneView->UnscaledViewRect.Height() * InvDpiScale) / 2.0f;
+						ScreenPos.X = (static_cast<float>(SceneView->UnscaledViewRect.Width()) * InvDpiScale) / 2.0f;
+						ScreenPos.Y = (static_cast<float>(SceneView->UnscaledViewRect.Height()) * InvDpiScale) / 2.0f;
 					}
 					UpdateInViewportMenuLocation(ScreenPos);
 				}
diff --git a/Engine/Source/Editor/LevelEditor/Private/SSurfaceProperties.cpp b/Engine/Source/Editor/LevelEditor/Private/SSurfaceProperties.cpp
--- a/Engine/Source/Editor/LevelEditor/Private/SSurfaceProperties.cpp
+++ b/Engine/Source/Editor/LevelEditor/Private/SSurfaceProperties.cpp
@@ -652,7 +652,7 @@ FReply SSurfaceProperties::OnPanTexture(int32 PanAmount, TextureCoordChannel Cha
 	}
 
 	const float Mod = InvertPanDirection? -1.f : 1.f;
-	GUnrealEd->Exec( GetWorld(), *FString::Printf( TEXT("POLY TEXPAN U=%f V=%f"), PanU * Mod, PanV * Mod ) );
+	GUnrealEd->Exec( GetWorld(), *FString::Printf( TEXT("POLY TEXPAN U=%f V=%f"), static_cast<float>(PanU) * Mod, static_cast<float>(PanV) * Mod ) );
 
 	return FReply::Handled();
 }
@@ -661,7 +661,7 @@ FReply SSurfaceProperties::OnRotateTexture(int32 RotationAmount, RotationAction
 {
 	const float Mod =  bUseNegativeRotation ? -1.f : 1.f;
 
-	const float RotateRadians = RotationAmount / 180.0f * PI;
+	const float RotateRadians = static_cast<float>(RotationAmount) / 180.0f * PI;
 
 	const float UU = cosf(RotateRadians);
 	const float VV = UU;
