[Viewport Toolbar] Update the behavior of the perspective camera so the viewport toolbar will not show both "Perspective" and a camera actor as active radio buttons

Remaps the "Perspective" entry for Level Viewport:
- it will not be checked when piloting
- it will stop piloting if selected

diff --git a/Engine/Source/Editor/LevelEditor/Private/SLevelViewport.cpp b/Engine/Source/Editor/LevelEditor/Private/SLevelViewport.cpp
--- a/Engine/Source/Editor/LevelEditor/Private/SLevelViewport.cpp
+++ b/Engine/Source/Editor/LevelEditor/Private/SLevelViewport.cpp
@@ -1478,6 +1478,55 @@ void SLevelViewport::BindCommands()
 	BindViewCommands( UICommandListRef );
 	BindDropCommands( UICommandListRef );
 
+	// Re-Map Perspective command so that
+	// - Perspective menu entry radial checkbox is not flagged when actors are being piloted (including cameras)
+	// - Clicking on the Perspective entry interrupts piloting
+	{
+		const TWeakPtr<SLevelViewport>& LevelViewportWeak = SharedThis(this).ToWeakPtr();
+
+		UICommandListRef.MapAction(
+			FEditorViewportCommands::Get().Perspective,
+			FExecuteAction::CreateLambda(
+				[LevelViewportWeak]()
+				{
+					if (const TSharedPtr<::SLevelViewport>& LevelViewportPinned =
+							LevelViewportWeak.Pin())
+					{
+						// If piloting, stop
+						LevelViewportPinned->OnActorUnlock();
+
+						if (const TSharedPtr<FEditorViewportClient>& ViewportClient = LevelViewportPinned->GetViewportClient())
+						{
+							// Set to perspective
+							ViewportClient->SetViewportType(LVT_Perspective);
+						}
+					}
+				}
+			),
+			FCanExecuteAction(),
+			FIsActionChecked::CreateLambda(
+				[LevelViewportWeak]()
+				{
+					if (const TSharedPtr<::SLevelViewport>& LevelViewportPinned =
+							LevelViewportWeak.Pin())
+					{
+						// Don't check action when piloting an actor
+						if (LevelViewportPinned->IsAnyActorLocked())
+						{
+							return false;
+						}
+
+						if (const TSharedPtr<FEditorViewportClient>& ViewportClient = LevelViewportPinned->GetViewportClient())
+						{
+							return ViewportClient->IsPerspective();
+						}
+					}
+					return false;
+				}
+			)
+		);
+	}
+
 	if ( ParentLevelEditor.IsValid() )
 	{
 		UICommandListRef.Append(ParentLevelEditor.Pin()->GetLevelEditorActions().ToSharedRef());
@@ -2887,6 +2936,11 @@ void SLevelViewport::OnActorUnlock()
 		// remove roll and pitch from camera when unbinding from actors
 		GEditor->RemovePerspectiveViewRotation(true, true, false);
 
+		// Move perspective camera back to pre-piloting transform
+		LevelViewportClient->SetViewLocation(CachedPerspectiveCameraTransform.GetLocation());
+		LevelViewportClient->SetViewRotation(CachedPerspectiveCameraTransform.GetRotation());
+		Invalidate();
+
 		// If we had a camera actor locked, and it was selected, then we should re-show the inset preview
 		OnPreviewSelectedCamerasChange();
 	}
@@ -4986,6 +5040,8 @@ void SLevelViewport::LockActorInternal(AActor* NewActorToLock)
 		LevelViewportClient->SetActorLock(NewActorToLock);
 		if (LevelViewportClient->IsPerspective() && LevelViewportClient->GetActiveActorLock().IsValid())
 		{
+			// Store perspective camera transform before piloting
+			CachedPerspectiveCameraTransform = LevelViewportClient->GetViewTransform();
 			LevelViewportClient->MoveCameraToLockedActor();
 		}
 	}
diff --git a/Engine/Source/Editor/LevelEditor/Public/SLevelViewport.h b/Engine/Source/Editor/LevelEditor/Public/SLevelViewport.h
--- a/Engine/Source/Editor/LevelEditor/Public/SLevelViewport.h
+++ b/Engine/Source/Editor/LevelEditor/Public/SLevelViewport.h
@@ -979,6 +979,12 @@ private:
 	/** Viewport widget for warning messages */
 	TSharedPtr<SActionableMessageViewportWidget> ActionableMessageViewportWidget;
 
+	/**
+	 * Used to store last perspective camera transform before piloting (Actor Lock)
+	 * Once piloting ends, this transform is re-applied to the camera.
+	 */
+	FViewportCameraTransform CachedPerspectiveCameraTransform;
+
 protected:
 	void LockActorInternal(AActor* NewActorToLock);
 
