Level Editor: Add GetExactCameraView/SetExactCameraView to toggle bLockedCameraView

diff --git a/Engine/Source/Editor/LevelEditor/Private/LevelEditorSubsystem.cpp b/Engine/Source/Editor/LevelEditor/Private/LevelEditorSubsystem.cpp
--- a/Engine/Source/Editor/LevelEditor/Private/LevelEditorSubsystem.cpp
+++ b/Engine/Source/Editor/LevelEditor/Private/LevelEditorSubsystem.cpp
@@ -166,6 +166,30 @@ AActor* ULevelEditorSubsystem::GetPilotLevelActor(FName ViewportConfigKey)
 	return LevelViewportClient.GetActiveActorLock().Get();
 }
 
+void ULevelEditorSubsystem::SetExactCameraView(bool bExactCameraView, FName ViewportConfigKey)
+{
+	TSharedPtr<SLevelViewport> LevelViewport = InternalEditorLevelLibrary::GetLevelViewport(ViewportConfigKey);
+	if (!LevelViewport.IsValid())
+	{
+		return;
+	}
+
+	FLevelEditorViewportClient& LevelViewportClient = LevelViewport->GetLevelViewportClient();
+	LevelViewportClient.bLockedCameraView = bExactCameraView;
+}
+
+bool ULevelEditorSubsystem::GetExactCameraView(FName ViewportConfigKey)
+{
+	TSharedPtr<SLevelViewport> LevelViewport = InternalEditorLevelLibrary::GetLevelViewport(ViewportConfigKey);
+	if (!LevelViewport.IsValid())
+	{
+		return false;
+	}
+
+	FLevelEditorViewportClient& LevelViewportClient = LevelViewport->GetLevelViewportClient();
+	return LevelViewportClient.bLockedCameraView;
+}
+
 void ULevelEditorSubsystem::EditorPlaySimulate()
 {
 	FLevelEditorModule& LevelEditorModule = FModuleManager::GetModuleChecked<FLevelEditorModule>("LevelEditor");
diff --git a/Engine/Source/Editor/LevelEditor/Public/LevelEditorSubsystem.h b/Engine/Source/Editor/LevelEditor/Public/LevelEditorSubsystem.h
--- a/Engine/Source/Editor/LevelEditor/Public/LevelEditorSubsystem.h
+++ b/Engine/Source/Editor/LevelEditor/Public/LevelEditorSubsystem.h
@@ -52,6 +52,11 @@ public:
 	UFUNCTION(BlueprintCallable, Category = "Editor Scripting | Level Utility", meta=(DevelopmentOnly))
 	AActor* GetPilotLevelActor(FName ViewportConfigKey = NAME_None);
 
+	UFUNCTION(BlueprintCallable, Category = "Editor Scripting | Level Utility", meta = (DevelopmentOnly))
+	bool GetExactCameraView(FName ViewportConfigKey = NAME_None);
+	UFUNCTION(BlueprintCallable, Category = "Editor Scripting | Level Utility", meta = (DevelopmentOnly))
+	void SetExactCameraView(bool bExactCameraView, FName ViewportConfigKey = NAME_None);
+
 	UFUNCTION(BlueprintCallable, Category = "Editor Scripting | Level Utility", meta = (DevelopmentOnly))
 	void EditorPlaySimulate();
 
