Fix an issue where the LevelEditor-specific `ToggleHideViewportUI` command leaked into all `MainFrame.MainMenu.Window` menus.

- Remove unecessary comments.

diff --git a/Engine/Source/Editor/LevelEditor/Private/LevelEditorMenu.cpp b/Engine/Source/Editor/LevelEditor/Private/LevelEditorMenu.cpp
--- a/Engine/Source/Editor/LevelEditor/Private/LevelEditorMenu.cpp
+++ b/Engine/Source/Editor/LevelEditor/Private/LevelEditorMenu.cpp
@@ -245,7 +245,7 @@ void FLevelEditorMenu::RegisterLevelEditorMenus()
 
 		static void ExtendWindowMenu()
 		{
-			if (UToolMenu* Menu = UToolMenus::Get()->ExtendMenu("MainFrame.MainMenu.Window"))
+			if (UToolMenu* Menu = UToolMenus::Get()->ExtendMenu("LevelEditor.MainMenu.Window"))
 			{
 				FToolMenuSection& LayoutSection = Menu->FindOrAddSection("WindowLayout");
 
@@ -330,24 +330,13 @@ void FLevelEditorMenu::RegisterLevelEditorMenus()
 	// Add other top level menus
 	Local::ExtendMenuBar();
 
-	// Add level loading and saving menu items
 	Local::RegisterFileLoadAndSaveItems();
-
-	// Add recent / favorites
 	Local::FillFileRecentAndFavoriteFileItems();
-
-	// Extend the Edit menu
 	Local::ExtendEditMenu();
-
-	// Extend the Help menu
 	Local::ExtendHelpMenu();
-
 	Local::ExtendWindowMenu();
 
-	// Extend the Build menu
 	RegisterBuildMenu();
-
-	// Extend the Select menu
 	RegisterSelectMenu();
 }
 
