Fix an issue where DPI scaling was being double applied in transform toolbar width calculations.

Tested at 100% and 150% scale.

diff --git a/Engine/Source/Editor/LevelEditor/Private/SLevelViewportToolBar.cpp b/Engine/Source/Editor/LevelEditor/Private/SLevelViewportToolBar.cpp
--- a/Engine/Source/Editor/LevelEditor/Private/SLevelViewportToolBar.cpp
+++ b/Engine/Source/Editor/LevelEditor/Private/SLevelViewportToolBar.cpp
@@ -1006,23 +1006,13 @@ float SLevelViewportToolBar::GetTransformToolbarWidth() const
 		const float ToolbarWidthMinusPreviousTransformToolbar = GetDesiredSize().X - TransformToolbar_CachedMaxWidth;
 		const float ToolbarWidthEstimate = ToolbarWidthMinusPreviousTransformToolbar + TransformToolbarWidth;
 
-		float DpiScale = 1.0f;
-		{
-			IMainFrameModule& MainFrameModule = FModuleManager::LoadModuleChecked<IMainFrameModule>(TEXT("MainFrame"));
-			const TSharedPtr<SWindow>& MainFrameParentWindow = MainFrameModule.GetParentWindow();
-			if (MainFrameParentWindow.IsValid())
-			{
-				DpiScale = MainFrameParentWindow->GetDPIScaleFactor();
-			}
-		}
-
 		const float ViewportToolBarWidth = static_cast<float>(GetCachedGeometry().GetLocalSize().X);
-		const float OverflowWidth = ToolbarWidthEstimate * DpiScale - ViewportToolBarWidth;
+		const float OverflowWidth = ToolbarWidthEstimate - ViewportToolBarWidth;
 		if (OverflowWidth > 0.0f)
 		{
 			// There isn't enough space in the viewport to show the toolbar!
 			// Try and shrink the transform toolbar (which has an overflow area) to make things fit
-			TransformToolbar_CachedMaxWidth = FMath::Max(FMath::Min(4.0f, TransformToolbarWidth), TransformToolbarWidth - OverflowWidth / DpiScale);
+			TransformToolbar_CachedMaxWidth = FMath::Max(FMath::Min(4.0f, TransformToolbarWidth), TransformToolbarWidth - OverflowWidth);
 		}
 		else
 		{
