Added menu option "Open Actor In Reference Viewer".

diff --git a/Engine/Source/Editor/LevelEditor/Private/LevelEditor.cpp b/Engine/Source/Editor/LevelEditor/Private/LevelEditor.cpp
--- a/Engine/Source/Editor/LevelEditor/Private/LevelEditor.cpp
+++ b/Engine/Source/Editor/LevelEditor/Private/LevelEditor.cpp
@@ -768,6 +768,12 @@ void FLevelEditorModule::BindGlobalLevelEditorCommands()
 		FExecuteAction::CreateStatic(&FLevelEditorActionCallbacks::CopyActorFilePathtoClipboard_Clicked)
 		);
 
+	ActionList.MapAction(
+		Commands.OpenActorInReferenceViewer,
+		FExecuteAction::CreateStatic(&FLevelEditorActionCallbacks::ViewActorReferences_Clicked),
+		FCanExecuteAction::CreateStatic(&FLevelEditorActionCallbacks::ViewActorReferences_CanExecute)
+	);
+
 	ActionList.MapAction(
 		Commands.SaveActor,
 		FExecuteAction::CreateStatic(&FLevelEditorActionCallbacks::SaveActor_Clicked),
diff --git a/Engine/Source/Editor/LevelEditor/Private/LevelEditorActions.cpp b/Engine/Source/Editor/LevelEditor/Private/LevelEditorActions.cpp
--- a/Engine/Source/Editor/LevelEditor/Private/LevelEditorActions.cpp
+++ b/Engine/Source/Editor/LevelEditor/Private/LevelEditorActions.cpp
@@ -3087,6 +3087,36 @@ void FLevelEditorActionCallbacks::SnapObjectToView_Clicked()
 	GEditor->RedrawLevelEditingViewports();
 }
 
+void FLevelEditorActionCallbacks::ViewActorReferences_Clicked()
+{
+	for (FSelectionIterator It(GEditor->GetSelectedActorIterator()); It; ++It)
+	{
+		if (AActor* Actor = Cast<AActor>(*It))
+		{
+			FAssetIdentifier AssetIdentifier(Actor->GetPackage()->GetFName());
+			FEditorDelegates::OnOpenReferenceViewer.Broadcast({ AssetIdentifier }, FReferenceViewerParams());
+			break;
+		}
+	}
+}
+
+bool FLevelEditorActionCallbacks::ViewActorReferences_CanExecute()
+{
+	TSet<FName> PackageNames;
+	for (FSelectionIterator It(GEditor->GetSelectedActorIterator()); It; ++It)
+	{
+		if (AActor* Actor = Cast<AActor>(*It))
+		{
+			PackageNames.Add(Actor->GetPackage()->GetFName());
+			if (PackageNames.Num() > 1)
+			{
+				return false;
+			}
+		}
+	}
+	return PackageNames.Num() == 1;
+}
+
 void FLevelEditorActionCallbacks::CopyActorFilePathtoClipboard_Clicked()
 {
 	TStringBuilder<1024> Result;
@@ -3992,6 +4022,7 @@ void FLevelEditorCommands::RegisterCommands()
 	UI_COMMAND( SnapObjectToCamera, "Move Object to Camera", "Move the selected object to match the location and rotation of the current camera.", EUserInterfaceActionType::Button, FInputChord() );
 
 	UI_COMMAND(CopyActorFilePathtoClipboard, "Copy Selected Actor(s) File Path", "Copy the file path of the selected actors", EUserInterfaceActionType::Button, FInputChord());
+	UI_COMMAND(OpenActorInReferenceViewer, "Open Actor in Reference Viewer...", "Launches the reference viewer showing the selected actor references", EUserInterfaceActionType::Button, FInputChord(EModifierKey::Shift | EModifierKey::Alt, EKeys::R));
 	UI_COMMAND(SaveActor, "Save Selected Actor(s)", "Save the selected actors", EUserInterfaceActionType::Button, FInputChord());
 	UI_COMMAND(ShowActorHistory, "Show Actor History", "Shows the history of the file containing the actor.", EUserInterfaceActionType::Button, FInputChord());	
 
diff --git a/Engine/Source/Editor/LevelEditor/Private/LevelEditorContextMenu.cpp b/Engine/Source/Editor/LevelEditor/Private/LevelEditorContextMenu.cpp
--- a/Engine/Source/Editor/LevelEditor/Private/LevelEditorContextMenu.cpp
+++ b/Engine/Source/Editor/LevelEditor/Private/LevelEditorContextMenu.cpp
@@ -351,6 +351,13 @@ void FLevelEditorContextMenu::RegisterActorContextMenu()
 				FSlateIcon(FAppStyle::GetAppStyleSetName(), "GenericCommands.Copy")
 			);
 
+			Section.AddMenuEntry(
+				FLevelEditorCommands::Get().OpenActorInReferenceViewer,
+				TAttribute<FText>(), // use command's label
+				TAttribute<FText>(), // use command's tooltip
+				FSlateIcon(FAppStyle::GetAppStyleSetName(), "ContentBrowser.ReferenceViewer")
+			);
+
 			Section.AddMenuEntry(
 				FLevelEditorCommands::Get().SaveActor,
 				TAttribute<FText>(), // use command's label
diff --git a/Engine/Source/Editor/LevelEditor/Public/LevelEditorActions.h b/Engine/Source/Editor/LevelEditor/Public/LevelEditorActions.h
--- a/Engine/Source/Editor/LevelEditor/Public/LevelEditorActions.h
+++ b/Engine/Source/Editor/LevelEditor/Public/LevelEditorActions.h
@@ -145,6 +145,9 @@ public:
 	/** Copy the file path where the actor is saved. */
 	TSharedPtr< FUICommandInfo > CopyActorFilePathtoClipboard;
 
+	/** Opens the reference viewer for the selected actor. */
+	TSharedPtr< FUICommandInfo > OpenActorInReferenceViewer;
+
 	/** Save the selected actor. */
 	TSharedPtr< FUICommandInfo > SaveActor;
 
@@ -1358,6 +1361,17 @@ public:
 	*/
 	static void CopyActorFilePathtoClipboard_Clicked();
 
+	/**
+	* Opens the reference viewer for the selected actor.
+	*/
+	static void ViewActorReferences_Clicked();
+
+	/**
+	 * Checks whether OpenActorInReferenceViewer can be executed on the selected actor.
+	 * @return 	True if OpenActorInReferenceViewer can be executed on the selected actor.
+	 */
+	static bool ViewActorReferences_CanExecute();
+
 	/**
 	* Save the actor.
 	*/
