[Viewport Toolbar] Expose Editor Preferences in the "Settings" menu of the new level editor viewport toolbar

- Also adds a CreateNumericEntry function for int values

diff --git a/Engine/Source/Editor/LevelEditor/Private/LevelEditor.cpp b/Engine/Source/Editor/LevelEditor/Private/LevelEditor.cpp
--- a/Engine/Source/Editor/LevelEditor/Private/LevelEditor.cpp
+++ b/Engine/Source/Editor/LevelEditor/Private/LevelEditor.cpp
@@ -1843,7 +1843,21 @@ void FLevelEditorModule::BindGlobalLevelEditorCommands()
 		FCanExecuteAction(),
 		FIsActionChecked::CreateStatic(&FLevelEditorActionCallbacks::IsOrthoZoomToCursorChecked));
 
-	
+	ActionList.MapAction(
+		Commands.InvertMiddleMousePan,
+		FExecuteAction::CreateStatic(&FLevelEditorActionCallbacks::ToggleInvertMiddleMousePan),
+		FCanExecuteAction(),
+		FIsActionChecked::CreateStatic(&FLevelEditorActionCallbacks::IsInvertMiddleMousePanChecked));
+	ActionList.MapAction(
+		Commands.InvertOrbitYAxis,
+		FExecuteAction::CreateStatic(&FLevelEditorActionCallbacks::ToggleInvertOrbitYAxis),
+		FCanExecuteAction(),
+		FIsActionChecked::CreateStatic(&FLevelEditorActionCallbacks::IsInvertOrbitYAxisChecked));
+	ActionList.MapAction(
+		Commands.InvertRightMouseDollyYAxis,
+		FExecuteAction::CreateStatic(&FLevelEditorActionCallbacks::ToggleInvertRightMouseDollyYAxis),
+		FCanExecuteAction(),
+		FIsActionChecked::CreateStatic(&FLevelEditorActionCallbacks::IsInvertRightMouseDollyYAxisChecked));
 
 	const TArray<FPreviewPlatformMenuItem>& MenuItems = FDataDrivenPlatformInfoRegistry::GetAllPreviewPlatformMenuItems();
 	// We need one extra slot for the Disable Preview option
diff --git a/Engine/Source/Editor/LevelEditor/Private/LevelEditorActions.cpp b/Engine/Source/Editor/LevelEditor/Private/LevelEditorActions.cpp
--- a/Engine/Source/Editor/LevelEditor/Private/LevelEditorActions.cpp
+++ b/Engine/Source/Editor/LevelEditor/Private/LevelEditorActions.cpp
@@ -3881,6 +3881,57 @@ bool FLevelEditorActionCallbacks::IsOrthoZoomToCursorChecked()
 	return false;
 }
 
+void FLevelEditorActionCallbacks::ToggleInvertMiddleMousePan()
+{
+	if (ULevelEditorViewportSettings* ViewportSettings = GetMutableDefault<ULevelEditorViewportSettings>())
+	{
+		ViewportSettings->bInvertMiddleMousePan = !ViewportSettings->bInvertMiddleMousePan;
+	}
+}
+
+bool FLevelEditorActionCallbacks::IsInvertMiddleMousePanChecked()
+{
+	if (const ULevelEditorViewportSettings* const ViewportSettings = GetDefault<ULevelEditorViewportSettings>())
+	{
+		return ViewportSettings->bInvertMiddleMousePan;
+	}
+	return false;
+}
+
+void FLevelEditorActionCallbacks::ToggleInvertOrbitYAxis()
+{
+	if (ULevelEditorViewportSettings* ViewportSettings = GetMutableDefault<ULevelEditorViewportSettings>())
+	{
+		ViewportSettings->bInvertOrbitYAxis = !ViewportSettings->bInvertOrbitYAxis;
+	}
+}
+
+bool FLevelEditorActionCallbacks::IsInvertOrbitYAxisChecked()
+{
+	if (const ULevelEditorViewportSettings* const ViewportSettings = GetDefault<ULevelEditorViewportSettings>())
+	{
+		return ViewportSettings->bInvertOrbitYAxis;
+	}
+	return false;
+}
+
+void FLevelEditorActionCallbacks::ToggleInvertRightMouseDollyYAxis()
+{
+	if (ULevelEditorViewportSettings* ViewportSettings = GetMutableDefault<ULevelEditorViewportSettings>())
+	{
+		ViewportSettings->bInvertRightMouseDollyYAxis = !ViewportSettings->bInvertRightMouseDollyYAxis;
+	}
+}
+
+bool FLevelEditorActionCallbacks::IsInvertRightMouseDollyYAxisChecked()
+{
+	if (const ULevelEditorViewportSettings* const ViewportSettings = GetDefault<ULevelEditorViewportSettings>())
+	{
+		return ViewportSettings->bInvertRightMouseDollyYAxis;
+	}
+	return false;
+}
+
 class UWorld* FLevelEditorActionCallbacks::GetWorld()
 {
 	return GEditor->GetEditorWorldContext().World();
@@ -4242,11 +4293,12 @@ void FLevelEditorCommands::RegisterCommands()
 	UI_COMMAND(OrbitCameraAroundSelection, "Orbit Around Selection", "If enabled, the camera will orbit around the current selection in the viewport", EUserInterfaceActionType::ToggleButton, FInputChord());
 	UI_COMMAND(LinkOrthographicViewports, "Link Ortho Camera Movement", "If checked all orthographic view ports are linked to the same position and move together.", EUserInterfaceActionType::ToggleButton, FInputChord());
 	UI_COMMAND(OrthoZoomToCursor, "Ortho Zoom to Cursor", "If checked, in orthographic viewports zooming will center on the mouse position.  If unchecked, the zoom is around the center of the viewport.", EUserInterfaceActionType::ToggleButton, FInputChord());
-
-	/*TSharedPtr<FUICommandInfo> LinkOrthographicViewports;
-	TSharedPtr<FUICommandInfo> OrthoZoomToCursor;
-	TSharedPtr<FUICommandInfo> DistanceBasedCameraSpeed;*/
 	
+	// Mouse Controls
+	UI_COMMAND(InvertMiddleMousePan, "Invert Middle Mouse Pan", "Whether or not to invert the direction of middle mouse panning in viewports", EUserInterfaceActionType::ToggleButton, FInputChord());
+	UI_COMMAND(InvertOrbitYAxis, "Invert Orbit Axis", "Whether or not to invert mouse on y axis in orbit mode", EUserInterfaceActionType::ToggleButton, FInputChord());
+	UI_COMMAND(InvertRightMouseDollyYAxis, "Invert Right Mouse Dolly", "Whether or not to invert the direction of right mouse dolly on the Y axis in orbit mode", EUserInterfaceActionType::ToggleButton, FInputChord());
+
 	// Add preview platforms
 	TSet<FName> PreviewShaderPlatformNames;
 	for (const FPreviewPlatformMenuItem& Item : FDataDrivenPlatformInfoRegistry::GetAllPreviewPlatformMenuItems())
diff --git a/Engine/Source/Editor/LevelEditor/Private/ViewportToolbar/LevelEditorSubmenus.cpp b/Engine/Source/Editor/LevelEditor/Private/ViewportToolbar/LevelEditorSubmenus.cpp
--- a/Engine/Source/Editor/LevelEditor/Private/ViewportToolbar/LevelEditorSubmenus.cpp
+++ b/Engine/Source/Editor/LevelEditor/Private/ViewportToolbar/LevelEditorSubmenus.cpp
@@ -747,6 +747,144 @@ FToolMenuEntry CreatePreviewSelectedCamerasCheckBoxSubmenu()
 	return Entry;
 }
 
+FToolMenuEntry CreateMouseScrollCameraSpeedEntry()
+{
+	constexpr int32 SpeedMin = 1;
+	constexpr int32 SpeedMax = 8;
+
+	return UE::UnrealEd::CreateNumericEntry(
+		"MouseScrollCameraSpeed",
+		LOCTEXT("MouseScrollCameraSpeedLabel", "Mouse Scroll Zoom Speed"),
+		LOCTEXT("MouseScrollCameraSpeedTooltip", "How fast the perspective camera moves through the world when using mouse scroll"),
+		FCanExecuteAction(),
+		UnrealEd::FNumericEntryExecuteActionDelegateInt32::CreateLambda(
+			[](int32 InValue)
+			{
+				if (ULevelEditorViewportSettings* const ViewportSettings = GetMutableDefault<ULevelEditorViewportSettings>())
+				{
+					ViewportSettings->MouseScrollCameraSpeed = InValue;
+				}
+			}
+		),
+		TAttribute<int32>::CreateLambda(
+			[]()
+			{
+				if (const ULevelEditorViewportSettings* const ViewportSettings = GetDefault<ULevelEditorViewportSettings>())
+				{
+					return ViewportSettings->MouseScrollCameraSpeed;
+				}
+				return SpeedMin;
+			}
+		),
+		SpeedMin,
+		SpeedMax
+	);
+}
+
+FToolMenuEntry CreateMouseSensitivityEntry()
+{
+	constexpr float SpeedMin = 0.01f;
+	constexpr float SpeedMax = 1.0f;
+	return UnrealEd::CreateNumericEntry(
+		"MouseSensitivity",
+		LOCTEXT("MouseSensitivityLabel", "Mouse Sensitivity"),
+		LOCTEXT("MouseSensitivityTooltip", "How fast the perspective camera moves through the world when using mouse scroll"),
+		FCanExecuteAction(),
+		UnrealEd::FNumericEntryExecuteActionDelegate::CreateLambda(
+			[](float InValue)
+			{
+				if (ULevelEditorViewportSettings* const ViewportSettings = GetMutableDefault<ULevelEditorViewportSettings>())
+				{
+					ViewportSettings->MouseSensitivty = InValue;
+				}
+			}
+		),
+		TAttribute<float>::CreateLambda(
+			[]()
+			{
+				if (const ULevelEditorViewportSettings* const ViewportSettings = GetDefault<ULevelEditorViewportSettings>())
+				{
+					return ViewportSettings->MouseSensitivty;
+				}
+				return SpeedMin;
+			}
+		),
+		SpeedMin,
+		SpeedMax,
+		2
+	);
+}
+
+TSharedRef<SWidget> CreateGestureDirectionWidget(EScrollGestureDirection& InOutScrollGestureProperty, const FName InMenuName)
+{
+	if (!UToolMenus::Get()->IsMenuRegistered(InMenuName))
+	{
+		if (UToolMenu* Menu = UToolMenus::Get()->RegisterMenu(InMenuName, NAME_None))
+		{
+			FToolMenuSection& Section = Menu->AddSection(NAME_None, FText());
+			auto AddGestureRadioButton =
+				[&Section, &InOutScrollGestureProperty](const FText& InLabel, EScrollGestureDirection InDirection)
+			{
+				Section.AddMenuEntry(
+					NAME_None,
+					InLabel,
+					FText(),
+					FSlateIcon(),
+					FUIAction(
+						FExecuteAction::CreateLambda(
+							[InDirection, &InOutScrollGestureProperty]()
+							{
+								InOutScrollGestureProperty = InDirection;
+							}
+						),
+						FCanExecuteAction(),
+						FIsActionChecked::CreateLambda(
+							[InDirection, &InOutScrollGestureProperty]()
+							{
+								if (InOutScrollGestureProperty == InDirection)
+								{
+									return true;
+								}
+								return false;
+							}
+						)
+					),
+					EUserInterfaceActionType::RadioButton
+				);
+			};
+			AddGestureRadioButton(
+				LOCTEXT("ScrollGestureDirectionSystemSettingsLabel", "System Setting"), EScrollGestureDirection::UseSystemSetting
+			);
+			AddGestureRadioButton(
+				LOCTEXT("ScrollGestureDirectionStandardLabel", "Standard"), EScrollGestureDirection::Standard
+			);
+			AddGestureRadioButton(LOCTEXT("ScrollGestureDirectionNaturalLabel", "Natural"), EScrollGestureDirection::Natural);
+		}
+	}
+	FToolMenuContext MenuContext;
+	return UToolMenus::Get()->GenerateWidget(InMenuName, MenuContext);
+}
+
+TSharedRef<SWidget> CreatePerspectiveViewportGestureDirectionWidget()
+{
+	if (ULevelEditorViewportSettings* ViewportSettings = GetMutableDefault<ULevelEditorViewportSettings>())
+	{
+		EScrollGestureDirection& Property = ViewportSettings->ScrollGestureDirectionFor3DViewports;
+		return CreateGestureDirectionWidget(Property, "GestureDirectionSubmenu_Perspective");
+	}
+	return SNullWidget::NullWidget;
+}
+
+TSharedRef<SWidget> CreateOrthographicViewportGestureDirectionWidget()
+{
+	if (ULevelEditorViewportSettings* ViewportSettings = GetMutableDefault<ULevelEditorViewportSettings>())
+	{
+		EScrollGestureDirection& Property = ViewportSettings->ScrollGestureDirectionForOrthoViewports;
+		return CreateGestureDirectionWidget(Property, "GestureDirectionSubmenu_Orthographic");
+	}
+	return SNullWidget::NullWidget;
+}
+
 } // namespace UE::LevelEditor::Private
 
 namespace UE::LevelEditor
@@ -2203,46 +2341,103 @@ FToolMenuEntry CreateSettingsSubmenu()
 			[](UToolMenu* Submenu) -> void
 			{
 				{
-					FToolMenuSection& SettingsSection =
-						Submenu->FindOrAddSection("Settings", LOCTEXT("SettingsSectionLabel", "Settings"));
-
-					SettingsSection.AddEntry(FToolMenuEntry::InitWidget(
-						"Volume", BuildVolumeControlCustomWidget(), LOCTEXT("VolumeControlLabel", "Volume")
-					));
-
-					SettingsSection.AddSeparator("AdvancedSeparator");
+					{
+						FToolMenuSection& SettingsSection =
+							Submenu->FindOrAddSection("Settings", LOCTEXT("SettingsSectionLabel", "Settings"));
+						SettingsSection.AddEntry(
+							FToolMenuEntry::InitWidget(
+								"Volume", BuildVolumeControlCustomWidget(), LOCTEXT("VolumeControlLabel", "Volume")
+							)
+						);
+					}
+					// Mouse Control
+					{
+						FToolMenuSection& ControlsSection =
+							Submenu->FindOrAddSection("Controls", LOCTEXT("ControlsSectionLabel", "Controls"));
+						ControlsSection.AddEntry(UE::LevelEditor::Private::CreateMouseSensitivityEntry());
+						ControlsSection.AddEntry(UE::LevelEditor::Private::CreateMouseScrollCameraSpeedEntry());
+						ControlsSection.AddMenuEntry(FLevelEditorCommands::Get().InvertMiddleMousePan);
+						ControlsSection.AddMenuEntry(FLevelEditorCommands::Get().InvertOrbitYAxis);
+						ControlsSection.AddMenuEntry(FLevelEditorCommands::Get().InvertRightMouseDollyYAxis);
+						ControlsSection.AddSubMenu(
+							"ScrollGestures",
+							LOCTEXT("ScrollGesturesLabel", "Scroll Gestures"),
+							LOCTEXT("ScrollGesturesTooltip", "Scroll Gestures Options"),
+							FNewToolMenuDelegate::CreateLambda(
+								[](UToolMenu* InMenu)
+								{
+									// Perspective Scroll Gesture
+									{
+										FToolMenuSection& PerspectiveSection = InMenu->AddSection(
+											"PerspectiveScrollGestureDirection",
+											LOCTEXT("PerspectiveScrollGestureDirectionLabel", "Perspective Scroll Gesture Direction")
+										);
+										PerspectiveSection.AddEntry(
+											FToolMenuEntry::InitWidget(
+												"PerspectiveScrollGestureDirectionWidget",
+												UE::LevelEditor::Private::CreatePerspectiveViewportGestureDirectionWidget(),
+												FText()
+											)
+										);
+									}
 
+									// Orthographic Scroll Gesture
+									{
+										FToolMenuSection& OrthographicSection = InMenu->AddSection(
+											"OrthographicScrollGestureDirection",
+											LOCTEXT("OrthographicScrollGestureDirectionLabel", "Ortho Scroll Gesture Direction")
+										);
+										OrthographicSection.AddEntry(
+											FToolMenuEntry::InitWidget(
+												"OrthographicScrollGestureDirectionWidget",
+												UE::LevelEditor::Private::CreateOrthographicViewportGestureDirectionWidget(),
+												FText()
+											)
+										);
+									}
+								}
+							)
+						);
+					}
+					// Viewport advanced settings
 					{
+						{
+							FToolMenuSection& UnnamedSection = Submenu->FindOrAddSection(NAME_None, FText());
+							UnnamedSection.AddSeparator(NAME_None);
 						const FLevelViewportCommands& LevelViewportActions = FLevelViewportCommands::Get();
-						SettingsSection.AddMenuEntry(LevelViewportActions.AdvancedSettings);
+							UnnamedSection.AddMenuEntry(LevelViewportActions.AdvancedSettings);
+						}
 					}
 
-					FToolMenuSection& CascadeSection =
-						Submenu->FindOrAddSection("Cascade", LOCTEXT("CascadeSectionLabel", "Cascade"));
-
-					constexpr bool bInOpenSubMenuOnClick = false;
-
-					CascadeSection.AddSubMenu(
-						"CascadeSubmenu",
-						LOCTEXT("CascadeLabel", "Cascade"),
-						LOCTEXT("CascadeTooltip", "Cascade Options"),
-						FNewToolMenuDelegate::CreateLambda(
-							[](UToolMenu* InMenu)
-							{
-								FToolMenuSection& Section =
-									InMenu->AddSection("Cascade", LOCTEXT("CascadeLabel", "Cascade"));
-								Section.AddMenuEntry(FLevelEditorCommands::Get().ToggleParticleSystemLOD);
-								Section.AddMenuEntry(FLevelEditorCommands::Get().ToggleParticleSystemHelpers);
-								Section.AddMenuEntry(FLevelEditorCommands::Get().ToggleFreezeParticleSimulation);
-							}
-						),
-						bInOpenSubMenuOnClick,
-						FSlateIcon(FSlateIcon(FAppStyle::GetAppStyleSetName(), "ClassIcon.ParticleSystemComponent"))
-					);
+					// Cascade
+					{
+						FToolMenuSection& CascadeSection =
+							Submenu->FindOrAddSection("Cascade", LOCTEXT("CascadeSectionLabel", "Cascade"));
+
+						constexpr bool bInOpenSubMenuOnClick = false;
+						CascadeSection.AddSubMenu(
+							"CascadeSubmenu",
+							LOCTEXT("CascadeLabel", "Cascade"),
+							LOCTEXT("CascadeTooltip", "Cascade Options"),
+							FNewToolMenuDelegate::CreateLambda(
+								[](UToolMenu* InMenu)
+								{
+									FToolMenuSection& Section =
+										InMenu->AddSection("Cascade", LOCTEXT("CascadeLabel", "Cascade"));
+									Section.AddMenuEntry(FLevelEditorCommands::Get().ToggleParticleSystemLOD);
+									Section.AddMenuEntry(FLevelEditorCommands::Get().ToggleParticleSystemHelpers);
+									Section.AddMenuEntry(FLevelEditorCommands::Get().ToggleFreezeParticleSimulation);
+								}
+							),
+							bInOpenSubMenuOnClick,
+							FSlateIcon(FSlateIcon(FAppStyle::GetAppStyleSetName(), "ClassIcon.ParticleSystemComponent"))
+						);
+					}
 				}
 			}
 		)
 	);
+
 	Entry.Icon = FSlateIcon(FAppStyle::GetAppStyleSetName(), "LevelEditor.GameSettings");
 	Entry.InsertPosition.Position = EToolMenuInsertType::Last;
 	Entry.ToolBarData.LabelOverride = FText();
diff --git a/Engine/Source/Editor/LevelEditor/Public/LevelEditorActions.h b/Engine/Source/Editor/LevelEditor/Public/LevelEditorActions.h
--- a/Engine/Source/Editor/LevelEditor/Public/LevelEditorActions.h
+++ b/Engine/Source/Editor/LevelEditor/Public/LevelEditorActions.h
@@ -665,6 +665,14 @@ public:
 
 	TSharedPtr< FUICommandInfo > FixupGroupActor;
 
+	/**
+	 * Controls
+	 */
+
+	// Mouse Controls
+	TSharedPtr<FUICommandInfo> InvertMiddleMousePan;
+	TSharedPtr<FUICommandInfo> InvertOrbitYAxis;
+	TSharedPtr<FUICommandInfo> InvertRightMouseDollyYAxis;
 };
 
 /**
@@ -890,6 +898,15 @@ public:
 	static void ToggleOrthoZoomToCursor();
 	static bool IsOrthoZoomToCursorChecked();
 
+	static void ToggleInvertMiddleMousePan();
+	static bool IsInvertMiddleMousePanChecked();
+	
+	static void ToggleInvertOrbitYAxis();
+	static bool IsInvertOrbitYAxisChecked();
+	
+	static void ToggleInvertRightMouseDollyYAxis();
+	static bool IsInvertRightMouseDollyYAxisChecked();
+
 	/**
 	 * Called when the Scene Stats button is clicked.  Invokes the Primitive Stats dialog.
 	 */
diff --git a/Engine/Source/Editor/UnrealEd/Private/ViewportToolbar/UnrealEdViewportToolbar.cpp b/Engine/Source/Editor/UnrealEd/Private/ViewportToolbar/UnrealEdViewportToolbar.cpp
--- a/Engine/Source/Editor/UnrealEd/Private/ViewportToolbar/UnrealEdViewportToolbar.cpp
+++ b/Engine/Source/Editor/UnrealEd/Private/ViewportToolbar/UnrealEdViewportToolbar.cpp
@@ -2142,14 +2142,8 @@ FToolMenuEntry CreateCheckboxSubmenu(
 }
 
 TSharedRef<SWidget> CreateNumericEntryWidget(
-	const FText& InLabel,
-	const FText& InTooltip,
-	const FNumericEntryExecuteActionDelegate& InOnValueChanged,
-	const TAttribute<float>& InGetValue,
-	float InMinValue,
-	float InMaxValue,
-	int32 InMaxFractionalDigits
-)
+	const TSharedRef<SWidget>& InNumericBoxWidget,
+	const FText& InLabel)
 {
 	const FMargin WidgetsMargin(2.0f, 0.0f, 3.0f, 0.0f);
 
@@ -2174,27 +2168,13 @@ TSharedRef<SWidget> CreateNumericEntryWidget(
 			.Padding(WidgetsMargin)
 			.MinDesiredWidth(80.0f)
 			[
-				SNew(SNumericEntryBox<float>)
-				.ToolTipText(InTooltip)
-				.MinValue(InMinValue)
-				.MaxValue(InMaxValue)
-				.MaxSliderValue(InMaxValue)
-				.AllowSpin(true)
-				.MaxFractionalDigits(InMaxFractionalDigits)
-				.Font(FAppStyle::GetFontStyle(TEXT("MenuItem.Font")))
-				.OnValueChanged_Lambda([InOnValueChanged](float InValue)
-				{
-					InOnValueChanged.Execute(InValue);
-				})
-				.Value_Lambda([InGetValue]()
-				{
-					return InGetValue.Get();
-				})
+				InNumericBoxWidget
 			]
 		];
 	// clang-format on
 }
 
+// float version of numeric entry widget
 FToolMenuEntry CreateNumericEntry(
 	const FName InName,
 	const FText& InLabel,
@@ -2207,10 +2187,70 @@ FToolMenuEntry CreateNumericEntry(
 	int32 InMaxFractionalDigits
 )
 {
+	// clang-format off
+	TSharedRef<SWidget> NumericEntryWidget =
+		SNew(SNumericEntryBox<float>)
+		.ToolTipText(InTooltip)
+		.MinValue(InMinValue)
+		.MaxValue(InMaxValue)
+		.MaxSliderValue(InMaxValue)
+		.AllowSpin(true)
+		.MaxFractionalDigits(InMaxFractionalDigits)
+		.Font(FAppStyle::GetFontStyle(TEXT("MenuItem.Font")))
+		.OnValueChanged_Lambda([InOnValueChanged](float InValue)
+		{
+			InOnValueChanged.Execute(InValue);
+		})
+		.Value_Lambda([InGetValue]()
+		{
+			return InGetValue.Get();
+		});
+	// clang-format on
+
+	FToolMenuEntry NumericEntry = FToolMenuEntry::InitMenuEntry(
+		InName,
+		FUIAction(FExecuteAction(), InCanExecuteAction),
+		CreateNumericEntryWidget(NumericEntryWidget, InLabel)
+	);
+
+	return NumericEntry;
+}
+
+// int32 version of numeric entry widget
+FToolMenuEntry CreateNumericEntry(
+	const FName InName,
+	const FText& InLabel,
+	const FText& InTooltip,
+	const FCanExecuteAction& InCanExecuteAction,
+	const FNumericEntryExecuteActionDelegateInt32& InOnValueChanged,
+	const TAttribute<int32>& InGetValue,
+	int32 InMinValue,
+	int32 InMaxValue
+)
+{
+	// clang-format off
+	TSharedRef<SWidget> NumericEntryWidget =
+		SNew(SNumericEntryBox<int32>)
+		.ToolTipText(InTooltip)
+		.MinValue(InMinValue)
+		.MaxValue(InMaxValue)
+		.MaxSliderValue(InMaxValue)
+		.AllowSpin(true)
+		.Font(FAppStyle::GetFontStyle(TEXT("MenuItem.Font")))
+		.OnValueChanged_Lambda([InOnValueChanged](int32 InValue)
+		{
+			InOnValueChanged.Execute(InValue);
+		})
+		.Value_Lambda([InGetValue]()
+		{
+			return InGetValue.Get();
+		});
+	// clang-format on
+
 	FToolMenuEntry NumericEntry = FToolMenuEntry::InitMenuEntry(
 		InName,
 		FUIAction(FExecuteAction(), InCanExecuteAction),
-		CreateNumericEntryWidget(InLabel, InTooltip, InOnValueChanged, InGetValue, InMinValue, InMaxValue, InMaxFractionalDigits)
+		CreateNumericEntryWidget(NumericEntryWidget, InLabel)
 	);
 
 	return NumericEntry;
diff --git a/Engine/Source/Editor/UnrealEd/Public/ViewportToolbar/UnrealEdViewportToolbar.h b/Engine/Source/Editor/UnrealEd/Public/ViewportToolbar/UnrealEdViewportToolbar.h
--- a/Engine/Source/Editor/UnrealEd/Public/ViewportToolbar/UnrealEdViewportToolbar.h
+++ b/Engine/Source/Editor/UnrealEd/Public/ViewportToolbar/UnrealEdViewportToolbar.h
@@ -6,6 +6,7 @@
 #include "Engine/EngineBaseTypes.h"
 #include "Templates/SharedPointerFwd.h"
 #include "ToolMenuDelegates.h"
+#include "ToolMenuEntry.h"
 #include "UnrealEdViewportToolbarContext.h"
 #include "UnrealWidgetFwd.h"
 
@@ -81,6 +82,7 @@ DECLARE_DELEGATE_OneParam(FScaleGridCheckboxListExecuteActionDelegate, int);
 DECLARE_DELEGATE_RetVal_OneParam(bool, FScaleGridCheckboxListIsCheckedDelegate, int);
 
 DECLARE_DELEGATE_OneParam(FNumericEntryExecuteActionDelegate, float);
+DECLARE_DELEGATE_OneParam(FNumericEntryExecuteActionDelegateInt32, int32);
 
 DECLARE_DELEGATE_TwoParams(FOnViewportClientCamSpeedChanged, const TSharedRef<SEditorViewport>&, int32);
 DECLARE_DELEGATE_TwoParams(FOnViewportClientCamSpeedScalarChanged, const TSharedRef<SEditorViewport>&, float);
@@ -145,6 +147,17 @@ UNREALED_API FToolMenuEntry CreateNumericEntry(
 	int32 InMaxFractionalDigits = 2
 );
 
+UNREALED_API FToolMenuEntry CreateNumericEntry(
+	const FName InName,
+	const FText& InLabel,
+	const FText& InTooltip,
+	const FCanExecuteAction& InCanExecuteAction,
+	const FNumericEntryExecuteActionDelegateInt32& InOnValueChanged,
+	const TAttribute<int32>& InGetValue,
+	int32 InMinValue = 0,
+	int32 InMaxValue = 1
+);
+
 UNREALED_API FText GetCameraSpeedLabel(const TWeakPtr<SEditorViewport>& WeakViewport);
 UNREALED_API FText GetCameraSubmenuLabelFromViewportType(const ELevelViewportType ViewportType);
 UNREALED_API FName GetCameraSubmenuIconFNameFromViewportType(const ELevelViewportType ViewportType);
