[Viewport Toolbar] Expose Editor Preferences in the "Transform" menu of the new viewport toolbar

diff --git a/Engine/Source/Editor/LevelEditor/Private/LevelEditor.cpp b/Engine/Source/Editor/LevelEditor/Private/LevelEditor.cpp
--- a/Engine/Source/Editor/LevelEditor/Private/LevelEditor.cpp
+++ b/Engine/Source/Editor/LevelEditor/Private/LevelEditor.cpp
@@ -1859,6 +1859,27 @@ void FLevelEditorModule::BindGlobalLevelEditorCommands()
 		FCanExecuteAction(),
 		FIsActionChecked::CreateStatic(&FLevelEditorActionCallbacks::IsInvertRightMouseDollyYAxisChecked));
 
+	ActionList.MapAction(
+		Commands.AllowArcballRotation,
+		FExecuteAction::CreateStatic(&FLevelEditorActionCallbacks::ToggleAllowArcballRotation),
+		FCanExecuteAction(),
+		FIsActionChecked::CreateStatic(&FLevelEditorActionCallbacks::IsAllowArcballRotationChecked)
+	);
+
+	ActionList.MapAction(
+		Commands.AllowScreenspaceRotation,
+		FExecuteAction::CreateStatic(&FLevelEditorActionCallbacks::ToggleAllowScreenspaceRotation),
+		FCanExecuteAction(),
+		FIsActionChecked::CreateStatic(&FLevelEditorActionCallbacks::IsAllowScreenspaceRotationChecked)
+	);
+
+	ActionList.MapAction(
+		Commands.EnableViewportHoverFeedback,
+		FExecuteAction::CreateStatic(&FLevelEditorActionCallbacks::ToggleEnableViewportHoverFeedback),
+		FCanExecuteAction(),
+		FIsActionChecked::CreateStatic(&FLevelEditorActionCallbacks::IsEnableViewportHoverFeedbackChecked)
+	);
+
 	const TArray<FPreviewPlatformMenuItem>& MenuItems = FDataDrivenPlatformInfoRegistry::GetAllPreviewPlatformMenuItems();
 	// We need one extra slot for the Disable Preview option
 	check(MenuItems.Num() == Commands.PreviewPlatformOverrides.Num());
diff --git a/Engine/Source/Editor/LevelEditor/Private/LevelEditorActions.cpp b/Engine/Source/Editor/LevelEditor/Private/LevelEditorActions.cpp
--- a/Engine/Source/Editor/LevelEditor/Private/LevelEditorActions.cpp
+++ b/Engine/Source/Editor/LevelEditor/Private/LevelEditorActions.cpp
@@ -3804,6 +3804,60 @@ bool FLevelEditorActionCallbacks::GeometryCollection_IsChecked()
 	return true;
 }
 
+void FLevelEditorActionCallbacks::ToggleAllowArcballRotation()
+{
+	if (ULevelEditorViewportSettings* ViewportSettings = GetMutableDefault<ULevelEditorViewportSettings>())
+	{
+		ViewportSettings->bAllowArcballRotate = !ViewportSettings->bAllowArcballRotate;
+	}
+}
+
+bool FLevelEditorActionCallbacks::IsAllowArcballRotationChecked()
+{
+	if (const ULevelEditorViewportSettings* const ViewportSettings = GetDefault<ULevelEditorViewportSettings>())
+	{
+		return ViewportSettings->bAllowArcballRotate;
+	}
+
+	return false;
+}
+
+void FLevelEditorActionCallbacks::ToggleAllowScreenspaceRotation()
+{
+	if (ULevelEditorViewportSettings* ViewportSettings = GetMutableDefault<ULevelEditorViewportSettings>())
+	{
+		ViewportSettings->bAllowScreenRotate = !ViewportSettings->bAllowScreenRotate;
+	}
+}
+
+bool FLevelEditorActionCallbacks::IsAllowScreenspaceRotationChecked()
+{
+	if (const ULevelEditorViewportSettings* const ViewportSettings = GetDefault<ULevelEditorViewportSettings>())
+	{
+		return ViewportSettings->bAllowScreenRotate;
+	}
+
+	return false;
+}
+
+void FLevelEditorActionCallbacks::ToggleEnableViewportHoverFeedback()
+{
+	if (ULevelEditorViewportSettings* ViewportSettings = GetMutableDefault<ULevelEditorViewportSettings>())
+	{
+		ViewportSettings->bEnableViewportHoverFeedback = !ViewportSettings->bEnableViewportHoverFeedback;
+	}
+}
+
+bool FLevelEditorActionCallbacks::IsEnableViewportHoverFeedbackChecked()
+{
+	if (const ULevelEditorViewportSettings* const ViewportSettings = GetDefault<ULevelEditorViewportSettings>())
+	{
+		return ViewportSettings->bEnableViewportHoverFeedback;
+	}
+
+	return false;
+}
+
 void FLevelEditorActionCallbacks::TogglePreviewSelectedCameras(const TWeakPtr<::SLevelViewport>& InLevelViewportWeak)
 {
 	if (ULevelEditorViewportSettings* ViewportSettings = GetMutableDefault<ULevelEditorViewportSettings>())
@@ -4289,6 +4343,10 @@ void FLevelEditorCommands::RegisterCommands()
 
 	UI_COMMAND(ToggleFeatureLevelPreview, "Preview Mode Toggle", "Toggles the Preview Mode on or off for the currently selected Preview target", EUserInterfaceActionType::ToggleButton, FInputChord());
 
+	UI_COMMAND(AllowArcballRotation, "Enable Arcball Rotation", "Allow arcball rotation with rotate widget", EUserInterfaceActionType::ToggleButton, FInputChord());
+	UI_COMMAND(AllowScreenspaceRotation, "Enable Screenspace Rotation", "Allow screen rotation with rotate widget", EUserInterfaceActionType::ToggleButton, FInputChord());
+	UI_COMMAND(EnableViewportHoverFeedback, "Preselection Highlight", "Enables real-time hover feedback when mousing over objects in editor view ports", EUserInterfaceActionType::ToggleButton, FInputChord());
+
 	// Camera Preferences
 	UI_COMMAND(OrbitCameraAroundSelection, "Orbit Around Selection", "If enabled, the camera will orbit around the current selection in the viewport", EUserInterfaceActionType::ToggleButton, FInputChord());
 	UI_COMMAND(LinkOrthographicViewports, "Link Ortho Camera Movement", "If checked all orthographic view ports are linked to the same position and move together.", EUserInterfaceActionType::ToggleButton, FInputChord());
diff --git a/Engine/Source/Editor/LevelEditor/Private/ViewportToolbar/LevelEditorSubmenus.cpp b/Engine/Source/Editor/LevelEditor/Private/ViewportToolbar/LevelEditorSubmenus.cpp
--- a/Engine/Source/Editor/LevelEditor/Private/ViewportToolbar/LevelEditorSubmenus.cpp
+++ b/Engine/Source/Editor/LevelEditor/Private/ViewportToolbar/LevelEditorSubmenus.cpp
@@ -2959,6 +2959,9 @@ void ExtendTransformSubmenu(FName InTransformSubmenuName)
 		FToolMenuSection& GizmoSection = Submenu->FindOrAddSection("Gizmo", LOCTEXT("GizmoLabel", "Gizmo"));
 		FToolMenuEntry& ShowTransformGizmoEntry = GizmoSection.AddMenuEntry(FLevelEditorCommands::Get().ShowTransformWidget);
 		ShowTransformGizmoEntry.InsertPosition.Position = EToolMenuInsertType::First;
+
+		GizmoSection.AddMenuEntry(FLevelEditorCommands::Get().AllowArcballRotation);
+		GizmoSection.AddMenuEntry(FLevelEditorCommands::Get().AllowScreenspaceRotation);
 	}
 
 	// Selection
@@ -2969,6 +2972,7 @@ void ExtendTransformSubmenu(FName InTransformSubmenuName)
 		SelectionSection.AddMenuEntry(FLevelEditorCommands::Get().StrictBoxSelect);
 		SelectionSection.AddMenuEntry(FLevelEditorCommands::Get().TransparentBoxSelect);
 		SelectionSection.AddMenuEntry(FLevelEditorCommands::Get().ShowSelectionSubcomponents);
+		SelectionSection.AddMenuEntry(FLevelEditorCommands::Get().EnableViewportHoverFeedback);
 	}
 }
 
diff --git a/Engine/Source/Editor/LevelEditor/Public/LevelEditorActions.h b/Engine/Source/Editor/LevelEditor/Public/LevelEditorActions.h
--- a/Engine/Source/Editor/LevelEditor/Public/LevelEditorActions.h
+++ b/Engine/Source/Editor/LevelEditor/Public/LevelEditorActions.h
@@ -665,6 +665,9 @@ public:
 
 	TSharedPtr< FUICommandInfo > FixupGroupActor;
 
+	TSharedPtr<FUICommandInfo> AllowArcballRotation;
+	TSharedPtr<FUICommandInfo> AllowScreenspaceRotation;
+	TSharedPtr<FUICommandInfo> EnableViewportHoverFeedback;
 	/**
 	 * Controls
 	 */
@@ -886,6 +889,15 @@ public:
 	static void GeometryCollection_SelectInverseGeometry();
 	static bool GeometryCollection_IsChecked();
 
+	static void ToggleAllowArcballRotation();
+	static bool IsAllowArcballRotationChecked();
+
+	static void ToggleAllowScreenspaceRotation();
+	static bool IsAllowScreenspaceRotationChecked();
+
+	static void ToggleEnableViewportHoverFeedback();
+	static bool IsEnableViewportHoverFeedbackChecked();
+
 	static void TogglePreviewSelectedCameras(const TWeakPtr<SLevelViewport>& InLevelViewportWeak = {});
 	static bool IsPreviewSelectedCamerasChecked();
 
