Fix an issue where the level editor toolbar started to clip items while there was still free space.

diff --git a/Engine/Source/Editor/LevelEditor/Private/LevelEditorToolBar.cpp b/Engine/Source/Editor/LevelEditor/Private/LevelEditorToolBar.cpp
--- a/Engine/Source/Editor/LevelEditor/Private/LevelEditorToolBar.cpp
+++ b/Engine/Source/Editor/LevelEditor/Private/LevelEditorToolBar.cpp
@@ -1651,6 +1651,7 @@ TSharedRef< SWidget > FLevelEditorToolBar::MakeLevelEditorToolBar( const TShared
 			]
 			+ SHorizontalBox::Slot()
 			.HAlign(HAlign_Left)
+			.FillWidth(1.0f)
 			[
 				SNew(SBorder)
 				.Padding(0.f)
@@ -1662,6 +1663,7 @@ TSharedRef< SWidget > FLevelEditorToolBar::MakeLevelEditorToolBar( const TShared
 			]
 			+ SHorizontalBox::Slot()
 			.HAlign(HAlign_Right)
+			.AutoWidth()
 			.Padding(0.0f, 0.0f, 7.0f, 0.0f)
 			[
 				SNew(SBorder)
