Level Editor Actions
Fix AActor::GetSceneOutlinerTopParentPackage returning the external package while it was meant to return the actual package, regardless if it is external or not.
Fix several cases using AActor::GetSceneOutlinerTopParentPackage instead of AActor::GetSceneOutlinerItemPackage, now that AActor::GetSceneOutlinerTopParentPackage is fixed.
Fix FLevelEditorActionCallbacks::CopyActorFilePathtoClipboard_Clicked to always append the actor labels, regardless if they are actor proxies or not (useful for scene graph entities).

diff --git a/Engine/Source/Editor/LevelEditor/Private/LevelEditorActions.cpp b/Engine/Source/Editor/LevelEditor/Private/LevelEditorActions.cpp
--- a/Engine/Source/Editor/LevelEditor/Private/LevelEditorActions.cpp
+++ b/Engine/Source/Editor/LevelEditor/Private/LevelEditorActions.cpp
@@ -3139,20 +3139,14 @@ void FLevelEditorActionCallbacks::CopyActorFilePathtoClipboard_Clicked()
 
 		if (SelectedActors.Num() > 1)
 		{
-			const UPackage* ActorPackage = Actor->GetPackage();
-
-			// If the packages differ, it's probably an actor proxy, so we can't rely on the actor's name
-			if (SceneOutlinerPackage == ActorPackage)
+			const FString& ActorLabel = Actor->GetActorLabel(false);
+			if (ActorLabel.Len())
 			{
-				const FString& ActorLabel = Actor->GetActorLabel(false);
-				if (ActorLabel.Len())
-				{
-					Result.Append(ActorLabel);
-				}
-				Result.Append(TEXT("("));
-				Result.Append(Actor->GetName());
-				Result.Append(TEXT("): "));
+				Result.Append(ActorLabel);
 			}
+			Result.Append(TEXT("("));
+			Result.Append(Actor->GetName());
+			Result.Append(TEXT("): "));
 		}
 
 		LocalFullPath = FPaths::ConvertRelativePathToFull(LocalFullPath);
@@ -3173,7 +3167,7 @@ void FLevelEditorActionCallbacks::SaveActor_Clicked()
 	for (FSelectionIterator It(GEditor->GetSelectedActorIterator()); It; ++It)
 	{
 		AActor* Actor = Cast<AActor>(*It);
-		if (UPackage* ActorPackage = Actor->GetSceneOutlinerTopParentPackage())
+		if (UPackage* ActorPackage = Actor->GetSceneOutlinerItemPackage())
 		{
 			PackagesToSave.Add(ActorPackage);
 		}
@@ -3196,7 +3190,7 @@ bool FLevelEditorActionCallbacks::SaveActor_CanExecute()
 	for (FSelectionIterator It(GEditor->GetSelectedActorIterator()); It; ++It)
 	{
 		AActor* Actor = Cast<AActor>(*It);
-		if (UPackage* ActorPackage = Actor->GetSceneOutlinerTopParentPackage())
+		if (UPackage* ActorPackage = Actor->GetSceneOutlinerItemPackage())
 		{
 			return true;
 		}
@@ -3212,7 +3206,7 @@ static TArray<FString> GetSelectedActorsPackageFullpath()
 	{
 		if (const AActor* Actor = Cast<AActor>(*It))
 		{
-			if (const UPackage* Package = Actor->GetSceneOutlinerTopParentPackage())
+			if (const UPackage* Package = Actor->GetSceneOutlinerItemPackage())
 			{
 				const FString LocalFullPath(Package->GetLoadedPath().GetLocalFullPath());
 
diff --git a/Engine/Source/Runtime/Engine/Private/ActorEditor.cpp b/Engine/Source/Runtime/Engine/Private/ActorEditor.cpp
--- a/Engine/Source/Runtime/Engine/Private/ActorEditor.cpp
+++ b/Engine/Source/Runtime/Engine/Private/ActorEditor.cpp
@@ -1092,7 +1092,7 @@ UPackage* AActor::GetSceneOutlinerItemPackage() const
 
 UPackage* AActor::GetSceneOutlinerTopParentPackage() const
 {
-	return GetExternalPackage();
+	return GetPackage();
 }
 
 class UHLODLayer* AActor::GetHLODLayer() const
