[Viewport Toolbar] Remove redundant entries from Level Editor Camera menu

diff --git a/Engine/Source/Editor/LevelEditor/Private/ViewportToolbar/LevelEditorSubmenus.cpp b/Engine/Source/Editor/LevelEditor/Private/ViewportToolbar/LevelEditorSubmenus.cpp
--- a/Engine/Source/Editor/LevelEditor/Private/ViewportToolbar/LevelEditorSubmenus.cpp
+++ b/Engine/Source/Editor/LevelEditor/Private/ViewportToolbar/LevelEditorSubmenus.cpp
@@ -2891,19 +2891,6 @@ void ExtendCameraSubmenu(FName InCameraOptionsSubmenuName)
 			}
 		)
 	);
-	
-	// Extend Camera Movement Menu
-	{
-		UToolMenu* CameraMovementMenu = UToolMenus::Get()->ExtendMenu(UToolMenus::JoinMenuPaths(InCameraOptionsSubmenuName, "CameraMovement"));
-		
-		FToolMenuSection& CameraMovementSection = CameraMovementMenu->FindOrAddSection("CameraMovement");
-		CameraMovementSection.AddMenuEntry(FLevelEditorCommands::Get().SnapObjectToCamera)
-			.InsertPosition = FToolMenuInsert("SnapCameraToObject", EToolMenuInsertType::After);
-		
-		FToolMenuSection& CameraOrthoSection = CameraMovementMenu->FindOrAddSection("Orthographic");
-		CameraOrthoSection.AddMenuEntry(FLevelEditorCommands::Get().LinkOrthographicViewports)
-			.InsertPosition = FToolMenuInsert("OrthoZoomToCursor", EToolMenuInsertType::Before);
-	}
 }
 
 void ExtendTransformSubmenu(FName InTransformSubmenuName)
