Update the Delta Transform dialog to support the EAxisList::LeftUpForward display coordinate system.

diff --git a/Engine/Source/Editor/LevelEditor/Private/DlgDeltaTransform.cpp b/Engine/Source/Editor/LevelEditor/Private/DlgDeltaTransform.cpp
--- a/Engine/Source/Editor/LevelEditor/Private/DlgDeltaTransform.cpp
+++ b/Engine/Source/Editor/LevelEditor/Private/DlgDeltaTransform.cpp
@@ -4,6 +4,7 @@
 #include "InputCoreTypes.h"
 #include "Input/Reply.h"
 #include "Math/MathFwd.h"
+#include "Misc/AxisDisplayInfo.h"
 #include "Widgets/DeclarativeSyntaxSupport.h"
 #include "Widgets/SCompoundWidget.h"
 #include "Widgets/SBoxPanel.h"
@@ -73,6 +74,9 @@ public:
 							.X( this, &SDlgDeltaTransform::GetDeltaX )
 							.Y( this, &SDlgDeltaTransform::GetDeltaY )
 							.Z( this, &SDlgDeltaTransform::GetDeltaZ )
+							.XDisplayName(AxisDisplayInfo::GetAxisDisplayName(EAxisList::Forward))
+							.YDisplayName(AxisDisplayInfo::GetAxisDisplayName(EAxisList::Left))
+							.ZDisplayName(AxisDisplayInfo::GetAxisDisplayName(EAxisList::Up))
 							.bColorAxisLabels( true )
 							.AllowSpin( false )
 							.OnXCommitted( this, &SDlgDeltaTransform::OnSetDelta, 0 )
@@ -193,9 +197,15 @@ private:
 
 		if (ButtonID == FDlgDeltaTransform::OK)
 		{
+			const bool bIsAxisDisplayLeftUpForward = AxisDisplayInfo::GetAxisDisplayCoordinateSystem() == EAxisList::LeftUpForward;
+
+			const FVector::FReal DeltaX = DeltaTransform.X.GetValue();
+			const FVector::FReal DeltaY = DeltaTransform.Y.GetValue();
+			const FVector::FReal DeltaZ = DeltaTransform.Z.GetValue();		
+
 			GUnrealEd->Exec( GEditor->GetEditorWorldContext().World(), 
 				*FString::Printf(TEXT("ELEMENT DELTAMOVE X=%.5f Y=%.5f Z=%.5f"), 
-				DeltaTransform.X.GetValue(), DeltaTransform.Y.GetValue(), DeltaTransform.Z.GetValue()) );
+				DeltaX, bIsAxisDisplayLeftUpForward ? -DeltaY : DeltaY, DeltaZ) );
 		}
 
 		return FReply::Handled();
