[Transforms] Fix an issue where scale changes were allowable through mirror functions if the component or selected actor had disallowed it
Fix spin boxes not getting their context menu extenders passed through when created by vector input boxes
Update mirror axis menu strings to pull from axis info

diff --git a/Engine/Plugins/Experimental/Avalanche/Source/AvalancheLevelViewport/Private/Customizations/AvaLevelViewportComponentTransformDetails.cpp b/Engine/Plugins/Experimental/Avalanche/Source/AvalancheLevelViewport/Private/Customizations/AvaLevelViewportComponentTransformDetails.cpp
--- a/Engine/Plugins/Experimental/Avalanche/Source/AvalancheLevelViewport/Private/Customizations/AvaLevelViewportComponentTransformDetails.cpp
+++ b/Engine/Plugins/Experimental/Avalanche/Source/AvalancheLevelViewport/Private/Customizations/AvaLevelViewportComponentTransformDetails.cpp
@@ -974,8 +974,8 @@ void FAvaLevelViewportComponentTransformDetails::ExtendXScaleContextMenu(FMenuBu
 {
 	MenuBuilder.BeginSection("ScaleOperations", LOCTEXT("ScaleOperations", "Scale Operations"));
 	MenuBuilder.AddMenuEntry(
-		LOCTEXT("MirrorValueX", "Mirror X"),
-		LOCTEXT("MirrorValueX_Tooltip", "Mirror scale value on the X axis"),
+		FText::Format(LOCTEXT("MirrorValue", "Mirror {0} Axis"), AxisDisplayInfo::GetAxisDisplayName(EAxisList::Forward)),
+		FText::Format(LOCTEXT("MirrorValue_Tooltip", "Mirror scale value on the {0} axis"), AxisDisplayInfo::GetAxisDisplayName(EAxisList::Forward)),
 		FSlateIcon(),
 		FUIAction(FExecuteAction::CreateSP(this, &FAvaLevelViewportComponentTransformDetails::OnXScaleMirrored), FCanExecuteAction())
 	);
@@ -986,8 +986,8 @@ void FAvaLevelViewportComponentTransformDetails::ExtendYScaleContextMenu(FMenuBu
 {
 	MenuBuilder.BeginSection("ScaleOperations", LOCTEXT("ScaleOperations", "Scale Operations"));
 	MenuBuilder.AddMenuEntry(
-		LOCTEXT("MirrorValueY", "Mirror Y"),
-		LOCTEXT("MirrorValueY_Tooltip", "Mirror scale value on the Y axis"),
+		FText::Format(LOCTEXT("MirrorValue", "Mirror {0} Axis"), AxisDisplayInfo::GetAxisDisplayName(EAxisList::Left)),
+		FText::Format(LOCTEXT("MirrorValue_Tooltip", "Mirror scale value on the {0} axis"), AxisDisplayInfo::GetAxisDisplayName(EAxisList::Left)),
 		FSlateIcon(),
 		FUIAction(FExecuteAction::CreateSP(this, &FAvaLevelViewportComponentTransformDetails::OnYScaleMirrored), FCanExecuteAction())
 	);
@@ -998,8 +998,8 @@ void FAvaLevelViewportComponentTransformDetails::ExtendZScaleContextMenu(FMenuBu
 {
 	MenuBuilder.BeginSection("ScaleOperations", LOCTEXT("ScaleOperations", "Scale Operations"));
 	MenuBuilder.AddMenuEntry(
-		LOCTEXT("MirrorValueZ", "Mirror Z"),
-		LOCTEXT("MirrorValueZ_Tooltip", "Mirror scale value on the Z axis"),
+		FText::Format(LOCTEXT("MirrorValue", "Mirror {0} Axis"), AxisDisplayInfo::GetAxisDisplayName(EAxisList::Up)),
+		FText::Format(LOCTEXT("MirrorValue_Tooltip", "Mirror scale value on the {0} axis"), AxisDisplayInfo::GetAxisDisplayName(EAxisList::Up)),
 		FSlateIcon(),
 		FUIAction(FExecuteAction::CreateSP(this, &FAvaLevelViewportComponentTransformDetails::OnZScaleMirrored), FCanExecuteAction())
 	);
@@ -1009,7 +1009,7 @@ void FAvaLevelViewportComponentTransformDetails::ExtendZScaleContextMenu(FMenuBu
 void FAvaLevelViewportComponentTransformDetails::OnXScaleMirrored()
 {
 	FSlateApplication::Get().ClearKeyboardFocus(EFocusCause::Mouse);
-	FScopedTransaction Transaction(LOCTEXT("MirrorActorScaleX", "Mirror actor scale X"));
+	FScopedTransaction Transaction(FText::Format(LOCTEXT("MirrorScaleTransaction", "Scale - Mirror {0} Axis"), AxisDisplayInfo::GetAxisDisplayName(EAxisList::Forward)));
 	OnBeginChange(EAvaLevelViewportTransformFieldType::Scale);
 	OnSetTransform(EAvaLevelViewportTransformFieldType::Scale, EAxisList::X, FVector(1.0f), /* bMirror */ true, /* bCommitted */ true);
 	OnEndChange(EAvaLevelViewportTransformFieldType::Scale, EAxisList::X);
@@ -1018,7 +1018,7 @@ void FAvaLevelViewportComponentTransformDetails::OnXScaleMirrored()
 void FAvaLevelViewportComponentTransformDetails::OnYScaleMirrored()
 {
 	FSlateApplication::Get().ClearKeyboardFocus(EFocusCause::Mouse);
-	FScopedTransaction Transaction(LOCTEXT("MirrorActorScaleY", "Mirror actor scale Y"));
+	FScopedTransaction Transaction(FText::Format(LOCTEXT("MirrorScaleTransaction", "Scale - Mirror {0} Axis"), AxisDisplayInfo::GetAxisDisplayName(EAxisList::Left)));
 	OnBeginChange(EAvaLevelViewportTransformFieldType::Scale);
 	OnSetTransform(EAvaLevelViewportTransformFieldType::Scale, EAxisList::Y, FVector(1.0f), /* bMirror */ true, /* bCommitted */ true);
 	OnEndChange(EAvaLevelViewportTransformFieldType::Scale, EAxisList::Y);
@@ -1027,7 +1027,7 @@ void FAvaLevelViewportComponentTransformDetails::OnYScaleMirrored()
 void FAvaLevelViewportComponentTransformDetails::OnZScaleMirrored()
 {
 	FSlateApplication::Get().ClearKeyboardFocus(EFocusCause::Mouse);
-	FScopedTransaction Transaction(LOCTEXT("MirrorActorScaleZ", "Mirror actor scale Z"));
+	FScopedTransaction Transaction(FText::Format(LOCTEXT("MirrorScaleTransaction", "Scale - Mirror {0} Axis"), AxisDisplayInfo::GetAxisDisplayName(EAxisList::Up)));
 	OnBeginChange(EAvaLevelViewportTransformFieldType::Scale);
 	OnSetTransform(EAvaLevelViewportTransformFieldType::Scale, EAxisList::Z, FVector(1.0f), /* bMirror */ true, /* bCommitted */ true);
 	OnEndChange(EAvaLevelViewportTransformFieldType::Scale, EAxisList::Z);
diff --git a/Engine/Source/Editor/DetailCustomizations/Private/ComponentTransformDetails.cpp b/Engine/Source/Editor/DetailCustomizations/Private/ComponentTransformDetails.cpp
--- a/Engine/Source/Editor/DetailCustomizations/Private/ComponentTransformDetails.cpp
+++ b/Engine/Source/Editor/DetailCustomizations/Private/ComponentTransformDetails.cpp
@@ -1034,11 +1034,11 @@ void FComponentTransformDetails::OnScaleResetClicked()
 void FComponentTransformDetails::ExtendXScaleContextMenu( FMenuBuilder& MenuBuilder )
 {
 	MenuBuilder.BeginSection( "ScaleOperations", LOCTEXT( "ScaleOperations", "Scale Operations" ) );
-	MenuBuilder.AddMenuEntry( 
-		LOCTEXT( "MirrorValueX", "Mirror X" ),  
-		LOCTEXT( "MirrorValueX_Tooltip", "Mirror scale value on the X axis" ), 
+	MenuBuilder.AddMenuEntry(
+		FText::Format(LOCTEXT("MirrorValue", "Mirror {0} Axis"), AxisDisplayInfo::GetAxisDisplayName(EAxisList::Forward)),
+		FText::Format(LOCTEXT("MirrorValue_Tooltip", "Mirror scale value on the {0} axis"), AxisDisplayInfo::GetAxisDisplayName(EAxisList::Forward)),
 		FSlateIcon(), 		
-		FUIAction( FExecuteAction::CreateSP( this, &FComponentTransformDetails::OnXScaleMirrored ), FCanExecuteAction() )
+		FUIAction( FExecuteAction::CreateSP( this, &FComponentTransformDetails::OnXScaleMirrored ), FCanExecuteAction::CreateSP( this, &FComponentTransformDetails::GetIsScaleEnabled ) )
 	);
 	MenuBuilder.EndSection();
 }
@@ -1046,11 +1046,11 @@ void FComponentTransformDetails::ExtendXScaleContextMenu( FMenuBuilder& MenuBuil
 void FComponentTransformDetails::ExtendYScaleContextMenu( FMenuBuilder& MenuBuilder )
 {
 	MenuBuilder.BeginSection( "ScaleOperations", LOCTEXT( "ScaleOperations", "Scale Operations" ) );
-	MenuBuilder.AddMenuEntry( 
-		LOCTEXT( "MirrorValueY", "Mirror Y" ),  
-		LOCTEXT( "MirrorValueY_Tooltip", "Mirror scale value on the Y axis" ), 
+	MenuBuilder.AddMenuEntry(
+		FText::Format(LOCTEXT("MirrorValue", "Mirror {0} Axis"), AxisDisplayInfo::GetAxisDisplayName(EAxisList::Left)),
+		FText::Format(LOCTEXT("MirrorValue_Tooltip", "Mirror scale value on the {0} axis"), AxisDisplayInfo::GetAxisDisplayName(EAxisList::Left)),
 		FSlateIcon(), 		
-		FUIAction( FExecuteAction::CreateSP( this, &FComponentTransformDetails::OnYScaleMirrored ), FCanExecuteAction() )
+		FUIAction( FExecuteAction::CreateSP( this, &FComponentTransformDetails::OnYScaleMirrored ), FCanExecuteAction::CreateSP( this, &FComponentTransformDetails::GetIsScaleEnabled ) )
 	);
 	MenuBuilder.EndSection();
 }
@@ -1058,11 +1058,11 @@ void FComponentTransformDetails::ExtendYScaleContextMenu( FMenuBuilder& MenuBuil
 void FComponentTransformDetails::ExtendZScaleContextMenu( FMenuBuilder& MenuBuilder )
 {
 	MenuBuilder.BeginSection( "ScaleOperations", LOCTEXT( "ScaleOperations", "Scale Operations" ) );
-	MenuBuilder.AddMenuEntry( 
-		LOCTEXT( "MirrorValueZ", "Mirror Z" ),  
-		LOCTEXT( "MirrorValueZ_Tooltip", "Mirror scale value on the Z axis" ), 
+	MenuBuilder.AddMenuEntry(
+		FText::Format(LOCTEXT("MirrorValue", "Mirror {0} Axis"), AxisDisplayInfo::GetAxisDisplayName(EAxisList::Up)),
+		FText::Format(LOCTEXT("MirrorValue_Tooltip", "Mirror scale value on the {0} axis"), AxisDisplayInfo::GetAxisDisplayName(EAxisList::Up)),
 		FSlateIcon(), 		
-		FUIAction( FExecuteAction::CreateSP( this, &FComponentTransformDetails::OnZScaleMirrored ), FCanExecuteAction() )
+		FUIAction( FExecuteAction::CreateSP( this, &FComponentTransformDetails::OnZScaleMirrored ), FCanExecuteAction::CreateSP( this, &FComponentTransformDetails::GetIsScaleEnabled ) )
 	);
 	MenuBuilder.EndSection();
 }
@@ -1070,21 +1070,21 @@ void FComponentTransformDetails::ExtendZScaleContextMenu( FMenuBuilder& MenuBuil
 void FComponentTransformDetails::OnXScaleMirrored()
 {
 	FSlateApplication::Get().ClearKeyboardFocus(EFocusCause::Mouse);
-	FScopedTransaction Transaction(LOCTEXT("MirrorActorScaleX", "Mirror actor scale X"));
+	FScopedTransaction Transaction(FText::Format(LOCTEXT("MirrorScaleTransaction", "Scale - Mirror {0} Axis"), AxisDisplayInfo::GetAxisDisplayName(EAxisList::Forward)));
 	OnSetTransform(ETransformField::Scale, EAxisList::X, FVector(1.0f), true, true);
 }
 
 void FComponentTransformDetails::OnYScaleMirrored()
 {
 	FSlateApplication::Get().ClearKeyboardFocus(EFocusCause::Mouse);
-	FScopedTransaction Transaction(LOCTEXT("MirrorActorScaleY", "Mirror actor scale Y"));
+	FScopedTransaction Transaction(FText::Format(LOCTEXT("MirrorScaleTransaction", "Scale - Mirror {0} Axis"), AxisDisplayInfo::GetAxisDisplayName(EAxisList::Left)));
 	OnSetTransform(ETransformField::Scale, EAxisList::Y, FVector(1.0f), true, true);
 }
 
 void FComponentTransformDetails::OnZScaleMirrored()
 {
 	FSlateApplication::Get().ClearKeyboardFocus(EFocusCause::Mouse);
-	FScopedTransaction Transaction(LOCTEXT("MirrorActorScaleZ", "Mirror actor scale Z"));
+	FScopedTransaction Transaction(FText::Format(LOCTEXT("MirrorScaleTransaction", "Scale - Mirror {0} Axis"), AxisDisplayInfo::GetAxisDisplayName(EAxisList::Up)));
 	OnSetTransform(ETransformField::Scale, EAxisList::Z, FVector(1.0f), true, true);
 }
 
diff --git a/Engine/Source/Editor/EditorFramework/Private/Subsystems/EditorElementSubsystem.cpp b/Engine/Source/Editor/EditorFramework/Private/Subsystems/EditorElementSubsystem.cpp
--- a/Engine/Source/Editor/EditorFramework/Private/Subsystems/EditorElementSubsystem.cpp
+++ b/Engine/Source/Editor/EditorFramework/Private/Subsystems/EditorElementSubsystem.cpp
@@ -39,22 +39,26 @@ FTypedElementListRef UEditorElementSubsystem::GetEditorNormalizedSelectionSet(co
 	return SelectionSet.GetNormalizedSelection(NormalizationOptions);
 }
 
-FTypedElementListRef UEditorElementSubsystem::GetEditorManipulableElements(const FTypedElementListRef& NormalizedSelection)
+FTypedElementListRef UEditorElementSubsystem::GetEditorManipulableElements(const FTypedElementListRef& NormalizedSelection, UE::Widget::EWidgetMode InManipulationType, UWorld* InRequiredWorld)
 {
-	 NormalizedSelection->RemoveAll<ITypedElementWorldInterface>([](const TTypedElement<ITypedElementWorldInterface>& InWorldElement)
+	 NormalizedSelection->RemoveAll<ITypedElementWorldInterface>([InManipulationType, InRequiredWorld](const TTypedElement<ITypedElementWorldInterface>& InWorldElement)
 		{
-			return !UEditorElementSubsystem::IsElementEditorManipulable(InWorldElement);
+			return !UEditorElementSubsystem::IsElementEditorManipulable(InWorldElement, InManipulationType, InRequiredWorld);
 		});
 
 	 return NormalizedSelection;
 }
 
-TTypedElement<ITypedElementWorldInterface> UEditorElementSubsystem::GetLastSelectedEditorManipulableElement(const FTypedElementListRef& NormalizedSelection)
+TTypedElement<ITypedElementWorldInterface> UEditorElementSubsystem::GetLastSelectedEditorManipulableElement(const FTypedElementListRef& NormalizedSelection, UE::Widget::EWidgetMode InManipulationType, UWorld* InRequiredWorld)
 {
-	return NormalizedSelection->GetBottomElement<ITypedElementWorldInterface>(&UEditorElementSubsystem::IsElementEditorManipulable);
+	return NormalizedSelection->GetBottomElement<ITypedElementWorldInterface>(
+		[InManipulationType, InRequiredWorld](const TTypedElement<ITypedElementWorldInterface> Element)
+		{
+			return UEditorElementSubsystem::IsElementEditorManipulable(Element, InManipulationType, InRequiredWorld);
+		});
 }
 
-bool UEditorElementSubsystem::IsElementEditorManipulable(const TTypedElement<ITypedElementWorldInterface>& WorldElement)
+bool UEditorElementSubsystem::IsElementEditorManipulable(const TTypedElement<ITypedElementWorldInterface>& WorldElement, UE::Widget::EWidgetMode InManipulationType, UWorld* InRequiredWorld)
 {
 	UWorld* OwnerWorld = WorldElement.GetOwnerWorld();
 	if (!OwnerWorld)
@@ -69,5 +73,18 @@ bool UEditorElementSubsystem::IsElementEditorManipulable(const TTypedElement<ITy
 		return false;
 	}
 
+	if (InManipulationType == UE::Widget::WM_Scale && !WorldElement.CanScaleElement())
+	{
+		return false;
+	}
+
+	if (InRequiredWorld != nullptr)
+	{
+		if (OwnerWorld != InRequiredWorld)
+		{
+			return false;
+		}
+	}
+
 	return true;
 }
\ No newline at end of file
diff --git a/Engine/Source/Editor/EditorFramework/Public/Subsystems/EditorElementSubsystem.h b/Engine/Source/Editor/EditorFramework/Public/Subsystems/EditorElementSubsystem.h
--- a/Engine/Source/Editor/EditorFramework/Public/Subsystems/EditorElementSubsystem.h
+++ b/Engine/Source/Editor/EditorFramework/Public/Subsystems/EditorElementSubsystem.h
@@ -10,6 +10,7 @@
 #include "Math/Transform.h"
 #include "UObject/ObjectMacros.h"
 #include "UObject/UObjectGlobals.h"
+#include "UnrealWidgetFwd.h"
 
 #include "EditorElementSubsystem.generated.h"
 
@@ -38,14 +39,14 @@ public:
 	 * Require a editor normalized selection list to behave like the editor selection
 	 * Note: A manipulable element is a element that would be move when manipulating the gizmo in the editor
 	 */
-	static FTypedElementListRef GetEditorManipulableElements(const FTypedElementListRef& NormalizedSelection);
+	static FTypedElementListRef GetEditorManipulableElements(const FTypedElementListRef& NormalizedSelection, UE::Widget::EWidgetMode InManipulationType = UE::Widget::WM_Translate, UWorld* InRequiredWorld = nullptr);
 
 	/**
 	 * Return the most recently selected element that is manipulable.
 	 * Require a editor normalized selection list to behave like the editor selection
 	 * Note: A manipulable element is a element that would be move when manipulating the gizmo in the editor.
 	 */
-	static TTypedElement<ITypedElementWorldInterface> GetLastSelectedEditorManipulableElement(const FTypedElementListRef& NormalizedSelection);
+	static TTypedElement<ITypedElementWorldInterface> GetLastSelectedEditorManipulableElement(const FTypedElementListRef& NormalizedSelection, UE::Widget::EWidgetMode InManipulationType = UE::Widget::WM_Translate, UWorld* InRequiredWorld = nullptr);
 
-	static bool IsElementEditorManipulable(const TTypedElement<ITypedElementWorldInterface>& WorldElement);
+	static bool IsElementEditorManipulable(const TTypedElement<ITypedElementWorldInterface>& WorldElement, UE::Widget::EWidgetMode InManipulationType = UE::Widget::WM_Translate, UWorld* InRequiredWorld = nullptr);
 };
\ No newline at end of file
diff --git a/Engine/Source/Editor/LevelEditor/Private/LevelEditor.cpp b/Engine/Source/Editor/LevelEditor/Private/LevelEditor.cpp
--- a/Engine/Source/Editor/LevelEditor/Private/LevelEditor.cpp
+++ b/Engine/Source/Editor/LevelEditor/Private/LevelEditor.cpp
@@ -992,19 +992,19 @@ void FLevelEditorModule::BindGlobalLevelEditorCommands()
 	ActionList.MapAction(
 		Commands.MirrorActorX,
 		FExecuteAction::CreateStatic( &FLevelEditorActionCallbacks::ExecuteExecCommand, FString( TEXT("ELEMENT MIRROR X=-1") ) ),
-		FCanExecuteAction::CreateStatic(&FLevelEditorActionCallbacks::ElementSelected_CanExecuteMove)
+		FCanExecuteAction::CreateStatic(&FLevelEditorActionCallbacks::ElementSelected_CanExecuteScale)
 		);
 
 	ActionList.MapAction(
 		Commands.MirrorActorY,
 		FExecuteAction::CreateStatic( &FLevelEditorActionCallbacks::ExecuteExecCommand, FString( TEXT("ELEMENT MIRROR Y=-1") ) ),
-		FCanExecuteAction::CreateStatic(&FLevelEditorActionCallbacks::ElementSelected_CanExecuteMove)
+		FCanExecuteAction::CreateStatic(&FLevelEditorActionCallbacks::ElementSelected_CanExecuteScale)
 		);
 
 	ActionList.MapAction(
 		Commands.MirrorActorZ,
 		FExecuteAction::CreateStatic( &FLevelEditorActionCallbacks::ExecuteExecCommand, FString( TEXT("ELEMENT MIRROR Z=-1") ) ),
-		FCanExecuteAction::CreateStatic(&FLevelEditorActionCallbacks::ElementSelected_CanExecuteMove)
+		FCanExecuteAction::CreateStatic(&FLevelEditorActionCallbacks::ElementSelected_CanExecuteScale)
 		);
 
 	ActionList.MapAction(
diff --git a/Engine/Source/Editor/LevelEditor/Private/LevelEditorActions.cpp b/Engine/Source/Editor/LevelEditor/Private/LevelEditorActions.cpp
--- a/Engine/Source/Editor/LevelEditor/Private/LevelEditorActions.cpp
+++ b/Engine/Source/Editor/LevelEditor/Private/LevelEditorActions.cpp
@@ -108,6 +108,7 @@
 #include "ISettingsModule.h"
 #include "PlatformInfo.h"
 #include "Misc/CoreMisc.h"
+#include "Misc/AxisDisplayInfo.h"
 #include "Misc/ScopeExit.h"
 #include "Dialogs/DlgPickPath.h"
 #include "AssetToolsModule.h"
@@ -3759,9 +3760,10 @@ bool FLevelEditorActionCallbacks::ElementsSelected_CanExecute()
 bool FLevelEditorActionCallbacks::ElementSelected_CanExecuteMove()
 {
 	// TODO: Ideally this would come from some level editor context
-	if (GCurrentLevelEditingViewportClient)
+	if (const UTypedElementSelectionSet* SelectionSet = GEditor->GetSelectedActors()->GetElementSelectionSet())
 	{
-		return GCurrentLevelEditingViewportClient->GetElementsToManipulate(false)->Num() > 0;
+		FTypedElementListRef NormalizedElements = UEditorElementSubsystem::GetEditorNormalizedSelectionSet(*SelectionSet);
+		return UEditorElementSubsystem::GetEditorManipulableElements(NormalizedElements)->Num() > 0;
 	}
 
 	return false;
@@ -3770,9 +3772,34 @@ bool FLevelEditorActionCallbacks::ElementSelected_CanExecuteMove()
 bool FLevelEditorActionCallbacks::ElementsSelected_CanExecuteMove()
 {
 	// TODO: Ideally this would come from some level editor context
-	if (GCurrentLevelEditingViewportClient)
+	if (const UTypedElementSelectionSet* SelectionSet = GEditor->GetSelectedActors()->GetElementSelectionSet())
+	{
+		FTypedElementListRef NormalizedElements = UEditorElementSubsystem::GetEditorNormalizedSelectionSet(*SelectionSet);
+		return UEditorElementSubsystem::GetEditorManipulableElements(NormalizedElements)->Num() > 1;
+	}
+
+	return false;
+}
+
+bool FLevelEditorActionCallbacks::ElementSelected_CanExecuteScale()
+{
+	// TODO: Ideally this would come from some level editor context
+	if (const UTypedElementSelectionSet* SelectionSet = GEditor->GetSelectedActors()->GetElementSelectionSet())
 	{
-		return GCurrentLevelEditingViewportClient->GetElementsToManipulate(false)->Num() > 1;
+		FTypedElementListRef NormalizedElements = UEditorElementSubsystem::GetEditorNormalizedSelectionSet(*SelectionSet);
+		return UEditorElementSubsystem::GetEditorManipulableElements(NormalizedElements, UE::Widget::WM_Scale)->Num() > 0;
+	}
+
+	return false;
+}
+
+bool FLevelEditorActionCallbacks::ElementsSelected_CanExecuteScale()
+{
+	// TODO: Ideally this would come from some level editor context
+	if (const UTypedElementSelectionSet* SelectionSet = GEditor->GetSelectedActors()->GetElementSelectionSet())
+	{
+		FTypedElementListRef NormalizedElements = UEditorElementSubsystem::GetEditorNormalizedSelectionSet(*SelectionSet);
+		return UEditorElementSubsystem::GetEditorManipulableElements(NormalizedElements, UE::Widget::WM_Scale)->Num() > 1;
 	}
 
 	return false;
@@ -4158,9 +4185,20 @@ void FLevelEditorCommands::RegisterCommands()
 	UI_COMMAND( SnapBottomCenterBoundsToActor, "Snap Bottom Center Bounds to Actor", "Snaps the actor to another actor at its bottom center bounds", EUserInterfaceActionType::Button, FInputChord() );
 	UI_COMMAND( AlignBottomCenterBoundsToActor, "Align Bottom Center Bounds to Actor", "Aligns the actor with another actor at its bottom center bounds", EUserInterfaceActionType::Button, FInputChord() );
 	UI_COMMAND( DeltaTransformToActors, "Delta Transform", "Apply Delta Transform to selected elements", EUserInterfaceActionType::Button, FInputChord() );
-	UI_COMMAND( MirrorActorX, "Mirror X", "Mirrors the element along the X axis", EUserInterfaceActionType::Button, FInputChord() );
-	UI_COMMAND( MirrorActorY, "Mirror Y", "Mirrors the element along the Y axis", EUserInterfaceActionType::Button, FInputChord() );
-	UI_COMMAND( MirrorActorZ, "Mirror Z", "Mirrors the element along the Z axis", EUserInterfaceActionType::Button, FInputChord() );
+
+	#define AXIS_UI_COMMAND(CommandName, Axis) \
+	CommandName = FUICommandInfoDecl(this->AsShared(), \
+		TEXT(#CommandName), \
+		FText::Format(LOCTEXT("MirrorAxisLabel", "Mirror {0} Axis"), AxisDisplayInfo::GetAxisDisplayName(Axis)), \
+		FText::Format(LOCTEXT("MirrorAxisTooltip", "Mirrors the element along the {0} axis"), AxisDisplayInfo::GetAxisDisplayName(Axis))) \
+		.UserInterfaceType(EUserInterfaceActionType::Button) \
+		.DefaultChord(FInputChord());
+
+	AXIS_UI_COMMAND(MirrorActorX, EAxisList::Forward)
+	AXIS_UI_COMMAND(MirrorActorY, EAxisList::Left)
+	AXIS_UI_COMMAND(MirrorActorZ, EAxisList::Up)
+	#undef AXIS_UI_COMMAND
+
 	UI_COMMAND( LockActorMovement, "Lock Actor Movement", "Locks the actor so it cannot be moved", EUserInterfaceActionType::ToggleButton, FInputChord() );
 	UI_COMMAND( DetachFromParent, "Detach", "Detach the actor from its parent", EUserInterfaceActionType::Button, FInputChord() );
 	UI_COMMAND( AttachSelectedActors, "Attach Selected Actors", "Attach the selected actors to the last selected actor", EUserInterfaceActionType::Button, FInputChord(EModifierKey::Alt, EKeys::B) );
diff --git a/Engine/Source/Editor/LevelEditor/Public/LevelEditorActions.h b/Engine/Source/Editor/LevelEditor/Public/LevelEditorActions.h
--- a/Engine/Source/Editor/LevelEditor/Public/LevelEditorActions.h
+++ b/Engine/Source/Editor/LevelEditor/Public/LevelEditorActions.h
@@ -1533,6 +1533,18 @@ public:
 	 */
 	static bool ElementsSelected_CanExecuteMove();
 
+	/**
+	 * Checks to see if at least one element is selected that can be scaled
+	 *	@return true if it can execute.
+	 */
+	static bool ElementSelected_CanExecuteScale();
+
+	/**
+	 * Checks to see if multiple elements are selected that can be scaled
+	 *	@return true if it can execute.
+	 */
+	static bool ElementsSelected_CanExecuteScale();
+
 	/** Called when 'Open Merge Actor' is clicked */
 	static void OpenMergeActor_Clicked();
 
diff --git a/Engine/Source/Editor/UnrealEd/Private/LevelEditorViewport.cpp b/Engine/Source/Editor/UnrealEd/Private/LevelEditorViewport.cpp
--- a/Engine/Source/Editor/UnrealEd/Private/LevelEditorViewport.cpp
+++ b/Engine/Source/Editor/UnrealEd/Private/LevelEditorViewport.cpp
@@ -114,6 +114,7 @@
 #include "Rendering/StaticLightingSystemInterface.h"
 #include "TextureResource.h"
 #include "Subsystems/PlacementSubsystem.h"
+#include "Subsystems/EditorElementSubsystem.h"
 
 DEFINE_LOG_CATEGORY(LogEditorViewport);
 
@@ -4240,41 +4241,8 @@ void FLevelEditorViewportClient::CacheElementsToManipulate(const bool bForceRefr
 
 	if (!bHasCachedElementsToManipulate)
 	{
-		const FTypedElementSelectionNormalizationOptions NormalizationOptions = FTypedElementSelectionNormalizationOptions()
-			.SetExpandGroups(true)
-			.SetFollowAttachment(true);
-
-		const UTypedElementSelectionSet* SelectionSet = GetSelectionSet();
-		SelectionSet->GetNormalizedSelection(NormalizationOptions, CachedElementsToManipulate);
-
-		const UE::Widget::EWidgetMode WidgetMode = GetWidgetMode();
-
-		// Remove any elements that cannot be moved
-		CachedElementsToManipulate->RemoveAll<ITypedElementWorldInterface>([this, WidgetMode](const TTypedElement<ITypedElementWorldInterface>& InWorldElement)
-		{
-			if (!InWorldElement.CanMoveElement(bIsSimulateInEditorViewport ? ETypedElementWorldType::Game : ETypedElementWorldType::Editor))
-			{
-				return true;
-			}
-
-			if (WidgetMode == UE::Widget::WM_Scale && !InWorldElement.CanScaleElement())
-			{
-				return true;
-			}
-
-			// This element must belong to the current viewport world
-			if (GEditor->PlayWorld)
-			{
-				const UWorld* CurrentWorld = InWorldElement.GetOwnerWorld();
-				const UWorld* RequiredWorld = bIsSimulateInEditorViewport ? GEditor->PlayWorld : GEditor->EditorWorld;
-				if (CurrentWorld != RequiredWorld)
-				{
-					return true;
-				}
-			}
-
-			return false;
-		});
+		CachedElementsToManipulate = UEditorElementSubsystem::GetEditorNormalizedSelectionSet(*GetSelectionSet());
+		UEditorElementSubsystem::GetEditorManipulableElements(CachedElementsToManipulate, GetWidgetMode(), bIsSimulateInEditorViewport ? GEditor->PlayWorld : GEditor->EditorWorld);
 
 		bHasCachedElementsToManipulate = true;
 	}
diff --git a/Engine/Source/Editor/UnrealEd/Public/LevelEditorViewport.h b/Engine/Source/Editor/UnrealEd/Public/LevelEditorViewport.h
--- a/Engine/Source/Editor/UnrealEd/Public/LevelEditorViewport.h
+++ b/Engine/Source/Editor/UnrealEd/Public/LevelEditorViewport.h
@@ -342,6 +342,7 @@ public:
 	 * Get the elements (from the current selection set) that this viewport can manipulate (eg, via the transform gizmo).
 	 */
 	UNREALED_API FTypedElementListConstRef GetElementsToManipulate(const bool bForceRefresh = false);
+	UNREALED_API void NormalizeElementsByWidgetMode(FTypedElementListRef& OutElementList, UE::Widget::EWidgetMode InWidgetMode) const;
 
 	UNREALED_API virtual void SetIsSimulateInEditorViewport( bool bInIsSimulateInEditorViewport ) override;
 
diff --git a/Engine/Source/Runtime/CoreUObject/Public/UObject/NoExportTypes.h b/Engine/Source/Runtime/CoreUObject/Public/UObject/NoExportTypes.h
--- a/Engine/Source/Runtime/CoreUObject/Public/UObject/NoExportTypes.h
+++ b/Engine/Source/Runtime/CoreUObject/Public/UObject/NoExportTypes.h
@@ -138,11 +138,11 @@ namespace EAxisList
 	enum Type : int
 	{
 		None = 0,
-		X = 1,
-		Y = 2,
-		Z = 4,
+		X = 1 << 0,
+		Y = 1 << 1,
+		Z = 1 << 2,
 
-		Screen = 8,
+		Screen = 1 << 3,
 		XY = X | Y,
 		XZ = X | Z,
 		YZ = Y | Z,
@@ -154,6 +154,15 @@ namespace EAxisList
 
 		/** alias over Screen since it isn't used when the 2d translate rotate widget is being used */
 		Rotate2D = Screen,
+
+		Left = 1 << 4,
+		Up = 1 << 5,
+		Forward = 1 << 6,
+
+		LU = Left | Up,
+		LF = Left | Forward,
+		UF = Up | Forward,
+		LeftUpForward = Left | Up | Forward,
 	};
 }
 
diff --git a/Engine/Source/Runtime/Slate/Public/Widgets/Input/SNumericEntryBox.h b/Engine/Source/Runtime/Slate/Public/Widgets/Input/SNumericEntryBox.h
--- a/Engine/Source/Runtime/Slate/Public/Widgets/Input/SNumericEntryBox.h
+++ b/Engine/Source/Runtime/Slate/Public/Widgets/Input/SNumericEntryBox.h
@@ -276,6 +276,7 @@ public:
 				.MaxSliderValue(InArgs._MaxSliderValue)
 				.MaxValue(InArgs._MaxValue)
 				.MinValue(InArgs._MinValue)
+				.ContextMenuExtender(InArgs._ContextMenuExtender)
 				.SliderExponent(InArgs._SliderExponent)
 				.SliderExponentNeutralValue(InArgs._SliderExponentNeutralValue)
 				.EnableWheel(InArgs._AllowWheel)
