Fix to stop the Placement Module to be loaded when GIsEditor is null (like in server builds)

diff --git a/Engine/Source/Editor/LevelEditor/Private/LevelEditorOutlinerSettings.cpp b/Engine/Source/Editor/LevelEditor/Private/LevelEditorOutlinerSettings.cpp
--- a/Engine/Source/Editor/LevelEditor/Private/LevelEditorOutlinerSettings.cpp
+++ b/Engine/Source/Editor/LevelEditor/Private/LevelEditorOutlinerSettings.cpp
@@ -47,6 +47,11 @@ FLevelEditorOutlinerSettings::~FLevelEditorOutlinerSettings()
 
 void FLevelEditorOutlinerSettings::Initialize()
 {
+	if (!GIsEditor) // to avoid crashes in server builds as CreateDefaultFilters will force loading the PlacementMode Module
+	{
+		return;
+	} 
+	
 	FDelayedAutoRegisterHelper EndOfEngineInit(EDelayedRegisterRunPhase::EndOfEngineInit, [WeakThis = AsWeak()]()
 	{
 		if (const TSharedPtr<FLevelEditorOutlinerSettings> This = WeakThis.Pin())
