The "Settings" menu in the level editor is now hidden by default. All contents have been migrated elsewhere.
# Set "LevelEditorToolbarSettings 1" to re-enable the menu.

[RN] minor creator/ui

diff --git a/Engine/Source/Editor/LevelEditor/Private/LevelEditorToolBar.cpp b/Engine/Source/Editor/LevelEditor/Private/LevelEditorToolBar.cpp
--- a/Engine/Source/Editor/LevelEditor/Private/LevelEditorToolBar.cpp
+++ b/Engine/Source/Editor/LevelEditor/Private/LevelEditorToolBar.cpp
@@ -68,7 +68,7 @@
 namespace UE::Editor::Private
 {
 
-int32 CVarLevelEditorToolbarSettingsValue = 1;
+int32 CVarLevelEditorToolbarSettingsValue = 0;
 
 }
 
