Add Viewport option in Editor to turn On/Off the Constrained Aspect Ratio in Preview

diff --git a/Engine/Source/Editor/LevelEditor/Private/LevelViewportActions.cpp b/Engine/Source/Editor/LevelEditor/Private/LevelViewportActions.cpp
--- a/Engine/Source/Editor/LevelEditor/Private/LevelViewportActions.cpp
+++ b/Engine/Source/Editor/LevelEditor/Private/LevelViewportActions.cpp
@@ -86,6 +86,7 @@ void FLevelViewportCommands::RegisterCommands()
 	UI_COMMAND( ApplyMaterialToActor, "Apply Material", "Attempts to apply a dropped material to this object", EUserInterfaceActionType::Button, FInputChord() );
 
 	UI_COMMAND( ToggleCinematicPreview, "Allow Cinematic Control", "If enabled, allows cinematic (Sequencer) previews to play in this viewport", EUserInterfaceActionType::ToggleButton, FInputChord() );
+	UI_COMMAND( ToggleAllowConstrainedAspectRatioInPreview, "Constrained Aspect Ratio Preview", "Allows this viewport to have a Constrained Aspect Ration in Preview", EUserInterfaceActionType::ToggleButton, FInputChord() );
 
 	UI_COMMAND( FindInLevelScriptBlueprint, "Find In Level Script", "Finds references of a selected actor in the level script blueprint", EUserInterfaceActionType::Button, FInputChord(EModifierKey::Control, EKeys::K) );
 	UI_COMMAND(
diff --git a/Engine/Source/Editor/LevelEditor/Private/SLevelViewport.cpp b/Engine/Source/Editor/LevelEditor/Private/SLevelViewport.cpp
--- a/Engine/Source/Editor/LevelEditor/Private/SLevelViewport.cpp
+++ b/Engine/Source/Editor/LevelEditor/Private/SLevelViewport.cpp
@@ -1576,6 +1576,14 @@ void SLevelViewport::BindOptionCommands( FUICommandList& OutCommandList )
 		FIsActionChecked::CreateSP( this, &SLevelViewport::AllowsCinematicPreview )
 		);
 
+	OutCommandList.MapAction(
+		ViewportActions.ToggleAllowConstrainedAspectRatioInPreview,
+		FExecuteAction::CreateLambda([]() { GEditor->ToggleAllowConstrainedAspectRatioInPreview(); }),
+		FCanExecuteAction(),
+		FIsActionChecked::CreateLambda([]() { return GEditor->GetAllowConstrainedAspectRatioInPreview(); }),
+		FIsActionButtonVisible::CreateLambda([]() { return GEditor->IsFeatureLevelPreviewActive() && (GEditor->PreviewPlatform.GetConstrainedAspectRatio() != 0.0f); })
+	);
+
 	IAssetRegistry & AssetRegistry = FModuleManager::LoadModuleChecked<FAssetRegistryModule>("AssetRegistry").Get();
 	TArray<FTopLevelAssetPath> ClassNames;
 	TSet<FTopLevelAssetPath> DerivedClassNames;
diff --git a/Engine/Source/Editor/LevelEditor/Private/SLevelViewportToolBar.cpp b/Engine/Source/Editor/LevelEditor/Private/SLevelViewportToolBar.cpp
--- a/Engine/Source/Editor/LevelEditor/Private/SLevelViewportToolBar.cpp
+++ b/Engine/Source/Editor/LevelEditor/Private/SLevelViewportToolBar.cpp
@@ -494,7 +494,7 @@ void SLevelViewportToolBar::FillOptionsMenu(UToolMenu* Menu)
 				Section.AddEntry(UE::LevelEditor::CreateShowStatsSubmenu());
 			}
 #endif
-
+			Section.AddMenuEntry(LevelViewportActions.ToggleAllowConstrainedAspectRatioInPreview);
 			Section.AddMenuEntry(LevelViewportActions.ToggleViewportToolbar);
 
 			if( bIsPerspective )
diff --git a/Engine/Source/Editor/LevelEditor/Public/LevelViewportActions.h b/Engine/Source/Editor/LevelEditor/Public/LevelViewportActions.h
--- a/Engine/Source/Editor/LevelEditor/Public/LevelViewportActions.h
+++ b/Engine/Source/Editor/LevelEditor/Public/LevelViewportActions.h
@@ -113,6 +113,9 @@ PRAGMA_ENABLE_DEPRECATION_WARNINGS
 	/** Allows this viewport to be used for previewing cinematic animations */ 
 	TSharedPtr< FUICommandInfo > ToggleCinematicPreview;
 
+	/** Allows this viewport to have a Constrained Aspect Ration in Preview */
+	TSharedPtr< FUICommandInfo > ToggleAllowConstrainedAspectRatioInPreview;
+
 	/** Finds instances of selected object in level script. */ 
 	TSharedPtr< FUICommandInfo > FindInLevelScriptBlueprint;
 
diff --git a/Engine/Source/Editor/UnrealEd/Classes/Editor/EditorEngine.h b/Engine/Source/Editor/UnrealEd/Classes/Editor/EditorEngine.h
--- a/Engine/Source/Editor/UnrealEd/Classes/Editor/EditorEngine.h
+++ b/Engine/Source/Editor/UnrealEd/Classes/Editor/EditorEngine.h
@@ -2300,6 +2300,16 @@ public:
 	 */
 	UNREALED_API void GridSizeDecrement();
 
+	/** 
+	 * Toggle on/off the ability of the Editor Viewport to apply a Constrained Aspect Ratio in the preview.
+	 */
+	UNREALED_API void ToggleAllowConstrainedAspectRatioInPreview();
+	
+	/** 
+	 * @return  returns if the Editor allows Constrained Aspect Ratio in the preview
+	 */
+	UNREALED_API bool GetAllowConstrainedAspectRatioInPreview();
+
 	/** 
 	 * Accesses the array of snap grid position options
 	 */
@@ -2840,6 +2850,9 @@ private:
 	/** True if we should disable actor/component modification on delta movement */
 	bool bDisableDeltaModification;
 
+	/** True if we allow Constrained Aspect Ratio in the Preview */
+	bool bAllowConstrainedAspectRatioInPreview;
+
 	/** List of editors who want to receive undo/redo events */
 	TSet< class FEditorUndoClient* > UndoClients;
 
diff --git a/Engine/Source/Editor/UnrealEd/Private/EditorConstraints.cpp b/Engine/Source/Editor/UnrealEd/Private/EditorConstraints.cpp
--- a/Engine/Source/Editor/UnrealEd/Private/EditorConstraints.cpp
+++ b/Engine/Source/Editor/UnrealEd/Private/EditorConstraints.cpp
@@ -64,6 +64,17 @@ void UEditorEngine::GridSizeDecrement( )
 	SetGridSize(ViewportSettings->CurrentPosGridSize - 1);
 }
 
+void UEditorEngine::ToggleAllowConstrainedAspectRatioInPreview()
+{
+	bAllowConstrainedAspectRatioInPreview ^= true;
+	FEditorSupportDelegates::RedrawAllViewports.Broadcast();
+}
+
+bool UEditorEngine::GetAllowConstrainedAspectRatioInPreview()
+{
+	return bAllowConstrainedAspectRatioInPreview;
+}
+
 const TArray<float>& UEditorEngine::GetCurrentPositionGridArray() const
 {
 	const ULevelEditorViewportSettings* ViewportSettings = GetDefault<ULevelEditorViewportSettings>();
diff --git a/Engine/Source/Editor/UnrealEd/Private/EditorEngine.cpp b/Engine/Source/Editor/UnrealEd/Private/EditorEngine.cpp
--- a/Engine/Source/Editor/UnrealEd/Private/EditorEngine.cpp
+++ b/Engine/Source/Editor/UnrealEd/Private/EditorEngine.cpp
@@ -576,6 +576,7 @@ UEditorEngine::UEditorEngine(const FObjectInitializer& ObjectInitializer)
 	CurrentPlayWorldDestination = -1;
 	bIsRebuildingAlteredBSP = false;
 	bDisableDeltaModification = false;
+	bAllowConstrainedAspectRatioInPreview = true;
 	bAllowMultiplePIEWorlds = true;
 	bIsEndingPlay = false;
 	DefaultWorldFeatureLevel = GMaxRHIFeatureLevel;
diff --git a/Engine/Source/Editor/UnrealEd/Private/EditorViewportClient.cpp b/Engine/Source/Editor/UnrealEd/Private/EditorViewportClient.cpp
--- a/Engine/Source/Editor/UnrealEd/Private/EditorViewportClient.cpp
+++ b/Engine/Source/Editor/UnrealEd/Private/EditorViewportClient.cpp
@@ -1022,7 +1022,7 @@ FSceneView* FEditorViewportClient::CalcSceneView(FSceneViewFamily* ViewFamily, c
 	TimeForForceRedraw = 0.0;
 
 	bool bConstrainAspectRatio = bUseControllingActorViewInfo && ControllingActorViewInfo.bConstrainAspectRatio;
-	if (GEditor->IsFeatureLevelPreviewActive())
+	if (GEditor->IsFeatureLevelPreviewActive() && GEditor->GetAllowConstrainedAspectRatioInPreview())
 	{
 		if (GEditor->PreviewPlatform.GetConstrainedAspectRatio() != 0.0f)
 		{
diff --git a/Engine/Source/Runtime/Engine/Private/LocalPlayer.cpp b/Engine/Source/Runtime/Engine/Private/LocalPlayer.cpp
--- a/Engine/Source/Runtime/Engine/Private/LocalPlayer.cpp
+++ b/Engine/Source/Runtime/Engine/Private/LocalPlayer.cpp
@@ -1224,7 +1224,7 @@ bool ULocalPlayer::GetProjectionData(FViewport* Viewport, FSceneViewProjectionDa
 	GetViewPoint(/*out*/ ViewInfo);
 
 #if WITH_EDITOR
-	if (GEditor != nullptr && GEditor->IsFeatureLevelPreviewActive())
+	if (GEditor != nullptr && GEditor->IsFeatureLevelPreviewActive() && GEditor->GetAllowConstrainedAspectRatioInPreview())
 	{
 		if (GEditor->PreviewPlatform.GetConstrainedAspectRatio() != 0.0f)
 		{
