Fix typo Ration->Ratio

diff --git a/Engine/Source/Editor/LevelEditor/Private/LevelViewportActions.cpp b/Engine/Source/Editor/LevelEditor/Private/LevelViewportActions.cpp
--- a/Engine/Source/Editor/LevelEditor/Private/LevelViewportActions.cpp
+++ b/Engine/Source/Editor/LevelEditor/Private/LevelViewportActions.cpp
@@ -86,7 +86,7 @@ void FLevelViewportCommands::RegisterCommands()
 	UI_COMMAND( ApplyMaterialToActor, "Apply Material", "Attempts to apply a dropped material to this object", EUserInterfaceActionType::Button, FInputChord() );
 
 	UI_COMMAND( ToggleCinematicPreview, "Allow Cinematic Control", "If enabled, allows cinematic (Sequencer) previews to play in this viewport", EUserInterfaceActionType::ToggleButton, FInputChord() );
-	UI_COMMAND( ToggleAllowConstrainedAspectRatioInPreview, "Constrained Aspect Ratio Preview", "Allows this viewport to have a Constrained Aspect Ration in Preview", EUserInterfaceActionType::ToggleButton, FInputChord() );
+	UI_COMMAND( ToggleAllowConstrainedAspectRatioInPreview, "Constrained Aspect Ratio Preview", "Allows this viewport to have a Constrained Aspect Ratio in Preview", EUserInterfaceActionType::ToggleButton, FInputChord() );
 
 	UI_COMMAND( FindInLevelScriptBlueprint, "Find In Level Script", "Finds references of a selected actor in the level script blueprint", EUserInterfaceActionType::Button, FInputChord(EModifierKey::Control, EKeys::K) );
 	UI_COMMAND(
diff --git a/Engine/Source/Editor/LevelEditor/Public/LevelViewportActions.h b/Engine/Source/Editor/LevelEditor/Public/LevelViewportActions.h
--- a/Engine/Source/Editor/LevelEditor/Public/LevelViewportActions.h
+++ b/Engine/Source/Editor/LevelEditor/Public/LevelViewportActions.h
@@ -113,7 +113,7 @@ PRAGMA_ENABLE_DEPRECATION_WARNINGS
 	/** Allows this viewport to be used for previewing cinematic animations */ 
 	TSharedPtr< FUICommandInfo > ToggleCinematicPreview;
 
-	/** Allows this viewport to have a Constrained Aspect Ration in Preview */
+	/** Allows this viewport to have a Constrained Aspect Ratio in Preview */
 	TSharedPtr< FUICommandInfo > ToggleAllowConstrainedAspectRatioInPreview;
 
 	/** Finds instances of selected object in level script. */ 
