Add an in-PIE viewport toolbar for the level editor with preview & scalability settings and in-game volume slider.
- Disable the old "Click for Mouse Control" and "Shift+F1 for Mouse Cursor" pop-up when the new toolbar is active.
- Move "Shift+F1 for Mouse Cursor" text to the top-left of the PIE toolbar.
- Cache the PIE UGameViewportClient so the PIE toolbar volume slider can function correctly.
- Remove unecessary dynamic menu nesting within UE::UnrealEd::CreateScalabilitySubmenu().

diff --git a/Engine/Source/Editor/EditorStyle/Private/StarshipStyle.cpp b/Engine/Source/Editor/EditorStyle/Private/StarshipStyle.cpp
--- a/Engine/Source/Editor/EditorStyle/Private/StarshipStyle.cpp
+++ b/Engine/Source/Editor/EditorStyle/Private/StarshipStyle.cpp
@@ -4863,6 +4863,7 @@ void FStarshipEditorStyle::FStyle::SetupLevelEditorStyle()
 		Set("LevelViewport.ToggleAllowConstrainedAspectRatioInPreview", new IMAGE_BRUSH_SVG("Starship/Common/ConstrainedAspectRatio_16", Icon16x16));
 		Set("LevelViewport.ToggleCameraShakePreview", new IMAGE_BRUSH_SVG("Starship/Common/CameraShake", Icon16x16));
 		Set("LevelViewport.AdvancedSettings", new CORE_IMAGE_BRUSH_SVG("Starship/Common/settings", Icon16x16));
+		Set("LevelViewport.PlaySettings", new CORE_IMAGE_BRUSH_SVG("Starship/Common/settings", Icon16x16));
 
 		Set("EditorViewport.ActiveBorderColor", FStyleColors::Primary);
 
diff --git a/Engine/Source/Editor/LevelEditor/Private/LevelViewportActions.cpp b/Engine/Source/Editor/LevelEditor/Private/LevelViewportActions.cpp
--- a/Engine/Source/Editor/LevelEditor/Private/LevelViewportActions.cpp
+++ b/Engine/Source/Editor/LevelEditor/Private/LevelViewportActions.cpp
@@ -96,6 +96,13 @@ void FLevelViewportCommands::RegisterCommands()
 		EUserInterfaceActionType::Button,
 		FInputChord()
 	);
+	UI_COMMAND(
+		PlaySettings,
+		"Open Play Preferences...",
+		"Opens the play settings",
+		EUserInterfaceActionType::Button,
+		FInputChord()
+	);
 
 	// Generate a command for each volume class
 	{
diff --git a/Engine/Source/Editor/LevelEditor/Private/SLevelViewport.cpp b/Engine/Source/Editor/LevelEditor/Private/SLevelViewport.cpp
--- a/Engine/Source/Editor/LevelEditor/Private/SLevelViewport.cpp
+++ b/Engine/Source/Editor/LevelEditor/Private/SLevelViewport.cpp
@@ -25,6 +25,7 @@
 #include "Widgets/Layout/SSpacer.h"
 #include "Widgets/Images/SImage.h"
 #include "Widgets/Layout/SBox.h"
+#include "Widgets/Layout/SWidgetSwitcher.h"
 #include "Widgets/Input/SButton.h"
 #include "Widgets/Views/SHeaderRow.h"
 #include "Framework/Docking/LayoutService.h"
@@ -1267,8 +1268,11 @@ void SLevelViewport::Tick( const FGeometry& AllottedGeometry, const double InCur
 	if(IsPlayInEditorViewportActive() && bPIEHasFocus != ActiveViewport->HasMouseCapture())
 	{
 		bPIEHasFocus = ActiveViewport->HasMouseCapture();
-		PIEOverlayAnim = FCurveSequence(0.0f, SLevelViewportPIEAnimation::MouseControlLabelFadeout, ECurveEaseFunction::CubicInOut);
-		PIEOverlayAnim.Play(this->AsShared());
+		if (!bPIEHasFocus || !UE::UnrealEd::ShowNewViewportToolbars())
+		{
+			PIEOverlayAnim = FCurveSequence(0.0f, SLevelViewportPIEAnimation::MouseControlLabelFadeout, ECurveEaseFunction::CubicInOut);
+			PIEOverlayAnim.Play(this->AsShared());
+		}
 	}
 
 	// Update actor preview viewports, if we have any
@@ -1547,6 +1551,9 @@ void SLevelViewport::BindOptionCommands( FUICommandList& OutCommandList )
 	OutCommandList.MapAction( 
 		ViewportActions.AdvancedSettings,
 		FExecuteAction::CreateSP( this, &SLevelViewport::OnAdvancedSettings ) );
+	OutCommandList.MapAction(
+		ViewportActions.PlaySettings,
+		FExecuteAction::CreateSP( this, &SLevelViewport::OnPlaySettings ) );
 
 	OutCommandList.MapAction( 
 		ViewportActions.ToggleMaximize,
@@ -2244,19 +2251,100 @@ TSharedPtr<SWidget> SLevelViewport::BuildViewportToolbar()
 	}
 
 	// clang-format off
-	const TSharedRef<SWidget> NewViewportToolbar = SNew(SBox)
-	.Visibility_Lambda(
-		[this]() -> EVisibility
-		{
-			return GetToolBarVisibility();
-		}
-	)
+	return SNew(SWidgetSwitcher)
+	.WidgetIndex(this, &SLevelViewport::GetViewportToolbarIndex)
+	+ SWidgetSwitcher::Slot()
 	[
 		UToolMenus::Get()->GenerateWidget(ViewportToolbarName, ViewportToolbarContext)
+	]
+	+ SWidgetSwitcher::Slot()
+	[
+		BuildPIEViewportToolbar()
 	];
 	// clang-format on
+}
+
+TSharedRef<SWidget> SLevelViewport::BuildPIEViewportToolbar()
+{
+	const FName ToolbarName = "LevelEditor.PIEViewportToolbar";
+	
+	if (!UToolMenus::Get()->IsMenuRegistered(ToolbarName))
+	{
+		UToolMenu* ToolbarMenu = UToolMenus::Get()->RegisterMenu(ToolbarName, NAME_None, EMultiBoxType::SlimHorizontalToolBar);
+		
+		ToolbarMenu->StyleName = "ViewportToolbar";
+		ToolbarMenu->bSeparateSections = false;
+		
+		{
+			FToolMenuSection& LeftSection = ToolbarMenu->FindOrAddSection("Left");
+			LeftSection.AddEntry(FToolMenuEntry::InitDynamicEntry(
+				"DynamicPIEDetachHint",
+				FNewToolMenuSectionDelegate::CreateLambda([](FToolMenuSection& Section)
+				{
+					ULevelViewportContext* Context = Section.FindContext<ULevelViewportContext>();
+					if (!Context)
+					{
+						return;
+					}
+					
+					TSharedPtr<SLevelViewport> Viewport = Context->LevelViewport.Pin();
+					
+					Section.AddEntry(FToolMenuEntry::InitWidget(
+						"PIEDetachHint",
+						SNew(SBox)
+						.Padding(FMargin(8.f, 0.f, 0.f, 0.f))
+						.Visibility_Lambda([WeakViewport = Context->LevelViewport]
+						{
+							const ULevelEditorPlaySettings* EditorPlayInSettings = GetDefault<ULevelEditorPlaySettings>();
+							if (!EditorPlayInSettings->ShowMouseControlLabel)
+							{
+								return EVisibility::Collapsed;
+							}
+						
+							if (TSharedPtr<SLevelViewport> Viewport = WeakViewport.Pin())
+							{
+								return Viewport->GetActiveViewport()->HasMouseCapture() ? EVisibility::Visible : EVisibility::Collapsed;
+							}
+							return EVisibility::Collapsed;
+						})
+						[
+							Viewport->BuildMouseCaptureWidget()
+						],
+						FText::GetEmpty()
+					));
+				})
+			));
+		}
+		
+		{
+			FToolMenuSection& RightSection = ToolbarMenu->FindOrAddSection("Right");
+			RightSection.Alignment = EToolMenuSectionAlign::Last;
+			
+			RightSection.AddEntry(UE::LevelEditor::CreatePIEPerformanceAndScalabilitySubmenu());
+			RightSection.AddEntry(UE::LevelEditor::CreatePIESettingsSubmenu());
+		}
+	}
+	
+	FToolMenuContext ToolbarContext;
+	{
+		ToolbarContext.AppendCommandList(GetCommandList());
+		
+		// Add the UnrealEd viewport toolbar context.
+		{
+			UUnrealEdViewportToolbarContext* const ContextObject = NewObject<UUnrealEdViewportToolbarContext>();
+			ContextObject->Viewport = SharedThis(this);
+			ToolbarContext.AddObject(ContextObject);
+		}
 
-	return NewViewportToolbar;
+		// Add the level editor viewport toolbar context.
+		{
+			ULevelViewportContext* const ContextObject = NewObject<ULevelViewportContext>();
+			ContextObject->LevelViewport = SharedThis(this);
+			ToolbarContext.AddObject(ContextObject);
+		}
+	}
+	
+	return UToolMenus::Get()->GenerateWidget(ToolbarName, ToolbarContext);
 }
 
 void SLevelViewport::OnUndo()
@@ -2279,11 +2367,21 @@ bool SLevelViewport::CanExecuteRedo() const
 	return GUnrealEd->Trans->CanRedo() && FSlateApplication::Get().IsNormalExecution();
 }
 
+int32 SLevelViewport::GetViewportToolbarIndex() const
+{
+	return IsPlayInEditorViewportActive() ? 1 : 0;
+}
+
 void SLevelViewport::OnAdvancedSettings()
 {
 	FModuleManager::LoadModuleChecked<ISettingsModule>("Settings").ShowViewer("Editor", "LevelEditor", "Viewport");
 }
 
+void SLevelViewport::OnPlaySettings()
+{
+	FModuleManager::LoadModuleChecked<ISettingsModule>("Settings").ShowViewer("Editor", "LevelEditor", "PlayIn");
+}
+
 void SLevelViewport::OnToggleImmersive()
 {
 	if( ParentLayout.IsValid() ) 
@@ -3201,6 +3299,11 @@ bool SLevelViewport::IsPlayInEditorViewportActive() const
 	return ActiveViewport->IsPlayInEditorViewport();
 }
 
+UGameViewportClient* SLevelViewport::GetPlayClient() const
+{
+	return PlayClient.Get();
+}
+
 void SLevelViewport::OnActorSelectionChanged(const TArray<UObject*>& NewSelection, bool bForceRefresh)
 {
 	// On the first actor selection after entering Game View, enable the selection show flag
@@ -4450,7 +4553,7 @@ PRAGMA_DISABLE_DEPRECATION_WARNINGS
 PRAGMA_ENABLE_DEPRECATION_WARNINGS
 }
 
-void SLevelViewport::StartPlayInEditorSession(UGameViewportClient* PlayClient, const bool bInSimulateInEditor)
+void SLevelViewport::StartPlayInEditorSession(UGameViewportClient* InPlayClient, const bool bInSimulateInEditor)
 {
 	check( !HasPlayInEditorViewport() );
 
@@ -4461,6 +4564,8 @@ void SLevelViewport::StartPlayInEditorSession(UGameViewportClient* PlayClient, c
 	// Save camera settings that may be adversely affected by PIE, so that they may be restored later
 	LevelViewportClient->PrepareCameraForPIE();
 
+	PlayClient = InPlayClient;
+
 	// Here we will swap the editor viewport client out for the client for the play in editor session
 	InactiveViewport = ActiveViewport;
 	// Store the content in the viewport widget (editor tool bar etc) so we can show the game UI content if it has any
@@ -4471,7 +4576,7 @@ void SLevelViewport::StartPlayInEditorSession(UGameViewportClient* PlayClient, c
 	FSlateApplication::Get().ClearKeyboardFocus(EFocusCause::SetDirectly);
 
 	// Attach global play world actions widget to view port
-	ActiveViewport = MakeShareable( new FSceneViewport( PlayClient, ViewportWidget) );
+	ActiveViewport = MakeShareable( new FSceneViewport( InPlayClient, ViewportWidget) );
 	ActiveViewport->SetPlayInEditorViewport( true );
 
 	// Whether to start with the game taking mouse control or leaving it shown in the editor
@@ -4515,7 +4620,7 @@ void SLevelViewport::StartPlayInEditorSession(UGameViewportClient* PlayClient, c
 
 	bPIEHasFocus = ActiveViewport->HasMouseCapture();
 
-	if(EditorPlayInSettings->ShowMouseControlLabel && !GEngine->IsStereoscopic3D( ActiveViewport.Get() ) )
+	if(EditorPlayInSettings->ShowMouseControlLabel && !GEngine->IsStereoscopic3D( ActiveViewport.Get() ) && !UE::UnrealEd::ShowNewViewportToolbars())
 	{
 		ELabelAnchorMode AnchorMode = EditorPlayInSettings->MouseControlLabelPosition.GetValue();
 
@@ -4527,9 +4632,9 @@ void SLevelViewport::StartPlayInEditorSession(UGameViewportClient* PlayClient, c
 	// register for preview feature level change
 	UEditorEngine* Editor = CastChecked<UEditorEngine>(GEngine);
 	
-	PIEPreviewFeatureLevelChangedHandle = Editor->OnPreviewFeatureLevelChanged().AddLambda([PlayClient](ERHIFeatureLevel::Type NewFeatureLevel)
+	PIEPreviewFeatureLevelChangedHandle = Editor->OnPreviewFeatureLevelChanged().AddLambda([InPlayClient](ERHIFeatureLevel::Type NewFeatureLevel)
 		{
-			PlayClient->GetWorld()->ChangeFeatureLevel(NewFeatureLevel);
+			InPlayClient->GetWorld()->ChangeFeatureLevel(NewFeatureLevel);
 		});
 
 }
@@ -4627,37 +4732,42 @@ void SLevelViewport::ShowMouseCaptureLabel(ELabelAnchorMode AnchorMode)
 				.ButtonColorAndOpacity( FSlateColor(FLinearColor::Black) )
 				.ForegroundColor( FLinearColor::White )
 				[
-					SNew( SHorizontalBox )
-					+ SHorizontalBox::Slot()
-					.MaxWidth(32.f)
-					.VAlign(VAlign_Center)
-					.Padding(0.0f, 2.0f, 2.0f, 2.0f)
-					[
-						SNew( SVerticalBox )
-						+ SVerticalBox::Slot()
-						.MaxHeight(16.f)
-						[
-							SNew(SImage)
-							.Image(FAppStyle::GetBrush("LevelViewport.CursorIcon"))
-						]
-					]
-					+ SHorizontalBox::Slot()
-					.AutoWidth()
-					.VAlign(VAlign_Center)
-					.HAlign(HAlign_Center)
-					.Padding(2.0f, 2.0f)
-					[
-						SNew(STextBlock)
-						.Text(this, &SLevelViewport::GetMouseCaptureLabelText)
-						.Font(FCoreStyle::GetDefaultFontStyle("Bold", 9))
-						.ColorAndOpacity(FLinearColor::White)
-					]
+					BuildMouseCaptureWidget()
 				]
 			]
 		];
 	}
 }
 
+TSharedRef<SWidget> SLevelViewport::BuildMouseCaptureWidget()
+{
+	return SNew( SHorizontalBox )
+		+ SHorizontalBox::Slot()
+		.MaxWidth(32.f)
+		.VAlign(VAlign_Center)
+		.Padding(0.0f, 2.0f, 2.0f, 2.0f)
+		[
+			SNew( SVerticalBox )
+			+ SVerticalBox::Slot()
+			.MaxHeight(16.f)
+			[
+				SNew(SImage)
+				.Image(FAppStyle::GetBrush("LevelViewport.CursorIcon"))
+			]
+		]
+		+ SHorizontalBox::Slot()
+		.AutoWidth()
+		.VAlign(VAlign_Center)
+		.HAlign(HAlign_Center)
+		.Padding(2.0f, 2.0f)
+		[
+			SNew(STextBlock)
+			.Text(this, &SLevelViewport::GetMouseCaptureLabelText)
+			.Font(FCoreStyle::GetDefaultFontStyle("Bold", 9))
+			.ColorAndOpacity(FLinearColor::White)
+		];
+}
+
 void SLevelViewport::HideMouseCaptureLabel()
 {
 	if (PIEOverlayBorder.IsValid())
@@ -4723,6 +4833,8 @@ void SLevelViewport::EndPlayInEditorSession()
 
 	// Reset our game layer manager's active to that of the active editor viewport
 	GameLayerManager->SetSceneViewport(ActiveViewport.Get());
+	
+	PlayClient.Reset();
 
 	// Reset the inactive viewport
 	InactiveViewport.Reset();
@@ -4805,7 +4917,7 @@ void SLevelViewport::SwapViewportsForPlayInEditor()
 	ULevelEditorPlaySettings const* EditorPlayInSettings = GetDefault<ULevelEditorPlaySettings>();
 	check(EditorPlayInSettings);
 	
-	if(EditorPlayInSettings->ShowMouseControlLabel && !GEngine->IsStereoscopic3D( ActiveViewport.Get() ) )
+	if(EditorPlayInSettings->ShowMouseControlLabel && !GEngine->IsStereoscopic3D( ActiveViewport.Get() ) && !UE::UnrealEd::ShowNewViewportToolbars())
 	{
 		ELabelAnchorMode AnchorMode = EditorPlayInSettings->MouseControlLabelPosition.GetValue();
 		
@@ -4896,6 +5008,13 @@ UWorld* SLevelViewport::GetWorld() const
 	return ParentLevelEditor.IsValid() ? ParentLevelEditor.Pin()->GetWorld() : nullptr;
 }
 
+FReply SLevelViewport::OnFocusReceived(const FGeometry& MyGeometry, const FFocusEvent& InFocusEvent)
+{
+	// SEditorViewport redirects all clicks to focus the viewport widget.
+	// This means clicking on the toolbar in PIE captures the mouse cursor for the game.
+	return FReply::Unhandled();
+}
+
 void SLevelViewport::ToggleInViewportContextMenu()
 {
 	USelection* ActorSelection = GEditor->GetSelectedActors();
diff --git a/Engine/Source/Editor/LevelEditor/Private/ViewportToolbar/LevelEditorSubmenus.cpp b/Engine/Source/Editor/LevelEditor/Private/ViewportToolbar/LevelEditorSubmenus.cpp
--- a/Engine/Source/Editor/LevelEditor/Private/ViewportToolbar/LevelEditorSubmenus.cpp
+++ b/Engine/Source/Editor/LevelEditor/Private/ViewportToolbar/LevelEditorSubmenus.cpp
@@ -2,11 +2,13 @@
 
 #include "ViewportToolbar/LevelEditorSubmenus.h"
 
+#include "AudioDevice.h"
 #include "Bookmarks/BookmarkUI.h"
 #include "Bookmarks/IBookmarkTypeTools.h"
 #include "Camera/CameraActor.h"
 #include "EditorViewportCommands.h"
 #include "Engine/SceneCapture.h"
+#include "Engine/GameViewportClient.h"
 #include "EngineUtils.h"
 #include "FoliageType.h"
 #include "GameFramework/ActorPrimitiveColorHandler.h"
@@ -2161,6 +2163,40 @@ FToolMenuEntry CreatePerformanceAndScalabilitySubmenu()
 	return Entry;
 }
 
+FToolMenuEntry CreatePIEPerformanceAndScalabilitySubmenu()
+{
+	FToolMenuEntry Entry = FToolMenuEntry::InitSubMenu(
+		"PerformanceAndScalability",
+		LOCTEXT("PerformanceAndScalabilityLabel", "Performance & Scalability"),
+		LOCTEXT("PerformanceAndScalabilityTooltip", "Performance and scalability tools tied to this viewport."),
+		FNewToolMenuDelegate::CreateLambda(
+			[](UToolMenu* Submenu) -> void
+			{
+				{
+					FToolMenuSection& UnnamedSection = Submenu->FindOrAddSection(NAME_None);
+
+					FToolMenuEntry& MenuEntry = UnnamedSection.AddMenuEntry(FLevelViewportCommands::Get().ToggleAllowConstrainedAspectRatioInPreview);
+					MenuEntry.SetShowInToolbarTopLevel(true);
+				}
+
+				{
+					FToolMenuSection& PerformanceAndScalabilitySection = Submenu->FindOrAddSection(
+						"PerformanceAndScalability",
+						LOCTEXT("PerformanceAndScalabilitySectionLabel", "Performance & Scalability")
+					);
+
+					PerformanceAndScalabilitySection.AddEntry(CreateFeatureLevelPreviewSubmenu());
+					PerformanceAndScalabilitySection.AddEntry(UE::UnrealEd::CreateScalabilitySubmenu());
+				}
+			}
+		)
+	);
+	Entry.Icon = FSlateIcon(FAppStyle::GetAppStyleSetName(), "Icons.Scalability");
+	Entry.ToolBarData.LabelOverride = FText();
+	Entry.ToolBarData.ResizeParams.ClippingPriority = 800;
+	return Entry;
+}
+
 void GenerateViewportLayoutsMenu(UToolMenu* InMenu)
 {
 	ULevelViewportContext* const LevelViewportContext = InMenu->FindContext<ULevelViewportContext>();
@@ -2326,6 +2362,78 @@ TSharedRef<SWidget> BuildVolumeControlCustomWidget()
 	// clang-format on
 }
 
+FAudioDeviceHandle GetAudioFromViewport(const TWeakPtr<::SLevelViewport>& WeakViewport)
+{
+	if (TSharedPtr<::SLevelViewport> Viewport = WeakViewport.Pin())
+	{
+		if (UGameViewportClient* Client = Viewport->GetPlayClient())
+		{
+			return Client->GetWorld()->GetAudioDevice();
+		}
+	}
+	return FAudioDeviceHandle();
+}
+
+TSharedRef<SWidget> BuildPIEVolumeControlCustomWidget(const TWeakPtr<::SLevelViewport>& Viewport)
+{
+	// clang-format off
+	return SNew(SHorizontalBox)
+		+ SHorizontalBox::Slot()
+		.FillWidth(0.9f)
+		.MinWidth(100.0f)
+		.Padding(FMargin(2.0f, 0.0f, 0.0f, 0.0f))
+		[
+			SNew(SVolumeControl)
+				.ToolTipText_Lambda([WeakViewport = Viewport]
+				{
+					if (FAudioDeviceHandle Audio = GetAudioFromViewport(WeakViewport))
+					{
+						if (Audio->IsAudioDeviceMuted())
+						{
+							return LOCTEXT("PIEMuted", "Muted");
+						}
+								
+						const float Volume = Audio->GetTransientPrimaryVolume() * 100.0f;
+						return FText::AsNumber(FMath::RoundToInt(Volume));
+					}
+					return LOCTEXT("PIEVolume", "PIE Volume");
+				})
+				.Volume_Lambda([WeakViewport = Viewport]
+				{
+					if (FAudioDeviceHandle Audio = GetAudioFromViewport(WeakViewport))
+					{
+						return Audio->GetTransientPrimaryVolume();
+					}
+					return 0.0f;
+				})
+				.OnVolumeChanged_Lambda([WeakViewport = Viewport](float Volume)
+				{
+					if (FAudioDeviceHandle Audio = GetAudioFromViewport(WeakViewport))
+					{
+						Audio->SetTransientPrimaryVolume(Volume);
+					}
+				})
+				.Muted_Lambda([WeakViewport = Viewport]
+				{
+					if (FAudioDeviceHandle Audio = GetAudioFromViewport(WeakViewport))
+					{
+						return Audio->IsAudioDeviceMuted();
+					}
+					return false;
+				})
+				.OnMuteChanged_Lambda([WeakViewport = Viewport](bool bMuted)
+				{
+					if (FAudioDeviceHandle Audio = GetAudioFromViewport(WeakViewport))
+					{
+						Audio->SetDeviceMuted(bMuted);
+					}
+				})
+		]
+		+ SHorizontalBox::Slot()
+		.FillWidth(0.1f);
+	// clang-format on
+}
+
 FToolMenuEntry CreateSettingsSubmenu()
 {
 	FToolMenuEntry Entry = FToolMenuEntry::InitSubMenu(
@@ -2337,13 +2445,12 @@ FToolMenuEntry CreateSettingsSubmenu()
 			{
 				{
 					{
-						FToolMenuSection& SettingsSection =
-							Submenu->FindOrAddSection("Settings", LOCTEXT("SettingsSectionLabel", "Settings"));
-						SettingsSection.AddEntry(
-							FToolMenuEntry::InitWidget(
-								"Volume", BuildVolumeControlCustomWidget(), LOCTEXT("VolumeControlLabel", "Volume")
-							)
-						);
+						FToolMenuSection& SettingsSection = Submenu->FindOrAddSection("Settings", LOCTEXT("SettingsSectionLabel", "Settings"));
+						SettingsSection.AddEntry(FToolMenuEntry::InitWidget(
+							"Volume",
+							BuildVolumeControlCustomWidget(),
+							LOCTEXT("VolumeControlLabel", "Volume")
+						));
 					}
 					// Mouse Control
 					{
@@ -2396,12 +2503,10 @@ FToolMenuEntry CreateSettingsSubmenu()
 					}
 					// Viewport advanced settings
 					{
-						{
-							FToolMenuSection& UnnamedSection = Submenu->FindOrAddSection(NAME_None, FText());
-							UnnamedSection.AddSeparator(NAME_None);
+						FToolMenuSection& UnnamedSection = Submenu->FindOrAddSection(NAME_None, FText());
+						UnnamedSection.AddSeparator(NAME_None);
 						const FLevelViewportCommands& LevelViewportActions = FLevelViewportCommands::Get();
-							UnnamedSection.AddMenuEntry(LevelViewportActions.AdvancedSettings);
-						}
+						UnnamedSection.AddMenuEntry(LevelViewportActions.AdvancedSettings);
 					}
 
 					// Cascade
@@ -2439,6 +2544,44 @@ FToolMenuEntry CreateSettingsSubmenu()
 	return Entry;
 }
 
+FToolMenuEntry CreatePIESettingsSubmenu()
+{
+	FToolMenuEntry Entry = FToolMenuEntry::InitSubMenu(
+		"Settings",
+		LOCTEXT("SettingsSubmenuLabel", "Settings"),
+		LOCTEXT("SettingsSubmenuTooltip", "Viewport-related settings"),
+		FNewToolMenuDelegate::CreateLambda(
+			[](UToolMenu* Submenu) -> void
+			{
+				{
+					if (ULevelViewportContext* Context = Submenu->FindContext<ULevelViewportContext>())
+					{
+						FToolMenuSection& SettingsSection = Submenu->FindOrAddSection("Settings", LOCTEXT("SettingsSectionLabel", "Settings"));
+						SettingsSection.AddEntry(FToolMenuEntry::InitWidget(
+							"Volume",
+							BuildPIEVolumeControlCustomWidget(Context->LevelViewport),
+							LOCTEXT("VolumeControlLabel", "Volume")
+						));
+					}
+					
+					// Viewport advanced settings
+					{
+						FToolMenuSection& UnnamedSection = Submenu->FindOrAddSection(NAME_None, FText());
+						UnnamedSection.AddSeparator(NAME_None);
+						const FLevelViewportCommands& LevelViewportActions = FLevelViewportCommands::Get();
+						UnnamedSection.AddMenuEntry(LevelViewportActions.PlaySettings);
+					}
+				}
+			}
+		)
+	);
+
+	Entry.Icon = FSlateIcon(FAppStyle::GetAppStyleSetName(), "LevelEditor.GameSettings");
+	Entry.ToolBarData.LabelOverride = FText();
+	Entry.ToolBarData.ResizeParams.ClippingPriority = 1000;
+	return Entry;
+}
+
 FToolMenuEntry CreateViewportSizingSubmenu()
 {
 	FToolMenuEntry Entry = FToolMenuEntry::InitSubMenu(
diff --git a/Engine/Source/Editor/LevelEditor/Private/ViewportToolbar/LevelEditorSubmenus.h b/Engine/Source/Editor/LevelEditor/Private/ViewportToolbar/LevelEditorSubmenus.h
--- a/Engine/Source/Editor/LevelEditor/Private/ViewportToolbar/LevelEditorSubmenus.h
+++ b/Engine/Source/Editor/LevelEditor/Private/ViewportToolbar/LevelEditorSubmenus.h
@@ -43,10 +43,13 @@ FToolMenuEntry CreateShowSubmenu();
 FToolMenuEntry CreateFeatureLevelPreviewSubmenu();
 FToolMenuEntry CreateMaterialQualityLevelSubmenu();
 FToolMenuEntry CreatePerformanceAndScalabilitySubmenu();
+FToolMenuEntry CreatePIEPerformanceAndScalabilitySubmenu();
 
 void GenerateViewportLayoutsMenu(UToolMenu* InMenu);
 TSharedRef<SWidget> BuildVolumeControlCustomWidget();
+TSharedRef<SWidget> BuildPIEVolumeControlCustomWidget(const TWeakPtr<SLevelViewport>& Viewport);
 FToolMenuEntry CreateSettingsSubmenu();
+FToolMenuEntry CreatePIESettingsSubmenu();
 
 FToolMenuEntry CreateViewportSizingSubmenu();
 
diff --git a/Engine/Source/Editor/LevelEditor/Public/LevelViewportActions.h b/Engine/Source/Editor/LevelEditor/Public/LevelViewportActions.h
--- a/Engine/Source/Editor/LevelEditor/Public/LevelViewportActions.h
+++ b/Engine/Source/Editor/LevelEditor/Public/LevelViewportActions.h
@@ -88,6 +88,9 @@ PRAGMA_ENABLE_DEPRECATION_WARNINGS
 
 	/** Opens the advanced viewport settings */
 	TSharedPtr< FUICommandInfo > AdvancedSettings;
+	
+	/** Opens the play settings for PIE */
+	TSharedPtr< FUICommandInfo > PlaySettings;
 
 	/** Toggles game preview in the viewport */
 	TSharedPtr< FUICommandInfo > ToggleGameView;
diff --git a/Engine/Source/Editor/LevelEditor/Public/SLevelViewport.h b/Engine/Source/Editor/LevelEditor/Public/SLevelViewport.h
--- a/Engine/Source/Editor/LevelEditor/Public/SLevelViewport.h
+++ b/Engine/Source/Editor/LevelEditor/Public/SLevelViewport.h
@@ -128,6 +128,7 @@ public:
 	virtual void OnFocusViewportToSelection() override;
 	virtual EVisibility GetTransformToolbarVisibility() const override;
 	virtual UWorld* GetWorld() const override;
+	virtual FReply OnFocusReceived(const FGeometry& MyGeometry, const FFocusEvent& InFocusEvent) override;
 	virtual void ToggleInViewportContextMenu() override;
 
 	virtual void HideInViewportContextMenu() override;
@@ -178,6 +179,11 @@ public:
 	 */
 	bool IsPlayInEditorViewportActive() const;
 
+	/**
+	 * @return The active play client, if any 
+	 */
+	UGameViewportClient* GetPlayClient() const;
+
 	/**
 	 * Called on all viewports, when actor selection changes.
 	 * 
@@ -464,6 +470,7 @@ protected:
 	virtual TSharedRef<FEditorViewportClient> MakeEditorViewportClient() override;
 	virtual TSharedPtr<SWidget> MakeViewportToolbar() override;
 	virtual TSharedPtr<SWidget> BuildViewportToolbar() override;
+	TSharedRef<SWidget> BuildPIEViewportToolbar();
 
 	virtual void OnIncrementPositionGridSize() override;
 	virtual void OnDecrementPositionGridSize() override;
@@ -484,11 +491,21 @@ private:
 	/** Called when a property is changed */
 	void HandleViewportSettingChanged(FName PropertyName);
 
+	/**
+	 * Handles which viewport toolbar is active.
+	 */
+	int32 GetViewportToolbarIndex() const;
+
 	/**
 	 * Called when the advanced settings should be opened.
 	 */
 	void OnAdvancedSettings();
 
+	/**
+	 * Called when the play settings should be opened.
+	 */
+	void OnPlaySettings();
+
 	/**
 	 * Called when immersive mode is toggled by the user
 	 */
@@ -770,6 +787,9 @@ private:
 
 	/** Show the mouse capture label with the specified vertical and horizontal alignment */
 	void ShowMouseCaptureLabel(ELabelAnchorMode AnchorMode);
+	
+	/** Build the mouse capture label widget */
+	TSharedRef<SWidget> BuildMouseCaptureWidget();
 
 	/** Hide the mouse capture label */
 	void HideMouseCaptureLabel();
@@ -838,6 +858,9 @@ private:
 
 	/** Level viewport client */
 	TSharedPtr<FLevelEditorViewportClient> LevelViewportClient;
+	
+	/** The GameViewportClient */
+	TWeakObjectPtr<UGameViewportClient> PlayClient;
 
 	/** The brush to use if this viewport is in debug mode */
 	const FSlateBrush* DebuggingBorder;
diff --git a/Engine/Source/Editor/UnrealEd/Private/ViewportToolbar/UnrealEdViewportToolbar.cpp b/Engine/Source/Editor/UnrealEd/Private/ViewportToolbar/UnrealEdViewportToolbar.cpp
--- a/Engine/Source/Editor/UnrealEd/Private/ViewportToolbar/UnrealEdViewportToolbar.cpp
+++ b/Engine/Source/Editor/UnrealEd/Private/ViewportToolbar/UnrealEdViewportToolbar.cpp
@@ -3693,56 +3693,50 @@ FToolMenuEntry CreateScreenPercentageSubmenu()
 
 FToolMenuEntry CreateScalabilitySubmenu()
 {
-	return FToolMenuEntry::InitDynamicEntry(
+	FToolMenuEntry ScalabilitySubmenu = FToolMenuEntry::InitSubMenu(
 		"Scalability",
-		FNewToolMenuSectionDelegate::CreateLambda(
-			[](FToolMenuSection& InnerSection) -> void
+		LOCTEXT("ScalabilitySubMenu", "Viewport Scalability"),
+		TAttribute<FText>::CreateLambda([]()
+		{
+			if (UE::UnrealEd::IsScalabilityWarningVisible())
 			{
-				FToolMenuEntry& ScalabilitySubmenu = InnerSection.AddSubMenu(
-					"Scalability",
-					LOCTEXT("ScalabilitySubMenu", "Viewport Scalability"),
-					TAttribute<FText>::CreateLambda([]()
-					{
-						if (UE::UnrealEd::IsScalabilityWarningVisible())
-						{
-							return UE::UnrealEd::GetScalabilityWarningTooltip();
-						}
-
-						return LOCTEXT("ScalabilitySubMenu", "Viewport Scalability");
-					}),
-					FNewToolMenuDelegate::CreateLambda(
-						[](UToolMenu* InMenu) -> void
-						{
-							FToolMenuSection& Section = InMenu->FindOrAddSection(NAME_None);
-							Section.AddEntry(FToolMenuEntry::InitWidget(
-								"ScalabilitySettings", SNew(SScalabilitySettings), FText(), true
-							));
-						}
-					),
-					FToolUIAction(),
-					EUserInterfaceActionType::Button,
-					false,
-					FSlateIcon(FAppStyle::GetAppStyleSetName(), "Icons.ViewportScalability")
-				);
+				return UE::UnrealEd::GetScalabilityWarningTooltip();
+			}
 
-				ScalabilitySubmenu.SetShowInToolbarTopLevel(TAttribute<bool>::Create(&UE::UnrealEd::IsScalabilityWarningVisible));
+			return LOCTEXT("ScalabilitySubMenu", "Viewport Scalability");
+		}),
+		FNewToolMenuDelegate::CreateLambda(
+			[](UToolMenu* InMenu) -> void
+			{
+				FToolMenuSection& Section = InMenu->FindOrAddSection(NAME_None);
+				Section.AddEntry(FToolMenuEntry::InitWidget(
+					"ScalabilitySettings", SNew(SScalabilitySettings), FText(), true
+				));
+			}
+		),
+		FToolUIAction(),
+		EUserInterfaceActionType::Button,
+		false,
+		FSlateIcon(FAppStyle::GetAppStyleSetName(), "Icons.ViewportScalability")
+	);
 
-				ScalabilitySubmenu.ToolBarData.ResizeParams.AllowClipping = TAttribute<bool>::CreateLambda(
-					[]() -> bool
-					{
-						return !UE::UnrealEd::IsScalabilityWarningVisible();
-					}
-				);
+	ScalabilitySubmenu.SetShowInToolbarTopLevel(TAttribute<bool>::Create(&UE::UnrealEd::IsScalabilityWarningVisible));
 
-				// We need the warning style when we've raised this entry to the top-level toolbar to draw
-				// attention to it. However, styles are not driven by an attribute, so they cannot be updated after an
-				// FToolMenuEntry has been created. Therefore we always apply the style instead because the style
-				// will not make this entry look different when it appears in a menu, but when it appears in a toolbar
-				// we get the desired warning effect.
-				ScalabilitySubmenu.StyleNameOverride = "ViewportToolbarWarning";
-			}
-		)
+	ScalabilitySubmenu.ToolBarData.ResizeParams.AllowClipping = TAttribute<bool>::CreateLambda(
+		[]() -> bool
+		{
+			return !UE::UnrealEd::IsScalabilityWarningVisible();
+		}
 	);
+
+	// We need the warning style when we've raised this entry to the top-level toolbar to draw
+	// attention to it. However, styles are not driven by an attribute, so they cannot be updated after an
+	// FToolMenuEntry has been created. Therefore we always apply the style instead because the style
+	// will not make this entry look different when it appears in a menu, but when it appears in a toolbar
+	// we get the desired warning effect.
+	ScalabilitySubmenu.StyleNameOverride = "ViewportToolbarWarning";
+	
+	return ScalabilitySubmenu;
 }
 
 FText GetCameraSpeedTooltip()
