[Localization] Concatenation "Activate {Mode} Mode" leads to grammar mistakes in localized languages. Changed to "Activate Mode: {Mode}" as per discussed in the slack thread

diff --git a/Engine/Source/Editor/LevelEditor/Private/LevelEditorModesActions.cpp b/Engine/Source/Editor/LevelEditor/Private/LevelEditorModesActions.cpp
--- a/Engine/Source/Editor/LevelEditor/Private/LevelEditorModesActions.cpp
+++ b/Engine/Source/Editor/LevelEditor/Private/LevelEditorModesActions.cpp
@@ -34,7 +34,7 @@ void FLevelEditorModesCommands::RegisterCommands()
 				ModeName = FText::FromName(Mode.ID);
 			}
 			Args.Add(TEXT("Mode"), ModeName);
-			const FText Tooltip = FText::Format( NSLOCTEXT("LevelEditor", "ModeTooltipF", "Activate {Mode} Mode"), Args );
+			const FText Tooltip = FText::Format( NSLOCTEXT("LevelEditor", "ModeTooltipF", "Activate Mode: {Mode}"), Args );
 
 			FInputChord DefaultKeyBinding;
 			if (Mode.IsVisible() && editorMode < EdModeKeys.Num())
