[Details Panel] Removing old function that causes Details Panels to refresh even though it is no longer needed in the current implementation of spawning Details Panel Windows

diff --git a/Engine/Source/Editor/LevelEditor/Private/SLevelEditor.cpp b/Engine/Source/Editor/LevelEditor/Private/SLevelEditor.cpp
--- a/Engine/Source/Editor/LevelEditor/Private/SLevelEditor.cpp
+++ b/Engine/Source/Editor/LevelEditor/Private/SLevelEditor.cpp
@@ -970,9 +970,7 @@ TSharedRef<SDockTab> SLevelEditor::SpawnLevelEditorTab( const FSpawnTabArgs& Arg
 	}
 	else if( TabIdentifier == TEXT("LevelEditorSelectionDetails") || TabIdentifier == TEXT("LevelEditorSelectionDetails2") || TabIdentifier == TEXT("LevelEditorSelectionDetails3") || TabIdentifier == TEXT("LevelEditorSelectionDetails4") )
 	{
-		TSharedRef<SDockTab> DetailsPanel = SummonDetailsPanel( TabIdentifier );
-		GUnrealEd->UpdateFloatingPropertyWindows();
-		return DetailsPanel;
+		return SummonDetailsPanel( TabIdentifier );
 	}
 	else if (TabIdentifier == LevelEditorTabIds::PlacementBrowser)
 	{
