Hide the new in-PIE level viewport toolbar when in immersive mode or in stereo 3D.

diff --git a/Engine/Source/Editor/LevelEditor/Private/SLevelViewport.cpp b/Engine/Source/Editor/LevelEditor/Private/SLevelViewport.cpp
--- a/Engine/Source/Editor/LevelEditor/Private/SLevelViewport.cpp
+++ b/Engine/Source/Editor/LevelEditor/Private/SLevelViewport.cpp
@@ -2255,6 +2255,15 @@ TSharedPtr<SWidget> SLevelViewport::BuildViewportToolbar()
 	// clang-format off
 	return SNew(SWidgetSwitcher)
 	.WidgetIndex(this, &SLevelViewport::GetViewportToolbarIndex)
+	.Visibility_Lambda([this]
+	{
+		// TODO: Move this to GetToolBarVisibility() once the old toolbar is deprecated
+		if ((IsPlayInEditorViewportActive() && IsImmersive()) || GEngine->IsStereoscopic3D(ActiveViewport.Get()))
+		{
+			return EVisibility::Collapsed;
+		}
+		return EVisibility::Visible;
+	})
 	+ SWidgetSwitcher::Slot()
 	[
 		UToolMenus::Get()->GenerateWidget(ViewportToolbarName, ViewportToolbarContext)
