Renamed "real time audio" slider to "Level Editor Volume" and added clear tooltips. Changed volume display to decibels.

[RN] Renamed "real time audio" slider to "Level Editor Volume" and added clear tooltips. Changed volume display to decibels.

diff --git a/Engine/Source/Editor/LevelEditor/Private/LevelEditorActions.cpp b/Engine/Source/Editor/LevelEditor/Private/LevelEditorActions.cpp
--- a/Engine/Source/Editor/LevelEditor/Private/LevelEditorActions.cpp
+++ b/Engine/Source/Editor/LevelEditor/Private/LevelEditorActions.cpp
@@ -3026,12 +3026,14 @@ bool FLevelEditorActionCallbacks::OnIsLevelStreamingVolumePrevisEnabled()
 
 FText FLevelEditorActionCallbacks::GetAudioVolumeToolTip()
 {
-	if ( !GEditor->IsRealTimeAudioMuted() )
+	float VolumeDecibels = -60.0f;
+	if (!GEditor->IsRealTimeAudioMuted())
 	{
-		const float Volume = GEditor->GetRealTimeAudioVolume() * 100.0f;
-		return FText::AsNumber( FMath::RoundToInt(Volume) );
+		const float Volume = GEditor->GetRealTimeAudioVolume();
+		VolumeDecibels = 20.0f * FMath::LogX(10.0f, FMath::Max(Volume, UE_SMALL_NUMBER));
+		VolumeDecibels = FMath::Max(VolumeDecibels, -60.0f);
 	}
-	return NSLOCTEXT("UnrealEd", "Muted", "Muted" );
+	return FText::Format(NSLOCTEXT("LevelEditorCommands", "LevelEditorVolumeToolTip", "Level Editor Volume is {0} dB."), FText::AsNumber(VolumeDecibels));
 }
 
 float FLevelEditorActionCallbacks::GetAudioVolume()
diff --git a/Engine/Source/Editor/LevelEditor/Private/LevelEditorToolBar.cpp b/Engine/Source/Editor/LevelEditor/Private/LevelEditorToolBar.cpp
--- a/Engine/Source/Editor/LevelEditor/Private/LevelEditorToolBar.cpp
+++ b/Engine/Source/Editor/LevelEditor/Private/LevelEditorToolBar.cpp
@@ -1779,13 +1779,6 @@ void FLevelEditorToolBar::RegisterQuickSettingsMenu()
 		Section.AddEntry(UE::LevelEditor::CreateFeatureLevelPreviewSubmenu());
 	}
 
-	{
-		FToolMenuSection& Section = Menu->AddSection("LevelEditorAudio", LOCTEXT("AudioHeading", "Real Time Audio") );
-
-		Section.AddEntry(FToolMenuEntry::InitWidget(
-			"Volume", UE::LevelEditor::BuildVolumeControlCustomWidget(), LOCTEXT("VolumeControlLabel", "Volume")));
-	}
-
 	{
 		FToolMenuSection& Section = Menu->AddSection( "Snapping", LOCTEXT("SnappingHeading","Snapping") );
 		Section.AddMenuEntry( FLevelEditorCommands::Get().EnableActorSnap );
diff --git a/Engine/Source/Editor/LevelEditor/Private/ViewportToolbar/LevelEditorSubmenus.cpp b/Engine/Source/Editor/LevelEditor/Private/ViewportToolbar/LevelEditorSubmenus.cpp
--- a/Engine/Source/Editor/LevelEditor/Private/ViewportToolbar/LevelEditorSubmenus.cpp
+++ b/Engine/Source/Editor/LevelEditor/Private/ViewportToolbar/LevelEditorSubmenus.cpp
@@ -2447,9 +2447,11 @@ FToolMenuEntry CreateSettingsSubmenu()
 					{
 						FToolMenuSection& SettingsSection = Submenu->FindOrAddSection("Settings", LOCTEXT("SettingsSectionLabel", "Settings"));
 						SettingsSection.AddEntry(FToolMenuEntry::InitWidget(
-							"Volume",
+							"Level Editor Volume (dB)",
 							BuildVolumeControlCustomWidget(),
-							LOCTEXT("VolumeControlLabel", "Volume")
+							LOCTEXT("VolumeControlLabel", "Level Editor Volume (dB)"), 
+							false, true, false,
+							LOCTEXT("VolumeControlToolTip", "Sets the level editor's preview volume of audio placed on actors in the level editor (e.g. Ambient Actors).")
 						));
 					}
 					// Mouse Control
