Add view mode and show flag menus in the level editor's in-PIE viewport toolbar.
- Expose various visualization modes in GameViewportClient
- Fixes overuse of direct viewport binding in the view modes menu.
- Enable immediate feedback for view mode, show flag, and performance profile changes while in PIE.

diff --git a/Engine/Source/Editor/CommonMenuExtensions/Private/ShowFlagMenuCommands.cpp b/Engine/Source/Editor/CommonMenuExtensions/Private/ShowFlagMenuCommands.cpp
--- a/Engine/Source/Editor/CommonMenuExtensions/Private/ShowFlagMenuCommands.cpp
+++ b/Engine/Source/Editor/CommonMenuExtensions/Private/ShowFlagMenuCommands.cpp
@@ -211,6 +211,11 @@ void FShowFlagMenuCommands::PopulateAllShowFlagsSection(FToolMenuSection& Sectio
 	);
 }
 
+const TArray<FShowFlagMenuCommands::FShowFlagCommand>& FShowFlagMenuCommands::GetCommands() const
+{
+	return ShowFlagCommands;
+}
+
 TSharedPtr<const FUICommandInfo> FShowFlagMenuCommands::FindCommandForFlag(FEngineShowFlags::EShowFlag FlagIndex) const
 {
 	for (int32 Index = 0; Index < ShowFlagCommands.Num(); ++Index)
diff --git a/Engine/Source/Editor/CommonMenuExtensions/Public/ShowFlagMenuCommands.h b/Engine/Source/Editor/CommonMenuExtensions/Public/ShowFlagMenuCommands.h
--- a/Engine/Source/Editor/CommonMenuExtensions/Public/ShowFlagMenuCommands.h
+++ b/Engine/Source/Editor/CommonMenuExtensions/Public/ShowFlagMenuCommands.h
@@ -60,6 +60,7 @@ public:
 		FToolMenuSection& Section, const FShowFlagFilter& Filter = FShowFlagFilter(FShowFlagFilter::IncludeAllFlagsByDefault)
 	) const;
 	
+	const TArray<FShowFlagCommand>& GetCommands() const;
 	TSharedPtr<const FUICommandInfo> FindCommandForFlag(FEngineShowFlags::EShowFlag FlagIndex) const;
 
 private:
diff --git a/Engine/Source/Editor/LevelEditor/Private/SLevelViewport.cpp b/Engine/Source/Editor/LevelEditor/Private/SLevelViewport.cpp
--- a/Engine/Source/Editor/LevelEditor/Private/SLevelViewport.cpp
+++ b/Engine/Source/Editor/LevelEditor/Private/SLevelViewport.cpp
@@ -105,8 +105,11 @@
 #include "EditorViewportTabContent.h"
 #include "EditorViewportCommands.h"
 #include "FunctionalUIScreenshotTest.h"
+#include "GPUSkinCache.h"
+#include "GPUSkinCacheVisualizationMenuCommands.h"
 #include "LevelEditorStyle.h"
 #include "SActionableMessageViewportWidget.h"
+#include "ShowFlagMenuCommands.h"
 #include "SkeletalRenderPublic.h"
 #include "ViewportToolbar/LevelEditorSubmenus.h"
 #include "ViewportToolbar/LevelViewportContext.h"
@@ -769,7 +772,9 @@ FReply SLevelViewport::OnKeyDown( const FGeometry& MyGeometry, const FKeyEvent&
 	if( HasPlayInEditorViewport() || LevelViewportClient->IsSimulateInEditorViewport() )
 	{
 		// Only process commands for pie when a play world is active
-		FPlayWorldCommands::GlobalPlayWorldActions->ProcessCommandBindings( InKeyEvent );
+		bool bHandled = FPlayWorldCommands::GlobalPlayWorldActions->ProcessCommandBindings( InKeyEvent );
+		
+		bHandled = bHandled || PIECommands->ProcessCommandBindings(InKeyEvent);
 
 		// Always handle commands in pie so they arent bubbled to editor only widgets
 		Reply = FReply::Handled();
@@ -1543,6 +1548,9 @@ void SLevelViewport::BindCommands()
 
 	// Exposes the current level viewport command list to subscribers from other systems
 	FInputBindingManager::Get().RegisterCommandList(FLevelViewportCommands::Get().GetContextName(), CommandList.ToSharedRef());
+	
+	PIECommands = MakeShared<FUICommandList>();
+	BindPIECommands(*PIECommands);
 }
 	
 void SLevelViewport::BindOptionCommands( FUICommandList& OutCommandList )
@@ -1866,6 +1874,139 @@ void SLevelViewport::BindStatCommand(const TSharedPtr<FUICommandInfo> InMenuItem
 		FIsActionChecked::CreateSP(this, &SLevelViewport::IsStatCommandVisible, InCommandName));
 }
 
+void SLevelViewport::BindPIECommands(FUICommandList& OutCommandList)
+{
+	OutCommandList.MapAction(FLevelViewportCommands::Get().UseDefaultShowFlags,
+		FExecuteAction::CreateSP(this, &SLevelViewport::ResetPIEShowFlags)
+	);
+
+	for (const FShowFlagMenuCommands::FShowFlagCommand& FlagCommand : FShowFlagMenuCommands::Get().GetCommands())
+	{
+		OutCommandList.MapAction(FlagCommand.ShowMenuItem,
+			FExecuteAction::CreateSP(this, &SLevelViewport::TogglePIEShowFlag, FlagCommand.FlagIndex),
+			FCanExecuteAction::CreateStatic(&FEngineShowFlags::IsForceFlagSet, static_cast<uint32>(FlagCommand.FlagIndex)),
+			FIsActionChecked::CreateSP(this, &SLevelViewport::IsPIEShowFlagEnabled, FlagCommand.FlagIndex)
+		);
+	}
+	
+	const FEditorViewportCommands& Commands = FEditorViewportCommands::Get();
+	
+	#define MAP_VIEWMODE_ACTION( ViewModeCommand, ViewModeID ) \
+	OutCommandList.MapAction( \
+		ViewModeCommand, \
+		FExecuteAction::CreateSP( this, &SLevelViewport::SetPIEViewMode, ViewModeID ), \
+		FCanExecuteAction(), \
+		FIsActionChecked::CreateSP( this, &SLevelViewport::IsPIEViewModeEnabled, ViewModeID ) ) 
+		
+	// Map each view mode
+	MAP_VIEWMODE_ACTION( Commands.WireframeMode, VMI_BrushWireframe );
+	MAP_VIEWMODE_ACTION( Commands.UnlitMode, VMI_Unlit );
+	MAP_VIEWMODE_ACTION( Commands.LitMode, VMI_Lit );
+	MAP_VIEWMODE_ACTION( Commands.LitWireframeMode, VMI_Lit_Wireframe);
+
+	if (IsRayTracingAllowed())
+	{
+		MAP_VIEWMODE_ACTION(Commands.PathTracingMode, VMI_PathTracing);
+		MAP_VIEWMODE_ACTION(Commands.RayTracingDebugMode, VMI_RayTracingDebug);
+
+		// const FRayTracingDebugVisualizationMenuCommands& RtDebugCommands = FRayTracingDebugVisualizationMenuCommands::Get();
+		// RtDebugCommands.BindCommands(CommandListRef, Client);
+	}
+
+	MAP_VIEWMODE_ACTION( Commands.DetailLightingMode, VMI_Lit_DetailLighting );
+	MAP_VIEWMODE_ACTION( Commands.LightingOnlyMode, VMI_LightingOnly );
+	MAP_VIEWMODE_ACTION( Commands.LightComplexityMode, VMI_LightComplexity );
+	MAP_VIEWMODE_ACTION( Commands.ShaderComplexityMode, VMI_ShaderComplexity );
+	MAP_VIEWMODE_ACTION( Commands.QuadOverdrawMode, VMI_QuadOverdraw);
+	MAP_VIEWMODE_ACTION( Commands.ShaderComplexityWithQuadOverdrawMode, VMI_ShaderComplexityWithQuadOverdraw );
+	MAP_VIEWMODE_ACTION( Commands.TexStreamAccPrimitiveDistanceMode, VMI_PrimitiveDistanceAccuracy );
+	MAP_VIEWMODE_ACTION( Commands.TexStreamAccMeshUVDensityMode, VMI_MeshUVDensityAccuracy);
+	MAP_VIEWMODE_ACTION( Commands.TexStreamAccMaterialTextureScaleMode, VMI_MaterialTextureScaleAccuracy );
+	MAP_VIEWMODE_ACTION( Commands.RequiredTextureResolutionMode, VMI_RequiredTextureResolution);
+	MAP_VIEWMODE_ACTION( Commands.StationaryLightOverlapMode, VMI_StationaryLightOverlap );
+
+	if (IsStaticLightingAllowed())
+	{
+		MAP_VIEWMODE_ACTION(Commands.LightmapDensityMode, VMI_LightmapDensity);
+	}
+
+	MAP_VIEWMODE_ACTION( Commands.ReflectionOverrideMode, VMI_ReflectionOverride );
+	MAP_VIEWMODE_ACTION( Commands.GroupLODColorationMode, VMI_GroupLODColoration);
+	MAP_VIEWMODE_ACTION( Commands.LODColorationMode, VMI_LODColoration );
+	MAP_VIEWMODE_ACTION( Commands.HLODColorationMode, VMI_HLODColoration);
+	MAP_VIEWMODE_ACTION( Commands.VisualizeBufferMode, VMI_VisualizeBuffer );
+	MAP_VIEWMODE_ACTION( Commands.VisualizeNaniteMode, VMI_VisualizeNanite );
+	MAP_VIEWMODE_ACTION( Commands.VisualizeLumenMode, VMI_VisualizeLumen );
+	MAP_VIEWMODE_ACTION( Commands.VisualizeSubstrateMode, VMI_VisualizeSubstrate);
+	MAP_VIEWMODE_ACTION( Commands.VisualizeGroomMode, VMI_VisualizeGroom);
+	MAP_VIEWMODE_ACTION( Commands.VisualizeVirtualShadowMapMode, VMI_VisualizeVirtualShadowMap );
+	MAP_VIEWMODE_ACTION( Commands.VisualizeVirtualTextureMode, VMI_VisualizeVirtualTexture);
+	MAP_VIEWMODE_ACTION( Commands.CollisionPawn, VMI_CollisionPawn);
+	MAP_VIEWMODE_ACTION( Commands.CollisionVisibility, VMI_CollisionVisibility);
+
+	MAP_VIEWMODE_ACTION( Commands.VisualizeLWCComplexity, VMI_LWCComplexity);
+	
+}
+
+void SLevelViewport::ResetPIEShowFlags()
+{
+	if (UGameViewportClient* GameClient = GetPlayClient())
+	{
+		GameClient->EngineShowFlags = FEngineShowFlags(EShowFlagInitMode::ESFIM_Game);
+		RefreshPIEViewport();
+	}
+}
+
+void SLevelViewport::TogglePIEShowFlag(FEngineShowFlags::EShowFlag Flag)
+{
+	if (UGameViewportClient* GameClient = GetPlayClient())
+	{
+		GameClient->GetEngineShowFlags()->SetSingleFlag(Flag, !GameClient->EngineShowFlags.GetSingleFlag(Flag));
+		RefreshPIEViewport();
+	}
+}
+
+bool SLevelViewport::IsPIEShowFlagEnabled(FEngineShowFlags::EShowFlag Flag) const
+{
+	if (UGameViewportClient* GameClient = GetPlayClient())
+	{
+		return GameClient->GetEngineShowFlags()->GetSingleFlag(Flag);
+	}
+	return false;
+}
+
+void SLevelViewport::SetPIEViewMode(EViewModeIndex ViewMode)
+{
+	if (UGameViewportClient* GameClient = GetPlayClient())
+	{
+		GameClient->SetViewMode(ViewMode);
+		RefreshPIEViewport();
+	}
+}
+
+bool SLevelViewport::IsPIEViewModeEnabled(EViewModeIndex ViewMode) const
+{
+	if (UGameViewportClient* GameClient = GetPlayClient())
+	{
+		return GameClient->ViewModeIndex == ViewMode;
+	}
+	return false;
+}
+
+void SLevelViewport::RefreshPIEViewport()
+{
+	if (UGameViewportClient* GameClient = GetPlayClient())
+	{
+		FSlateThrottleManager::Get().DisableThrottle(true);
+		
+		FTimerHandle TimerHandle;
+		GameClient->GetWorld()->GetTimerManager().SetTimer(TimerHandle, []
+		{
+			FSlateThrottleManager::Get().DisableThrottle(false);
+		}, 0.1f, false);
+	}
+}
+
 const FSlateBrush* SLevelViewport::OnGetViewportBorderBrush() const
 {
 	const FSlateBrush* BorderBrush = nullptr;
@@ -2331,6 +2472,8 @@ TSharedRef<SWidget> SLevelViewport::BuildPIEViewportToolbar()
 			FToolMenuSection& RightSection = ToolbarMenu->AddSection("Right");
 			RightSection.Alignment = EToolMenuSectionAlign::Last;
 			
+			RightSection.AddEntry(UE::LevelEditor::CreatePIEViewModesSubmenu());
+			RightSection.AddEntry(UE::LevelEditor::CreatePIEShowSubmenu());
 			RightSection.AddEntry(UE::LevelEditor::CreatePIEPerformanceAndScalabilitySubmenu());
 			RightSection.AddEntry(UE::LevelEditor::CreatePIESettingsSubmenu());
 		}
@@ -2338,6 +2481,7 @@ TSharedRef<SWidget> SLevelViewport::BuildPIEViewportToolbar()
 	
 	FToolMenuContext ToolbarContext;
 	{
+		ToolbarContext.AppendCommandList(PIECommands);
 		ToolbarContext.AppendCommandList(GetCommandList());
 		
 		// Add the UnrealEd viewport toolbar context.
@@ -2345,6 +2489,36 @@ TSharedRef<SWidget> SLevelViewport::BuildPIEViewportToolbar()
 			UUnrealEdViewportToolbarContext* const ContextObject = NewObject<UUnrealEdViewportToolbarContext>();
 			ContextObject->Viewport = SharedThis(this);
 			ToolbarContext.AddObject(ContextObject);
+			
+			ContextObject->ExcludedShowMenuFlags.Append({
+				FEngineShowFlags::SF_Cameras,
+				FEngineShowFlags::SF_CameraFrustums,
+				FEngineShowFlags::SF_Grid
+			});
+			
+			ContextObject->IsViewModeSupported.BindLambda([](EViewModeIndex ViewMode)
+			{
+				switch (ViewMode)
+				{
+				case VMI_Lit_Wireframe:
+					return false;
+				default:
+					return true;
+				}
+			});
+			
+			ContextObject->DoesViewModeMenuShowSection.BindLambda([](UE::UnrealEd::EHidableViewModeMenuSections InMenuSection)
+			{
+				switch (InMenuSection)
+				{
+				case UE::UnrealEd::EHidableViewModeMenuSections::GPUSkinCache:
+					return false;
+				case UE::UnrealEd::EHidableViewModeMenuSections::RayTracingDebug:
+					return false;
+				default:
+					return true;
+				}
+			});
 		}
 
 		// Add the level editor viewport toolbar context.
@@ -3228,6 +3402,8 @@ void SLevelViewport::RedrawViewport( bool bInvalidateHitProxies )
 			}			
 		}
 	}
+	
+	RefreshPIEViewport();
 }
 
 bool SLevelViewport::CanToggleMaximizeMode() const
@@ -4805,6 +4981,8 @@ void SLevelViewport::EndPlayInEditorSession()
 	const bool bCheckMissingOverride = false;
 	LevelViewportClient->RemoveRealtimeOverride(LOCTEXT("LevelViewport_RealTimeDisableOnPie", "Disable LevelViewport Realtime for PIE"), bCheckMissingOverride);
 
+	PlayClient.Reset();
+
 	if( IsPlayInEditorViewportActive() )
 	{
 		{
@@ -4845,8 +5023,6 @@ void SLevelViewport::EndPlayInEditorSession()
 	// Reset our game layer manager's active to that of the active editor viewport
 	GameLayerManager->SetSceneViewport(ActiveViewport.Get());
 	
-	PlayClient.Reset();
-
 	// Reset the inactive viewport
 	InactiveViewport.Reset();
 
diff --git a/Engine/Source/Editor/LevelEditor/Private/ViewportToolbar/LevelEditorSubmenus.cpp b/Engine/Source/Editor/LevelEditor/Private/ViewportToolbar/LevelEditorSubmenus.cpp
--- a/Engine/Source/Editor/LevelEditor/Private/ViewportToolbar/LevelEditorSubmenus.cpp
+++ b/Engine/Source/Editor/LevelEditor/Private/ViewportToolbar/LevelEditorSubmenus.cpp
@@ -1254,6 +1254,58 @@ void ExtendViewModesSubmenu(FName InViewModesSubmenuName)
 	);
 }
 
+FToolMenuEntry CreatePIEViewModesSubmenu()
+{
+	return FToolMenuEntry::InitDynamicEntry(
+		"DynamicViewModes",
+		FNewToolMenuSectionDelegate::CreateLambda(
+			[](FToolMenuSection& InDynamicSection) -> void
+			{
+				ULevelViewportContext* const Context = InDynamicSection.FindContext<ULevelViewportContext>();
+				if (!Context)
+				{
+					return;
+				}
+			
+				TAttribute<FText> LabelAttribute = TAttribute<FText>::CreateLambda(
+					[WeakViewport = Context->LevelViewport]
+					{
+						if (TSharedPtr<SLevelViewport> Viewport = WeakViewport.Pin())
+						{
+							return UE::UnrealEd::GetViewModesSubmenuLabel(Viewport->GetPlayClient());
+						}
+						return FText::GetEmpty();
+					}
+				);
+
+				TAttribute<FSlateIcon> IconAttribute = TAttribute<FSlateIcon>::CreateLambda(
+					[WeakViewport = Context->LevelViewport]
+					{
+						if (TSharedPtr<SLevelViewport> Viewport = WeakViewport.Pin())
+						{
+							return UE::UnrealEd::GetViewModesSubmenuIcon(Viewport->GetPlayClient());
+						}
+						return FSlateIcon();
+					}
+				);
+				
+				FToolMenuEntry& Entry = InDynamicSection.AddSubMenu(
+					"ViewModes",
+					LabelAttribute,
+					LOCTEXT("ViewModesSubmenuTooltip", "View mode settings for the game viewport."),
+					FNewToolMenuDelegate::CreateLambda([](UToolMenu* Submenu)
+					{
+						UE::UnrealEd::PopulateViewModesMenu(Submenu);
+					}),
+					false,
+					IconAttribute
+				);
+				Entry.ToolBarData.ResizeParams.ClippingPriority = 800;
+			}
+		)
+	);
+}
+
 FToolMenuEntry CreateShowFoliageSubmenu()
 {
 	return FToolMenuEntry::InitSubMenu(
@@ -1987,6 +2039,30 @@ FToolMenuEntry CreateShowSubmenu()
 	return Entry;
 }
 
+FToolMenuEntry CreatePIEShowSubmenu()
+{
+	FToolMenuEntry Entry = FToolMenuEntry::InitSubMenu(
+		"Show",
+		LOCTEXT("ShowSubmenuLabel", "Show"),
+		LOCTEXT("ShowSubmenuTooltip", "Show flags related to the game viewport"),
+		FNewToolMenuDelegate::CreateLambda(
+			[](UToolMenu* InMenu) -> void
+			{
+				{
+					FToolMenuSection& UnnamedSection = InMenu->FindOrAddSection(NAME_None);
+					UnnamedSection.AddMenuEntry(FLevelViewportCommands::Get().UseDefaultShowFlags);
+				}
+
+				UE::UnrealEd::AddDefaultShowFlags(InMenu);
+			}
+		)
+	);
+	Entry.Icon = FSlateIcon(FAppStyle::GetAppStyleSetName(), "Level.VisibleHighlightIcon16x");
+	Entry.ToolBarData.LabelOverride = FText();
+	Entry.ToolBarData.ResizeParams.ClippingPriority = 800;
+	return Entry;
+}
+
 FToolMenuEntry CreateFeatureLevelPreviewSubmenu()
 {
 	return FToolMenuEntry::InitSubMenu(
diff --git a/Engine/Source/Editor/LevelEditor/Private/ViewportToolbar/LevelEditorSubmenus.h b/Engine/Source/Editor/LevelEditor/Private/ViewportToolbar/LevelEditorSubmenus.h
--- a/Engine/Source/Editor/LevelEditor/Private/ViewportToolbar/LevelEditorSubmenus.h
+++ b/Engine/Source/Editor/LevelEditor/Private/ViewportToolbar/LevelEditorSubmenus.h
@@ -27,6 +27,7 @@ FToolMenuEntry CreateToolbarCameraSubmenu();
 TSharedPtr<FExtender> GetViewModesLegacyExtenders();
 void PopulateViewModesMenu(UToolMenu* InMenu);
 void ExtendViewModesSubmenu(FName InViewModesSubmenuName);
+FToolMenuEntry CreatePIEViewModesSubmenu();
 
 FToolMenuEntry CreateShowFoliageSubmenu();
 FToolMenuEntry CreateShowHLODsSubmenu();
@@ -39,6 +40,7 @@ FToolMenuEntry CreateShowStatsSubmenu(
 #endif
 FToolMenuEntry CreateShowVolumesSubmenu();
 FToolMenuEntry CreateShowSubmenu();
+FToolMenuEntry CreatePIEShowSubmenu();
 
 FToolMenuEntry CreateFeatureLevelPreviewSubmenu();
 FToolMenuEntry CreateMaterialQualityLevelSubmenu();
diff --git a/Engine/Source/Editor/LevelEditor/Public/SLevelViewport.h b/Engine/Source/Editor/LevelEditor/Public/SLevelViewport.h
--- a/Engine/Source/Editor/LevelEditor/Public/SLevelViewport.h
+++ b/Engine/Source/Editor/LevelEditor/Public/SLevelViewport.h
@@ -722,6 +722,42 @@ private:
 	* @param InCommandName	The command used by the functions
 	*/
 	void BindStatCommand(const TSharedPtr<FUICommandInfo> InMenuItem, const FString& InCommandName);
+
+	/**
+	 * Binds commands for PIE-specific behavior
+	 * @param OutCommandList	The list to bind commands to
+	 */
+	void BindPIECommands( FUICommandList& OutCommandList );
+
+	/**
+	 * Resets the PIE show flags to default.
+	 */
+	void ResetPIEShowFlags();
+	
+	/**
+	 * Toggles a PIE show flag on or off.
+	 */
+	void TogglePIEShowFlag(FEngineShowFlags::EShowFlag Flag);
+
+	/**
+	 * Returns whether a show flag is on in PIE.
+	 */
+	bool IsPIEShowFlagEnabled(FEngineShowFlags::EShowFlag Flag) const;
+
+	/**
+	 * Sets the view mode on the PIE viewport.
+	 */
+	void SetPIEViewMode(EViewModeIndex ViewMode);
+
+	/**
+	 * Gets whether the view mode is enabled in the PIE Viewport
+	 */
+	bool IsPIEViewModeEnabled(EViewModeIndex ViewMode) const;
+
+	/**
+	 * Refreshes the PIE viewport.
+	 */
+	void RefreshPIEViewport();
 	
 	/**
 	 * Called to build the viewport context menu when the user is Drag Dropping from the Content Browser
@@ -861,6 +897,9 @@ private:
 	
 	/** The GameViewportClient */
 	TWeakObjectPtr<UGameViewportClient> PlayClient;
+	
+	/** Commands specifically for PIE */
+	TSharedPtr<FUICommandList> PIECommands;
 
 	/** The brush to use if this viewport is in debug mode */
 	const FSlateBrush* DebuggingBorder;
diff --git a/Engine/Source/Editor/UnrealEd/Private/SScalabilitySettings.cpp b/Engine/Source/Editor/UnrealEd/Private/SScalabilitySettings.cpp
--- a/Engine/Source/Editor/UnrealEd/Private/SScalabilitySettings.cpp
+++ b/Engine/Source/Editor/UnrealEd/Private/SScalabilitySettings.cpp
@@ -72,6 +72,7 @@ void SScalabilitySettings::OnGroupQualityLevelChanged(ECheckBoxState NewState, c
 	Scalability::SetQualityLevels(CachedQualityLevels);
 	Scalability::SaveState(GEditorSettingsIni);
 	GEditor->RedrawAllViewports();
+	GEditor->RedrawLevelEditingViewports();
 }
 
 void SScalabilitySettings::OnResolutionScaleChanged(float InValue)
@@ -86,6 +87,7 @@ void SScalabilitySettings::OnResolutionScaleChanged(float InValue)
 	Scalability::SetQualityLevels(CachedQualityLevels);
 	Scalability::SaveState(GEditorSettingsIni);
 	GEditor->RedrawAllViewports();
+	GEditor->RedrawLevelEditingViewports();
 }
 
 float SScalabilitySettings::GetResolutionScale() const
@@ -111,6 +113,7 @@ FReply SScalabilitySettings::OnResolutionPresetClicked(int32 PresetId)
 	Scalability::SetQualityLevels(CachedQualityLevels);
 	Scalability::SaveState(GEditorSettingsIni);
 	GEditor->RedrawAllViewports();
+	GEditor->RedrawLevelEditingViewports();
 
 	return FReply::Handled();
 }
@@ -201,6 +204,7 @@ FReply SScalabilitySettings::OnHeaderClicked(int32 InQualityLevel, Scalability::
 	Scalability::SetQualityLevels(CachedQualityLevels);
 	Scalability::SaveState(GEditorSettingsIni);
 	GEditor->RedrawAllViewports();
+	GEditor->RedrawLevelEditingViewports();
 	return FReply::Handled();
 }
 
@@ -213,6 +217,7 @@ FReply SScalabilitySettings::OnAutoClicked()
 	CachedQualityLevels = Settings->EngineBenchmarkResult;
 
 	GEditor->RedrawAllViewports();
+	GEditor->RedrawLevelEditingViewports();
 	return FReply::Handled();
 }
 
@@ -223,6 +228,7 @@ FReply SScalabilitySettings::OnDefaultClicked()
 	Scalability::SetQualityLevels(CachedQualityLevels);
 
 	GEditor->RedrawAllViewports();
+	GEditor->RedrawLevelEditingViewports();
 	return FReply::Handled();
 }
 
diff --git a/Engine/Source/Editor/UnrealEd/Private/ViewportToolbar/UnrealEdViewportToolbar.cpp b/Engine/Source/Editor/UnrealEd/Private/ViewportToolbar/UnrealEdViewportToolbar.cpp
--- a/Engine/Source/Editor/UnrealEd/Private/ViewportToolbar/UnrealEdViewportToolbar.cpp
+++ b/Engine/Source/Editor/UnrealEd/Private/ViewportToolbar/UnrealEdViewportToolbar.cpp
@@ -9,6 +9,7 @@
 #include "Editor/EditorPerformanceSettings.h"
 #include "EditorViewportClient.h"
 #include "EditorViewportCommands.h"
+#include "Engine/GameViewportClient.h"
 #include "Framework/Commands/GenericCommands.h"
 #include "GPUSkinCache.h"
 #include "GPUSkinCacheVisualizationMenuCommands.h"
@@ -1420,6 +1421,16 @@ FText GetViewModesSubmenuLabel(TWeakPtr<SEditorViewport> InViewport)
 	return Label;
 }
 
+FText GetViewModesSubmenuLabel(const UGameViewportClient* InViewportClient)
+{
+	FText Label = LOCTEXT("ViewMenuTitlePIE_Default", "View");
+	if (InViewportClient)
+	{
+		Label = UViewModeUtils::GetViewModeDisplayName((EViewModeIndex)InViewportClient->ViewModeIndex);
+	}
+	return Label;
+}
+
 // Explicitly passing Editor Viewport, since the old toolbar might get an empty Context once calling this
 FToolMenuEntry GetLitWireframeEntryToSection(const TSharedPtr<SEditorViewport>& EditorViewport)
 {
@@ -1510,6 +1521,57 @@ FSlateIcon GetViewModesSubmenuIcon(const TWeakPtr<SEditorViewport>& InViewport)
 	return FSlateIcon();
 }
 
+FSlateIcon GetViewModesSubmenuIcon(const UGameViewportClient* InViewportClient)
+{
+	if (InViewportClient)
+	{	
+		return UViewModeUtils::GetViewModeDisplaySlateIcon((EViewModeIndex)InViewportClient->ViewModeIndex);
+	}
+	return FSlateIcon();
+}
+
+FGetActionCheckState BuildParentCheckStateForCommands(
+	const FToolMenuContext& Context,
+	const TArray<TSharedPtr<FUICommandInfo>>& Commands
+)
+{
+	FGetActionCheckState Result;
+	
+	TSharedPtr<const FUICommandList> CommandList;
+	
+	for (const TSharedPtr<FUICommandInfo>& Command : Commands)
+	{
+		Context.GetActionForCommand(Command, CommandList);
+		if (CommandList)
+		{
+			break;
+		}
+	}
+	
+	if (CommandList)
+	{
+		Result.BindLambda([WeakCommandList = CommandList.ToWeakPtr(), Commands = Commands]
+		{
+			if (TSharedPtr<const FUICommandList> CommandList = WeakCommandList.Pin())
+			{
+				for (const TSharedPtr<FUICommandInfo>& Command : Commands)
+				{
+					if (const FUIAction* Action = CommandList->GetActionForCommand(Command))
+					{
+						if (Action->GetCheckState() != ECheckBoxState::Unchecked)
+						{
+							return ECheckBoxState::Checked;
+						}
+					}
+				}
+			}
+			return ECheckBoxState::Unchecked;
+		});
+	}
+	
+	return Result;
+}
+
 void PopulateViewModesMenu(UToolMenu* InMenu)
 {
 	UUnrealEdViewportToolbarContext* const Context = InMenu->FindContext<UUnrealEdViewportToolbarContext>();
@@ -1712,35 +1774,26 @@ void PopulateViewModesMenu(UToolMenu* InMenu)
 
 			if (Local::ViewModesShouldShowOptimizationEntries(IsViewModeSupported))
 			{
+				FUIAction MenuAction;
+				MenuAction.GetActionCheckState = BuildParentCheckStateForCommands(InMenu->Context, {
+					BaseViewportActions.LightComplexityMode,
+					BaseViewportActions.LightmapDensityMode,
+					BaseViewportActions.StationaryLightOverlapMode,
+					BaseViewportActions.ShaderComplexityMode,
+					BaseViewportActions.ShaderComplexityWithQuadOverdrawMode,
+					BaseViewportActions.QuadOverdrawMode,
+					BaseViewportActions.TexStreamAccPrimitiveDistanceMode,
+					BaseViewportActions.TexStreamAccMeshUVDensityMode,
+					BaseViewportActions.TexStreamAccMaterialTextureScaleMode,
+					BaseViewportActions.RequiredTextureResolutionMode
+				});
+				
 				Section.AddSubMenu(
 					"OptimizationSubMenu",
 					LOCTEXT("OptimizationSubMenu", "Optimization Viewmodes"),
 					LOCTEXT("Optimization_ToolTip", "Select optimization visualizer"),
 					FNewToolMenuDelegate::CreateStatic(&Local::BuildOptimizationMenu, IsViewModeSupported),
-					FUIAction(
-						FExecuteAction(),
-						FCanExecuteAction(),
-						FIsActionChecked::CreateLambda(
-							[Viewport = EditorViewport.ToWeakPtr()]()
-							{
-								const TSharedRef<SEditorViewport> ViewportRef = Viewport.Pin().ToSharedRef();
-								const TSharedPtr<FEditorViewportClient> ViewportClient = ViewportRef->GetViewportClient();
-								check(ViewportClient.IsValid());
-								const EViewModeIndex ViewMode = ViewportClient->GetViewMode();
-								return (
-									// Texture Streaming Accuracy
-									ViewMode == VMI_LightComplexity || ViewMode == VMI_LightmapDensity
-									|| ViewMode == VMI_StationaryLightOverlap || ViewMode == VMI_ShaderComplexity
-									|| ViewMode == VMI_ShaderComplexityWithQuadOverdraw
-									|| ViewMode == VMI_QuadOverdraw
-									// Texture Streaming Accuracy
-									|| ViewMode == VMI_PrimitiveDistanceAccuracy || ViewMode == VMI_MeshUVDensityAccuracy
-									|| ViewMode == VMI_MaterialTextureScaleAccuracy
-									|| ViewMode == VMI_RequiredTextureResolution
-								);
-							}
-						)
-					),
+					MenuAction,
 					EUserInterfaceActionType::RadioButton,
 					/* bInOpenSubMenuOnClick = */ false,
 					FSlateIcon(FAppStyle::GetAppStyleSetName(), "EditorViewport.QuadOverdrawMode")
@@ -1797,25 +1850,18 @@ void PopulateViewModesMenu(UToolMenu* InMenu)
 			if (!IsViewModeSupported.IsBound()
 				|| (IsViewModeSupported.Execute(VMI_LODColoration) || IsViewModeSupported.Execute(VMI_HLODColoration)))
 			{
+				FUIAction MenuAction;
+				MenuAction.GetActionCheckState = BuildParentCheckStateForCommands(InMenu->Context, {
+					FEditorViewportCommands::Get().LODColorationMode,
+					FEditorViewportCommands::Get().HLODColorationMode
+				});
+				
 				Section.AddSubMenu(
 					"VisualizeGroupedLOD",
 					LOCTEXT("VisualizeGroupedLODDisplayName", "Level of Detail Coloration"),
 					LOCTEXT("GroupedLODMenu_ToolTip", "Select a mode for LOD Coloration"),
 					FNewToolMenuDelegate::CreateStatic(&Local::BuildLODMenu, IsViewModeSupported),
-					FUIAction(
-						FExecuteAction(),
-						FCanExecuteAction(),
-						FIsActionChecked::CreateLambda(
-							[WeakViewport = EditorViewport.ToWeakPtr()]()
-							{
-								const TSharedRef<SEditorViewport> ViewportRef = WeakViewport.Pin().ToSharedRef();
-								const TSharedPtr<FEditorViewportClient> ViewportClient = ViewportRef->GetViewportClient();
-								check(ViewportClient.IsValid());
-								const EViewModeIndex ViewMode = ViewportClient->GetViewMode();
-								return (ViewMode == VMI_LODColoration || ViewMode == VMI_HLODColoration);
-							}
-						)
-					),
+					MenuAction,
 					EUserInterfaceActionType::RadioButton,
 					/* bInOpenSubMenuOnClick = */ false,
 					FSlateIcon(FAppStyle::GetAppStyleSetName(), "EditorViewport.GroupLODColorationMode")
@@ -1828,24 +1874,18 @@ void PopulateViewModesMenu(UToolMenu* InMenu)
 				Context, EHidableViewModeMenuSections::GPUSkinCache
 			))
 		{
+			FUIAction MenuAction;
+			if (const FUIAction* FoundAction = InMenu->Context.GetActionForCommand(FEditorViewportCommands::Get().VisualizeGPUSkinCacheMode))
+			{
+				MenuAction.GetActionCheckState = FoundAction->GetActionCheckState;
+			}
+		
 			Section.AddSubMenu(
 				"VisualizeGPUSkinCacheViewMode",
 				LOCTEXT("VisualizeGPUSkinCacheViewModeDisplayName", "GPU Skin Cache"),
 				LOCTEXT("GPUSkinCacheVisualizationMenu_ToolTip", "Select a mode for GPU Skin Cache visualization."),
 				FNewMenuDelegate::CreateStatic(&FGPUSkinCacheVisualizationMenuCommands::BuildVisualisationSubMenu),
-				FUIAction(
-					FExecuteAction(),
-					FCanExecuteAction(),
-					FIsActionChecked::CreateLambda(
-						[WeakViewport = EditorViewport.ToWeakPtr()]()
-						{
-							const TSharedRef<SEditorViewport> ViewportRef = WeakViewport.Pin().ToSharedRef();
-							const TSharedPtr<FEditorViewportClient> ViewportClient = ViewportRef->GetViewportClient();
-							check(ViewportClient.IsValid());
-							return ViewportClient->IsViewModeEnabled(VMI_VisualizeGPUSkinCache);
-						}
-					)
-				),
+				MenuAction,
 				EUserInterfaceActionType::RadioButton,
 				/* bInOpenSubMenuOnClick = */ false,
 				FSlateIcon(FAppStyle::GetAppStyleSetName(), "EditorViewport.VisualizeGPUSkinCacheMode")
diff --git a/Engine/Source/Editor/UnrealEd/Public/ViewportToolbar/UnrealEdViewportToolbar.h b/Engine/Source/Editor/UnrealEd/Public/ViewportToolbar/UnrealEdViewportToolbar.h
--- a/Engine/Source/Editor/UnrealEd/Public/ViewportToolbar/UnrealEdViewportToolbar.h
+++ b/Engine/Source/Editor/UnrealEd/Public/ViewportToolbar/UnrealEdViewportToolbar.h
@@ -23,6 +23,7 @@ enum ERotationGridMode : int;
 struct FNewToolMenuChoice;
 struct FTabId;
 struct FToolMenuEntry;
+class UGameViewportClient;
 
 namespace UE::UnrealEd
 {
@@ -55,8 +56,10 @@ UNREALED_API FToolMenuEntry CreateViewportToolbarSnappingSubmenu();
 UNREALED_API FToolMenuEntry CreateSnappingSubmenu();
 
 UNREALED_API FText GetViewModesSubmenuLabel(TWeakPtr<SEditorViewport> InViewport);
+UNREALED_API FText GetViewModesSubmenuLabel(const UGameViewportClient* InViewportClient);
 
 UNREALED_API FSlateIcon GetViewModesSubmenuIcon(const TWeakPtr<SEditorViewport>& InViewport);
+UNREALED_API FSlateIcon GetViewModesSubmenuIcon(const UGameViewportClient* InViewportClient);
 
 /**
  * Populate a given UToolMenu with entries for a View Modes viewport toolbar submenu.
diff --git a/Engine/Source/Runtime/Engine/Classes/Engine/GameViewportClient.h b/Engine/Source/Runtime/Engine/Classes/Engine/GameViewportClient.h
--- a/Engine/Source/Runtime/Engine/Classes/Engine/GameViewportClient.h
+++ b/Engine/Source/Runtime/Engine/Classes/Engine/GameViewportClient.h
@@ -634,6 +634,9 @@ public:
 	{ 
 		return &EngineShowFlags; 
 	}
+	
+	/** Sets the view mode */
+	ENGINE_API void SetViewMode(EViewModeIndex ViewMode, FOutputDevice* Output = nullptr);
 
 	ENGINE_API bool SetHardwareCursor(EMouseCursor::Type CursorShape, FName GameContentPath, FVector2D HotSpot);
 
@@ -647,7 +650,6 @@ public:
 
 	ENGINE_API bool GetUseMouseForTouch() const;
 
-protected:
 	void SetCurrentBufferVisualizationMode(FName NewBufferVisualizationMode) { CurrentBufferVisualizationMode = NewBufferVisualizationMode; }
 	FName GetCurrentBufferVisualizationMode() const { return CurrentBufferVisualizationMode; }
 
@@ -666,6 +668,7 @@ protected:
 	void SetCurrentVirtualShadowMapVisualizationMode(FName NewVirtualShadowMapVisualizationMode) { CurrentVirtualShadowMapVisualizationMode = NewVirtualShadowMapVisualizationMode; }
 	FName GetCurrentVirtualShadowMapVisualizationMode() const { return CurrentVirtualShadowMapVisualizationMode; }
 
+protected:
 	bool HasAudioFocus() const { return bHasAudioFocus; }
 
 	/** Updates CSVProfiler camera stats */
diff --git a/Engine/Source/Runtime/Engine/Private/GameViewportClient.cpp b/Engine/Source/Runtime/Engine/Private/GameViewportClient.cpp
--- a/Engine/Source/Runtime/Engine/Private/GameViewportClient.cpp
+++ b/Engine/Source/Runtime/Engine/Private/GameViewportClient.cpp
@@ -3796,53 +3796,8 @@ bool UGameViewportClient::HandleViewModeCommand( const TCHAR* Cmd, FOutputDevice
 		Ar.Logf(TEXT("Available view modes: %s"), *ViewModes);
 	}
 
-	if (ViewModeIndex == VMI_StationaryLightOverlap)
-	{
-		Ar.Logf(TEXT("This view mode is currently not supported in game."));
-		ViewModeIndex = VMI_Lit;
-	}
-
-	if (FPlatformProperties::SupportsWindowedMode() == false)
-	{
-		if(ViewModeIndex == VMI_Unlit
-			|| ViewModeIndex == VMI_StationaryLightOverlap
-			|| ViewModeIndex == VMI_Lit_DetailLighting
-			|| ViewModeIndex == VMI_ReflectionOverride)
-		{
-			Ar.Logf(TEXT("This view mode is currently not supported on consoles."));
-			ViewModeIndex = VMI_Lit;
-		}
-	}
-
-#if UE_BUILD_TEST || UE_BUILD_SHIPPING
-	Ar.Logf(TEXT("Debug viewmodes not allowed in Test or Shipping builds."));
-	ViewModeIndex = VMI_Lit;
-#else
-
-	if ((ViewModeIndex != VMI_Lit && ViewModeIndex != VMI_ShaderComplexity) && !AllowDebugViewmodes())
-	{
-		Ar.Logf(TEXT("Debug viewmodes not allowed on consoles by default.  See AllowDebugViewmodes()."));
-		ViewModeIndex = VMI_Lit;
-	}
-
-	if (!GRHISupportsRayTracing || !GRHISupportsRayTracingShaders)
-	{
-		if (ViewModeIndex == VMI_PathTracing)
-		{
-			Ar.Logf(TEXT("Path Tracing view mode requires ray tracing support. It is not supported on this system."));
-			ViewModeIndex = VMI_Lit;
-		}
-
-		if (ViewModeIndex == VMI_RayTracingDebug)
-		{
-			Ar.Logf(TEXT("Ray tracing view mode requires ray tracing support. It is not supported on this system."));
-			ViewModeIndex = VMI_Lit;
-		}
-	}
-#endif
-
-	ApplyViewMode((EViewModeIndex)ViewModeIndex, true, EngineShowFlags);
-
+	SetViewMode((EViewModeIndex)ViewModeIndex);
+	
 	return true;
 }
 
@@ -4465,6 +4420,77 @@ void UGameViewportClient::HandleWindowDPIScaleChanged(TSharedRef<SWindow> InWind
 #endif
 }
 
+void UGameViewportClient::SetViewMode(EViewModeIndex ViewMode, FOutputDevice* Output)
+{
+	ViewModeIndex = ViewMode;
+
+	if (ViewModeIndex == VMI_StationaryLightOverlap)
+	{
+		if (Output)
+		{
+			Output->Logf(TEXT("This view mode is currently not supported in game."));
+		}
+		
+		ViewModeIndex = VMI_Lit;
+	}
+
+	if (FPlatformProperties::SupportsWindowedMode() == false)
+	{
+		if(ViewModeIndex == VMI_Unlit
+			|| ViewModeIndex == VMI_StationaryLightOverlap
+			|| ViewModeIndex == VMI_Lit_DetailLighting
+			|| ViewModeIndex == VMI_ReflectionOverride)
+		{
+			if (Output)
+			{
+				Output->Logf(TEXT("This view mode is currently not supported on consoles."));
+			}
+			ViewModeIndex = VMI_Lit;
+		}
+	}
+
+#if UE_BUILD_TEST || UE_BUILD_SHIPPING
+	if (Output)
+	{
+		Output->Logf(TEXT("Debug viewmodes not allowed in Test or Shipping builds."));
+	}
+	ViewModeIndex = VMI_Lit;
+#else
+
+	if ((ViewModeIndex != VMI_Lit && ViewModeIndex != VMI_ShaderComplexity) && !AllowDebugViewmodes())
+	{
+		if (Output)
+		{
+			Output->Logf(TEXT("Debug viewmodes not allowed on consoles by default.  See AllowDebugViewmodes()."));
+		}
+		ViewModeIndex = VMI_Lit;
+	}
+
+	if (!GRHISupportsRayTracing || !GRHISupportsRayTracingShaders)
+	{
+		if (ViewModeIndex == VMI_PathTracing)
+		{
+			if (Output)
+			{
+				Output->Logf(TEXT("Path Tracing view mode requires ray tracing support. It is not supported on this system."));
+			}
+			ViewModeIndex = VMI_Lit;
+		}
+
+		if (ViewModeIndex == VMI_RayTracingDebug)
+		{
+			if (Output)
+			{
+				Output->Logf(TEXT("Ray tracing view mode requires ray tracing support. It is not supported on this system."));
+			}
+			ViewModeIndex = VMI_Lit;
+		}
+	}
+#endif
+
+	ApplyViewMode((EViewModeIndex)ViewModeIndex, true, EngineShowFlags);
+}
+
 bool UGameViewportClient::SetHardwareCursor(EMouseCursor::Type CursorShape, FName GameContentPath, FVector2D HotSpot)
 {
 	TSharedPtr<ICursor> PlatformCursor = FSlateApplication::Get().GetPlatformCursor();
