Hotfix for 41650813 and namespace specificity issues.

diff --git a/Engine/Source/Editor/LevelEditor/Private/ViewportToolbar/LevelEditorSubmenus.cpp b/Engine/Source/Editor/LevelEditor/Private/ViewportToolbar/LevelEditorSubmenus.cpp
--- a/Engine/Source/Editor/LevelEditor/Private/ViewportToolbar/LevelEditorSubmenus.cpp
+++ b/Engine/Source/Editor/LevelEditor/Private/ViewportToolbar/LevelEditorSubmenus.cpp
@@ -1270,7 +1270,7 @@ FToolMenuEntry CreatePIEViewModesSubmenu()
 				TAttribute<FText> LabelAttribute = TAttribute<FText>::CreateLambda(
 					[WeakViewport = Context->LevelViewport]
 					{
-						if (TSharedPtr<SLevelViewport> Viewport = WeakViewport.Pin())
+						if (TSharedPtr<::SLevelViewport> Viewport = WeakViewport.Pin())
 						{
 							return UE::UnrealEd::GetViewModesSubmenuLabel(Viewport->GetPlayClient());
 						}
@@ -1281,7 +1281,7 @@ FToolMenuEntry CreatePIEViewModesSubmenu()
 				TAttribute<FSlateIcon> IconAttribute = TAttribute<FSlateIcon>::CreateLambda(
 					[WeakViewport = Context->LevelViewport]
 					{
-						if (TSharedPtr<SLevelViewport> Viewport = WeakViewport.Pin())
+						if (TSharedPtr<::SLevelViewport> Viewport = WeakViewport.Pin())
 						{
 							return UE::UnrealEd::GetViewModesSubmenuIcon(Viewport->GetPlayClient());
 						}
