[Immersive Mode] Maximized viewport scale is not considering Application Scale setting

diff --git a/Engine/Source/Editor/LevelEditor/Private/LevelViewportLayout.cpp b/Engine/Source/Editor/LevelEditor/Private/LevelViewportLayout.cpp
--- a/Engine/Source/Editor/LevelEditor/Private/LevelViewportLayout.cpp
+++ b/Engine/Source/Editor/LevelEditor/Private/LevelViewportLayout.cpp
@@ -502,13 +502,17 @@ FVector2D FLevelViewportLayout::GetMaximizedViewportSizeOnCanvas() const
 			{
 				FVector2D ClippedArea = FVector2D::ZeroVector;
 
+				const float AppScale = FSlateApplication::Get().GetApplicationScale();
+
 				if( OwnerWindow->IsWindowMaximized() )
 				{
 					// When the window is maximized and we are in immersive we size the canvas to the size of the visible area which does not include the window border
-					const FMargin& WindowContentMargin = OwnerWindow->GetWindowBorderSize();
+					const FMargin& WindowContentMargin = OwnerWindow->GetWindowBorderSize() * AppScale;
 					ClippedArea.Set( WindowContentMargin.GetTotalSpaceAlong<Orient_Horizontal>(), WindowContentMargin.GetTotalSpaceAlong<Orient_Vertical>() );
 				}
-				TargetSize = (OwnerWindow->GetSizeInScreen() - ClippedArea)/OwnerWindow->GetNativeWindow()->GetDPIScaleFactor();
+
+				const float ScaleFactor = OwnerWindow->GetNativeWindow()->GetDPIScaleFactor() * AppScale;
+				TargetSize = (OwnerWindow->GetSizeInScreen() - ClippedArea)/ScaleFactor;
 			}
 		}
 		else
