Fix backwards compatability issues with extending legacy level editor viewport toolbar menus.

Customize the positioning of the Revision Control menu in the new toolbar.

diff --git a/Engine/Source/Developer/SourceControlViewport/Private/SourceControlViewportMenu.cpp b/Engine/Source/Developer/SourceControlViewport/Private/SourceControlViewportMenu.cpp
--- a/Engine/Source/Developer/SourceControlViewport/Private/SourceControlViewportMenu.cpp
+++ b/Engine/Source/Developer/SourceControlViewport/Private/SourceControlViewportMenu.cpp
@@ -7,6 +7,7 @@
 #include "LevelEditorMenuContext.h"
 #include "LevelEditorViewport.h"
 #include "ToolMenus.h"
+#include "ViewportToolbar/UnrealEdViewportToolbar.h"
 
 #define LOCTEXT_NAMESPACE "SourceControlViewportMenu"
 
@@ -97,7 +98,10 @@ void FSourceControlViewportMenu::PopulateViewportMenu(UToolMenu* InMenu)
 				.MaxValue(100)
 				.MaxSliderValue(100);
 
-			FToolMenuSection& RevisionControlSection = InMenu->FindOrAddSection(TEXT("LevelViewportEditorShow"));
+			FToolMenuSection& RevisionControlSection = UE::UnrealEd::ShowNewViewportToolbars()
+				? InMenu->FindOrAddSection(TEXT("AllShowFlags"))
+				: InMenu->FindOrAddSection(TEXT("LevelViewportEditorShow"));
+			
 			RevisionControlSection.AddDynamicEntry(NAME_None, FNewToolMenuSectionDelegate::CreateLambda(
 				[this, ViewportClient](FToolMenuSection& InSection)
 				{
diff --git a/Engine/Source/Editor/LevelEditor/Private/SLevelViewportToolBar.cpp b/Engine/Source/Editor/LevelEditor/Private/SLevelViewportToolBar.cpp
--- a/Engine/Source/Editor/LevelEditor/Private/SLevelViewportToolBar.cpp
+++ b/Engine/Source/Editor/LevelEditor/Private/SLevelViewportToolBar.cpp
@@ -430,8 +430,14 @@ TSharedRef<SWidget> SLevelViewportToolBar::GenerateOptionsMenu()
 		UToolMenu* Menu = UToolMenus::Get()->RegisterMenu(MenuName);
 		Menu->AddDynamicSection("DynamicSection", FNewToolMenuDelegate::CreateLambda([](UToolMenu* InMenu)
 		{
-			ULevelViewportToolBarContext* Context = InMenu->FindContext<ULevelViewportToolBarContext>();
-			Context->LevelViewportToolBarWidget.Pin()->FillOptionsMenu(InMenu);
+			if (ULevelViewportToolBarContext* Context = InMenu->FindContext<ULevelViewportToolBarContext>())
+			{
+				if (TSharedPtr<SLevelViewportToolBar> ToolBar = Context->LevelViewportToolBarWidget.Pin())
+                {
+                	ToolBar->FillOptionsMenu(InMenu);
+                }
+			}
+			
 		}));
 	}
 
@@ -576,8 +582,13 @@ TSharedRef<SWidget> SLevelViewportToolBar::GenerateDevicePreviewMenu() const
 		UToolMenu* Menu = UToolMenus::Get()->RegisterMenu(MenuName);
 		Menu->AddDynamicSection("DynamicSection", FNewToolMenuDelegate::CreateLambda([](UToolMenu* InMenu)
 		{
-			ULevelViewportToolBarContext* Context = InMenu->FindContext<ULevelViewportToolBarContext>();
-			Context->LevelViewportToolBarWidget.Pin()->FillDevicePreviewMenu(InMenu);
+			if (ULevelViewportToolBarContext* Context = InMenu->FindContext<ULevelViewportToolBarContext>())
+			{
+				if (TSharedPtr<SLevelViewportToolBar> ToolBar = Context->LevelViewportToolBarWidget.Pin())
+				{
+					ToolBar->FillDevicePreviewMenu(InMenu);
+				}
+			}
 		}));
 	}
 
@@ -855,7 +866,10 @@ TSharedRef<SWidget> SLevelViewportToolBar::GenerateShowMenu() const
 				{
 					if (ULevelViewportToolBarContext* Context = InMenu->FindContext<ULevelViewportToolBarContext>())
 					{
-						Context->LevelViewportToolBarWidget.Pin()->FillShowMenu(InMenu);
+						if (TSharedPtr<SLevelViewportToolBar> ToolBar = Context->LevelViewportToolBarWidget.Pin())
+						{
+							ToolBar->FillShowMenu(InMenu);
+						}
 					}
 				}
 			)
@@ -1066,13 +1080,13 @@ TSharedRef<SWidget> SLevelViewportToolBar::GetScalabilityWarningMenuContent() co
 		];
 }
 
-FLevelEditorViewportClient* ULevelViewportToolBarContext::GetLevelViewportClient()
+FLevelEditorViewportClient* ULevelViewportToolBarContext::GetLevelViewportClient() const
 {
-	if (LevelViewportToolBarWidget.IsValid())
+	if (TSharedPtr<SLevelViewportToolBar> Toolbar = LevelViewportToolBarWidget.Pin())
 	{
-		return LevelViewportToolBarWidget.Pin()->GetLevelViewportClient();
+		return Toolbar->GetLevelViewportClient();
 	}
-
+	
 	return nullptr;
 }
 
diff --git a/Engine/Source/Editor/LevelEditor/Private/ViewportToolbar/LevelViewportContext.cpp b/Engine/Source/Editor/LevelEditor/Private/ViewportToolbar/LevelViewportContext.cpp
new file mode 100644
--- /dev/null
+++ b/Engine/Source/Editor/LevelEditor/Private/ViewportToolbar/LevelViewportContext.cpp
@@ -0,0 +1,15 @@
+﻿// Copyright Epic Games, Inc. All Rights Reserved.
+
+#include "LevelViewportContext.h"
+
+#include "Editor/LevelEditor/Private/LevelEditorInternalTools.h"
+
+FLevelEditorViewportClient* ULevelViewportContext::GetLevelViewportClient() const
+{
+	if (TSharedPtr<SLevelViewport> Viewport = LevelViewport.Pin())
+	{
+		return &Viewport->GetLevelViewportClient();
+	}
+	
+	return Super::GetLevelViewportClient();
+}
diff --git a/Engine/Source/Editor/LevelEditor/Private/ViewportToolbar/LevelViewportContext.h b/Engine/Source/Editor/LevelEditor/Private/ViewportToolbar/LevelViewportContext.h
--- a/Engine/Source/Editor/LevelEditor/Private/ViewportToolbar/LevelViewportContext.h
+++ b/Engine/Source/Editor/LevelEditor/Private/ViewportToolbar/LevelViewportContext.h
@@ -4,15 +4,18 @@
 
 #include "Templates/SharedPointer.h"
 #include "SLevelViewport.h"
-#include "UObject/Object.h"
+#include "LevelEditorMenuContext.h"
 
 #include "LevelViewportContext.generated.h"
 
+// Inherits from ULevelViewportToolBarContext for backwards compatability.
 UCLASS()
-class LEVELEDITOR_API ULevelViewportContext : public UObject
+class LEVELEDITOR_API ULevelViewportContext : public ULevelViewportToolBarContext
 {
 	GENERATED_BODY()
 
 public:
 	TWeakPtr<SLevelViewport> LevelViewport;
+	
+	virtual FLevelEditorViewportClient* GetLevelViewportClient() const override;
 };
diff --git a/Engine/Source/Editor/LevelEditor/Public/LevelEditorMenuContext.h b/Engine/Source/Editor/LevelEditor/Public/LevelEditorMenuContext.h
--- a/Engine/Source/Editor/LevelEditor/Public/LevelEditorMenuContext.h
+++ b/Engine/Source/Editor/LevelEditor/Public/LevelEditorMenuContext.h
@@ -87,7 +87,7 @@ class LEVELEDITOR_API ULevelViewportToolBarContext : public UObject
 public:
 	TWeakPtr<SLevelViewportToolBar> LevelViewportToolBarWidget;
 
-	FLevelEditorViewportClient* GetLevelViewportClient();
+	virtual FLevelEditorViewportClient* GetLevelViewportClient() const;
 };
 
 
