Fix issue where legacy extenders were not rendered in the new viewport toolbar.
- Hide legacy extensions that have already been migrated in the new toolbar menu.

diff --git a/Engine/Plugins/Experimental/DaySequence/Source/DaySequenceEditor/Private/DaySequenceEditorModule.cpp b/Engine/Plugins/Experimental/DaySequence/Source/DaySequenceEditor/Private/DaySequenceEditorModule.cpp
--- a/Engine/Plugins/Experimental/DaySequence/Source/DaySequenceEditor/Private/DaySequenceEditorModule.cpp
+++ b/Engine/Plugins/Experimental/DaySequence/Source/DaySequenceEditor/Private/DaySequenceEditorModule.cpp
@@ -43,6 +43,7 @@
 #include "Widgets/SNullWidget.h"
 #include "WorkspaceMenuStructure.h"
 #include "WorkspaceMenuStructureModule.h"
+#include "ViewportToolbar/UnrealEdViewportToolbar.h"
 
 #define LOCTEXT_NAMESPACE "DaySequenceEditor"
 
@@ -629,7 +630,7 @@ TSharedRef<SWidget> FDaySequenceEditorModule::CreateDaySequenceViewportToolbarEx
 			})
 			.Visibility_Lambda([this]()
 			{
-				return IsMenuVisible() ? EVisibility::Visible : EVisibility::Collapsed;
+				return IsMenuVisible() && UE::UnrealEd::ShowOldViewportToolbars() ? EVisibility::Visible : EVisibility::Collapsed;
 			});
 	}
 	return SNullWidget::NullWidget;
diff --git a/Engine/Source/Editor/LevelEditor/Private/SLevelViewport.cpp b/Engine/Source/Editor/LevelEditor/Private/SLevelViewport.cpp
--- a/Engine/Source/Editor/LevelEditor/Private/SLevelViewport.cpp
+++ b/Engine/Source/Editor/LevelEditor/Private/SLevelViewport.cpp
@@ -111,6 +111,8 @@
 #include "SActionableMessageViewportWidget.h"
 #include "ShowFlagMenuCommands.h"
 #include "SkeletalRenderPublic.h"
+#include "ViewportToolBarContext.h"
+#include "Subsystems/PanelExtensionSubsystem.h"
 #include "ViewportToolbar/LevelEditorSubmenus.h"
 #include "ViewportToolbar/LevelViewportContext.h"
 #include "ViewportToolbar/UnrealEdViewportToolbar.h"
@@ -2158,7 +2160,7 @@ TSharedRef<FEditorViewportClient> SLevelViewport::MakeEditorViewportClient()
 TSharedPtr<SWidget> SLevelViewport::MakeViewportToolbar()
 {
 	const TSharedRef<SLevelViewportToolBar> OldViewportToolbar =
-		SNew(SLevelViewportToolBar)
+		SAssignNew(LevelViewportToolBar, SLevelViewportToolBar)
 			.Viewport(SharedThis(this))
 			.Visibility_Lambda(
 				[this]() -> EVisibility
@@ -2311,6 +2313,39 @@ TSharedPtr<SWidget> SLevelViewport::BuildViewportToolbar()
 				ViewportToolbarMenu->AddSection("Layer2DSnap").Alignment = EToolMenuSectionAlign::Last;
 				ViewportToolbarMenu->AddSection("ScaleGridSnap").Alignment = EToolMenuSectionAlign::Last;
 				ViewportToolbarMenu->AddSection("CameraSpeed").Alignment = EToolMenuSectionAlign::Last;
+				
+				FToolMenuSection& FallbackSection = ViewportToolbarMenu->AddSection("Fallback");
+				FallbackSection.Alignment = EToolMenuSectionAlign::Last;
+				
+				FallbackSection.AddDynamicEntry("FallbackExtensionPanels", FNewToolMenuSectionDelegate::CreateLambda([](FToolMenuSection& Section)
+				{
+					ULevelViewportContext* Context = Section.FindContext<ULevelViewportContext>();
+					if (!Context)
+					{
+						return;
+					}
+					
+					UViewportToolBarContext* ExtensionContext = NewObject<UViewportToolBarContext>();
+					ExtensionContext->Viewport = Context->LevelViewport;
+					ExtensionContext->ViewportToolBar = Context->LevelViewportToolBarWidget;
+					
+					auto MakeExtensionPanel = [ExtensionContext, &Section](FName ExtensionId)
+					{
+						Section.AddEntry(FToolMenuEntry::InitWidget(
+							ExtensionId,
+							SNew(SExtensionPanel)
+								.ExtensionPanelID(ExtensionId)
+								.ExtensionContext(ExtensionContext),
+							FText::GetEmpty(),
+							true, false, true
+						));
+					};
+					
+					MakeExtensionPanel("LevelViewportToolBar.LeftExtension");
+					MakeExtensionPanel("LevelViewportToolBar.MiddleExtension");
+					MakeExtensionPanel("LevelViewportToolBar.RightExtension");
+					MakeExtensionPanel("LevelViewportToolBar.RightmostExtension");
+				}));
 			}
 			
 			{
@@ -2387,6 +2422,8 @@ TSharedPtr<SWidget> SLevelViewport::BuildViewportToolbar()
 		{
 			ULevelViewportContext* const ContextObject = NewObject<ULevelViewportContext>();
 			ContextObject->LevelViewport = SharedThis(this);
+			// Provide a non-null toolbar widget so that legacy extensions with bad null handling don't crash
+			ContextObject->LevelViewportToolBarWidget = LevelViewportToolBar;
 			ViewportToolbarContext.AddObject(ContextObject);
 		}
 
diff --git a/Engine/Source/Editor/LevelEditor/Private/SLevelViewportToolBar.cpp b/Engine/Source/Editor/LevelEditor/Private/SLevelViewportToolBar.cpp
--- a/Engine/Source/Editor/LevelEditor/Private/SLevelViewportToolBar.cpp
+++ b/Engine/Source/Editor/LevelEditor/Private/SLevelViewportToolBar.cpp
@@ -329,9 +329,9 @@ bool SLevelViewportToolBar::IsViewModeSupported(EViewModeIndex ViewModeIndex) co
 
 FLevelEditorViewportClient* SLevelViewportToolBar::GetLevelViewportClient() const
 {
-	if (Viewport.IsValid())
+	if (TSharedPtr<SLevelViewport> PinnedViewport = Viewport.Pin())
 	{
-		return &Viewport.Pin()->GetLevelViewportClient();
+		return &PinnedViewport->GetLevelViewportClient();
 	}
 
 	return nullptr;
diff --git a/Engine/Source/Editor/LevelEditor/Public/SLevelViewport.h b/Engine/Source/Editor/LevelEditor/Public/SLevelViewport.h
--- a/Engine/Source/Editor/LevelEditor/Public/SLevelViewport.h
+++ b/Engine/Source/Editor/LevelEditor/Public/SLevelViewport.h
@@ -21,6 +21,7 @@
 #include "IAssetViewport.h"
 #include "LevelEditorViewport.h"
 
+class SLevelViewportToolBar;
 class FLevelEditorViewportClient;
 class FLevelViewportLayout;
 class FSceneViewport;
@@ -892,6 +893,8 @@ private:
 	 */
 	FDelegateHandle PIEPreviewFeatureLevelChangedHandle;
 
+	TSharedPtr<SLevelViewportToolBar> LevelViewportToolBar;
+
 	/** Level viewport client */
 	TSharedPtr<FLevelEditorViewportClient> LevelViewportClient;
 	
