[LevelEditor]
* readded dll export on FLevelViewportCommands

diff --git a/Engine/Source/Editor/LevelEditor/Public/LevelViewportActions.h b/Engine/Source/Editor/LevelEditor/Public/LevelViewportActions.h
--- a/Engine/Source/Editor/LevelEditor/Public/LevelViewportActions.h
+++ b/Engine/Source/Editor/LevelEditor/Public/LevelViewportActions.h
@@ -38,7 +38,7 @@ namespace LevelViewportConfigurationNames
 /**
  * Class containing commands for level viewport actions
  */
-class FLevelViewportCommands : public TCommands<FLevelViewportCommands>
+class LEVELEDITOR_API FLevelViewportCommands : public TCommands<FLevelViewportCommands>
 {
 public:
 
