LevelEditor: ActionCallbacks now check if the operation is supported when checking if saving is possible.

diff --git a/Engine/Source/Editor/LevelEditor/Private/LevelEditorActions.cpp b/Engine/Source/Editor/LevelEditor/Private/LevelEditorActions.cpp
--- a/Engine/Source/Editor/LevelEditor/Private/LevelEditorActions.cpp
+++ b/Engine/Source/Editor/LevelEditor/Private/LevelEditorActions.cpp
@@ -448,7 +448,7 @@ bool FLevelEditorActionCallbacks::ToggleFavorite_IsChecked()
 
 bool FLevelEditorActionCallbacks::CanSaveWorld()
 {
-	return FSlateApplication::Get().IsNormalExecution() && (!GUnrealEd || !GUnrealEd->GetPackageAutoSaver().IsAutoSaving());
+	return FSlateApplication::Get().IsNormalExecution() && (!GUnrealEd || !GUnrealEd->GetPackageAutoSaver().IsAutoSaving()) && GLevelEditorModeTools().IsOperationSupportedForCurrentAsset(EAssetOperation::Save);
 }
 
 bool FLevelEditorActionCallbacks::CanSaveUnpartitionedWorld()
