Improve reliability of legacy SExtensionPanel support in the new viewport toolbar. Extensions are only enabled when there is enough information to support their existance.
- Enable legacy SExtensionPannel support in the Cinematic Viewport Toolbar.

diff --git a/Engine/Plugins/MovieScene/LevelSequenceEditor/Source/LevelSequenceEditor/Private/CinematicViewport/SCinematicLevelViewport.cpp b/Engine/Plugins/MovieScene/LevelSequenceEditor/Source/LevelSequenceEditor/Private/CinematicViewport/SCinematicLevelViewport.cpp
--- a/Engine/Plugins/MovieScene/LevelSequenceEditor/Source/LevelSequenceEditor/Private/CinematicViewport/SCinematicLevelViewport.cpp
+++ b/Engine/Plugins/MovieScene/LevelSequenceEditor/Source/LevelSequenceEditor/Private/CinematicViewport/SCinematicLevelViewport.cpp
@@ -346,6 +346,9 @@ public:
 		ExtendLevelToolbar();
 		ExtendOldLevelToolbar();
 
+		TSharedRef<SWidget> LegacyToolbar = SLevelViewport::MakeViewportToolbar().ToSharedRef();
+		TSharedRef<SWidget> Toolbar = SLevelViewport::BuildViewportToolbar().ToSharedRef();
+
 		// clang-format off
 		TSharedRef<SVerticalBox> ToolbarsWidget =
 			SNew(SVerticalBox)
@@ -360,7 +363,7 @@ public:
 					}
 				)
 				[
-					SLevelViewport::BuildViewportToolbar().ToSharedRef()
+					Toolbar
 				]
 			]
 			+ SVerticalBox::Slot()
@@ -381,7 +384,7 @@ public:
 					}
 				)
 				[
-					SLevelViewport::MakeViewportToolbar().ToSharedRef()
+					LegacyToolbar
 				]
 			];
 		// clang-format on
diff --git a/Engine/Source/Editor/LevelEditor/Private/SLevelViewport.cpp b/Engine/Source/Editor/LevelEditor/Private/SLevelViewport.cpp
--- a/Engine/Source/Editor/LevelEditor/Private/SLevelViewport.cpp
+++ b/Engine/Source/Editor/LevelEditor/Private/SLevelViewport.cpp
@@ -2319,8 +2319,8 @@ TSharedPtr<SWidget> SLevelViewport::BuildViewportToolbar()
 				
 				FallbackSection.AddDynamicEntry("FallbackExtensionPanels", FNewToolMenuSectionDelegate::CreateLambda([](FToolMenuSection& Section)
 				{
-					ULevelViewportContext* Context = Section.FindContext<ULevelViewportContext>();
-					if (!Context)
+					ULegacyLevelViewportToolbarContext* Context = Section.FindContext<ULegacyLevelViewportToolbarContext>();
+					if (!Context || !Context->LevelViewport.IsValid() || !Context->LevelViewportToolBarWidget.IsValid())
 					{
 						return;
 					}
@@ -2422,10 +2422,18 @@ TSharedPtr<SWidget> SLevelViewport::BuildViewportToolbar()
 		{
 			ULevelViewportContext* const ContextObject = NewObject<ULevelViewportContext>();
 			ContextObject->LevelViewport = SharedThis(this);
-			// Provide a non-null toolbar widget so that legacy extensions with bad null handling don't crash
-			ContextObject->LevelViewportToolBarWidget = LevelViewportToolBar;
 			ViewportToolbarContext.AddObject(ContextObject);
 		}
+		
+		// Support legacy SExtensionPanel extensions
+		if (LevelViewportToolBar)
+		{
+			// Provide a non-null toolbar widget so that legacy extensions with bad null handling don't crash
+			ULegacyLevelViewportToolbarContext* const LegacyContext = NewObject<ULegacyLevelViewportToolbarContext>();
+			LegacyContext->LevelViewport = SharedThis(this);
+			LegacyContext->LevelViewportToolBarWidget = LevelViewportToolBar;
+			ViewportToolbarContext.AddObject(LegacyContext);
+		}
 
 		ExtendToolbarContext(ViewportToolbarContext);
 	}
diff --git a/Engine/Source/Editor/LevelEditor/Private/ViewportToolbar/LevelViewportContext.cpp b/Engine/Source/Editor/LevelEditor/Private/ViewportToolbar/LevelViewportContext.cpp
--- a/Engine/Source/Editor/LevelEditor/Private/ViewportToolbar/LevelViewportContext.cpp
+++ b/Engine/Source/Editor/LevelEditor/Private/ViewportToolbar/LevelViewportContext.cpp
@@ -4,7 +4,7 @@
 
 #include "Editor/LevelEditor/Private/LevelEditorInternalTools.h"
 
-FLevelEditorViewportClient* ULevelViewportContext::GetLevelViewportClient() const
+FLevelEditorViewportClient* ULegacyLevelViewportToolbarContext::GetLevelViewportClient() const
 {
 	if (TSharedPtr<SLevelViewport> Viewport = LevelViewport.Pin())
 	{
diff --git a/Engine/Source/Editor/LevelEditor/Private/ViewportToolbar/LevelViewportContext.h b/Engine/Source/Editor/LevelEditor/Private/ViewportToolbar/LevelViewportContext.h
--- a/Engine/Source/Editor/LevelEditor/Private/ViewportToolbar/LevelViewportContext.h
+++ b/Engine/Source/Editor/LevelEditor/Private/ViewportToolbar/LevelViewportContext.h
@@ -8,12 +8,22 @@
 
 #include "LevelViewportContext.generated.h"
 
-// Inherits from ULevelViewportToolBarContext for backwards compatability.
 UCLASS()
-class LEVELEDITOR_API ULevelViewportContext : public ULevelViewportToolBarContext
+class LEVELEDITOR_API ULevelViewportContext : public UObject
 {
 	GENERATED_BODY()
 
+public:
+	TWeakPtr<SLevelViewport> LevelViewport;
+};
+
+// Customizes ULevelViewportToolBarContext for backwards compatability
+// TODO: Remove this once the old toolbars are deprecated.
+UCLASS()
+class LEVELEDITOR_API ULegacyLevelViewportToolbarContext : public ULevelViewportToolBarContext
+{
+	GENERATED_BODY()
+	
 public:
 	TWeakPtr<SLevelViewport> LevelViewport;
 	
