Restore the ability to toggle the visibility of viewport toolbars.
- Add a separate value for toolbar visibility in PIE and in Immersive Mode.

#rnx

diff --git a/Engine/Source/Editor/LevelEditor/Private/LevelEditor.cpp b/Engine/Source/Editor/LevelEditor/Private/LevelEditor.cpp
--- a/Engine/Source/Editor/LevelEditor/Private/LevelEditor.cpp
+++ b/Engine/Source/Editor/LevelEditor/Private/LevelEditor.cpp
@@ -1791,11 +1791,18 @@ void FLevelEditorModule::BindGlobalLevelEditorCommands()
 		FCanExecuteAction(),
 		FIsActionChecked::CreateStatic(FLevelEditorActionCallbacks::ScaleGridSnap_IsChecked)
 		);
+		
+	ActionList.MapAction(
+		FLevelViewportCommands::Get().ToggleViewportToolbar,
+		FExecuteAction::CreateStatic( &FLevelEditorActionCallbacks::OnToggleShowViewportToolbar),
+		FCanExecuteAction(),
+		FIsActionChecked::CreateStatic( &FLevelEditorActionCallbacks::IsViewportToolbarVisible )
+		);
 	ActionList.MapAction(
 		Commands.ToggleHideViewportUI,
-		FExecuteAction::CreateStatic( &FLevelEditorActionCallbacks::OnToggleHideViewportUI ),
+		FExecuteAction::CreateStatic( &FLevelEditorActionCallbacks::OnToggleShowViewportUI ),
 		FCanExecuteAction(),
-		FIsActionChecked::CreateStatic( &FLevelEditorActionCallbacks::IsViewportUIHidden ) 
+		FIsActionChecked::CreateStatic( &FLevelEditorActionCallbacks::IsViewportUIVisible ) 
 		);
 
 	ActionList.MapAction( 
diff --git a/Engine/Source/Editor/LevelEditor/Private/LevelEditorActions.cpp b/Engine/Source/Editor/LevelEditor/Private/LevelEditorActions.cpp
--- a/Engine/Source/Editor/LevelEditor/Private/LevelEditorActions.cpp
+++ b/Engine/Source/Editor/LevelEditor/Private/LevelEditorActions.cpp
@@ -3318,14 +3318,33 @@ void FLevelEditorActionCallbacks::SetActorSnapSetting(float Distance)
 	FSnappingUtils::EnableActorSnap( ( Distance > 0.0f ? true : false ) );
 }
 
-void FLevelEditorActionCallbacks::OnToggleHideViewportUI()
+void FLevelEditorActionCallbacks::OnToggleShowViewportToolbar()
+{
+	FLevelEditorModule& Module = FModuleManager::Get().GetModuleChecked<FLevelEditorModule>("LevelEditor");
+	if (TSharedPtr<SLevelViewport> Viewport = Module.GetFirstActiveLevelViewport())
+	{
+		Viewport->ToggleViewportToolbarVisibility();
+	}
+}
+
+bool FLevelEditorActionCallbacks::IsViewportToolbarVisible()
+{
+	FLevelEditorModule& Module = FModuleManager::Get().GetModuleChecked<FLevelEditorModule>("LevelEditor");
+	if (TSharedPtr<SLevelViewport> Viewport = Module.GetFirstActiveLevelViewport())
+	{
+		return Viewport->IsViewportToolbarVisible();
+	}
+	return false;
+}
+
+void FLevelEditorActionCallbacks::OnToggleShowViewportUI()
 {
 	GLevelEditorModeTools().SetHideViewportUI( !GLevelEditorModeTools().IsViewportUIHidden() );
 }
 
-bool FLevelEditorActionCallbacks::IsViewportUIHidden()
+bool FLevelEditorActionCallbacks::IsViewportUIVisible()
 {
-	return GLevelEditorModeTools().IsViewportUIHidden();
+	return !GLevelEditorModeTools().IsViewportUIHidden();
 }
 
 bool FLevelEditorActionCallbacks::IsEditorModeActive( FEditorModeID EditorMode )
@@ -4361,7 +4380,7 @@ void FLevelEditorCommands::RegisterCommands()
 	UI_COMMAND( LevelStreamingVolumePrevis, "Enable Automatic Level Streaming", "If enabled, the viewport will stream in levels automatically when the camera is moved", EUserInterfaceActionType::ToggleButton, FInputChord() );
 	UI_COMMAND( EnableActorSnap, "Actor", "If enabled, actors will snap to the location of other actors when they are within distance", EUserInterfaceActionType::ToggleButton, FInputChord(EModifierKey::Control | EModifierKey::Shift, EKeys::K) );
 	UI_COMMAND( EnableVertexSnap, "Vertex", "If enabled, actors will snap to the location of the nearest vertex on another actor in the direction of " "movement", EUserInterfaceActionType::ToggleButton, FInputChord() );
-	UI_COMMAND( ToggleHideViewportUI, "Hide Viewport UI", "Toggles hidden viewport UI mode. Hides all overlaid viewport UI widgets", EUserInterfaceActionType::ToggleButton, FInputChord() );
+	UI_COMMAND( ToggleHideViewportUI, "Show Viewport UI", "Sets the visibility of all overlaid viewport UI widgets.", EUserInterfaceActionType::ToggleButton, FInputChord() );
 
 	//if (FParse::Param( FCommandLine::Get(), TEXT( "editortoolbox" ) ))
 	//{
diff --git a/Engine/Source/Editor/LevelEditor/Private/LevelEditorMenu.cpp b/Engine/Source/Editor/LevelEditor/Private/LevelEditorMenu.cpp
--- a/Engine/Source/Editor/LevelEditor/Private/LevelEditorMenu.cpp
+++ b/Engine/Source/Editor/LevelEditor/Private/LevelEditorMenu.cpp
@@ -26,6 +26,7 @@
 #include "AssetSelection.h"
 #include "EditorBuildUtils.h"
 #include "EditorViewportCommands.h"
+#include "LevelViewportActions.h"
 #include "Styling/SlateIconFinder.h"
 #include "Toolkits/GlobalEditorCommonCommands.h"
 
@@ -254,7 +255,10 @@ void FLevelEditorMenu::RegisterLevelEditorMenus()
 				// See FMainMenu::RegisterWindowMenu().
 				LayoutSection.AddSeparator("FullscreenSeparator");
 
-				// Hide Viewport UI should appear before the ToggleFullscreen entry
+				FToolMenuEntry& ToggleViewportToolbar = LayoutSection.AddMenuEntry(FLevelViewportCommands::Get().ToggleViewportToolbar);
+				ToggleViewportToolbar.InsertPosition.Position = EToolMenuInsertType::Before;
+				ToggleViewportToolbar.InsertPosition.Name = "ToggleFullscreen";
+				
 				FToolMenuEntry& ToggleUIEntry = LayoutSection.AddMenuEntry(FLevelEditorCommands::Get().ToggleHideViewportUI);
 				ToggleUIEntry.InsertPosition.Position = EToolMenuInsertType::Before;
 				ToggleUIEntry.InsertPosition.Name = "ToggleFullscreen";
diff --git a/Engine/Source/Editor/LevelEditor/Private/LevelViewportActions.cpp b/Engine/Source/Editor/LevelEditor/Private/LevelViewportActions.cpp
--- a/Engine/Source/Editor/LevelEditor/Private/LevelViewportActions.cpp
+++ b/Engine/Source/Editor/LevelEditor/Private/LevelViewportActions.cpp
@@ -83,7 +83,7 @@ void FLevelViewportCommands::RegisterCommands()
 		FInputChord()
 	);
 
-	UI_COMMAND( ToggleViewportToolbar, "Show Toolbar", "Defines whether a toolbar should be displayed on this viewport", EUserInterfaceActionType::ToggleButton, FInputChord(EModifierKey::Control | EModifierKey::Shift, EKeys::T) );
+	UI_COMMAND( ToggleViewportToolbar, "Show Viewport Toolbar", "Sets the visibility of the viewport toolbar", EUserInterfaceActionType::ToggleButton, FInputChord(EModifierKey::Control | EModifierKey::Shift, EKeys::T) );
 
 	UI_COMMAND( ApplyMaterialToActor, "Apply Material", "Attempts to apply a dropped material to this object", EUserInterfaceActionType::Button, FInputChord() );
 
diff --git a/Engine/Source/Editor/LevelEditor/Private/SLevelViewport.cpp b/Engine/Source/Editor/LevelEditor/Private/SLevelViewport.cpp
--- a/Engine/Source/Editor/LevelEditor/Private/SLevelViewport.cpp
+++ b/Engine/Source/Editor/LevelEditor/Private/SLevelViewport.cpp
@@ -408,7 +408,7 @@ void SLevelViewport::ConstructViewportOverlayContent()
 				.Visibility_Lambda(
 					[this]()
 					{
-						if (!ShouldShowFullToolbar())
+						if (!IsViewportToolbarVisible())
 						{
 							return EVisibility::Collapsed;
 						}
@@ -679,7 +679,7 @@ void SLevelViewport::ConstructLevelEditorViewportClient(FLevelEditorViewportInst
 
 	LevelViewportClient->InitializeViewportInteraction();
 
-	bShowFullToolbar = ViewportInstanceSettings.bShowFullToolbar;
+	bShowEditorToolbar = ViewportInstanceSettings.bShowFullToolbar;
 
 	// Always set to true initially
 	bShowToolbarAndControls = true; 
@@ -723,6 +723,38 @@ FSceneViewport* SLevelViewport::GetGameSceneViewport() const
 	return ActiveViewport.Get();
 }
 
+void SLevelViewport::ToggleViewportToolbarVisibility()
+{
+	if (IsPlayInEditorViewportActive())
+	{
+		if (IsImmersive())
+		{
+			bShowImmersivePIEToolbar = !bShowImmersivePIEToolbar;
+		}
+		else
+		{
+			bShowPIEToolbar = !bShowPIEToolbar;
+			if (!bShowPIEToolbar)
+			{
+				bShowImmersivePIEToolbar = false;
+			}
+		}
+	}
+	else
+	{
+		bShowEditorToolbar = !bShowEditorToolbar;
+	}
+}
+
+bool SLevelViewport::IsViewportToolbarVisible() const
+{
+	if (IsPlayInEditorViewportActive())
+	{
+		return IsImmersive() ? bShowImmersivePIEToolbar : bShowPIEToolbar;
+	}
+	return bShowEditorToolbar;
+}
+
 void SLevelViewport::TransitionToPIE(bool bIsSimulating)
 {
 	for (FViewportActorPreview& ActorPreview : ActorPreviews)
@@ -776,7 +808,10 @@ FReply SLevelViewport::OnKeyDown( const FGeometry& MyGeometry, const FKeyEvent&
 		// Only process commands for pie when a play world is active
 		bool bHandled = FPlayWorldCommands::GlobalPlayWorldActions->ProcessCommandBindings( InKeyEvent );
 		
-		bHandled = bHandled || PIECommands->ProcessCommandBindings(InKeyEvent);
+		if (HasPlayInEditorViewport())
+		{
+			bHandled = bHandled || PIECommands->ProcessCommandBindings(InKeyEvent);
+		}
 
 		// Always handle commands in pie so they arent bubbled to editor only widgets
 		Reply = FReply::Handled();
@@ -1701,9 +1736,9 @@ void SLevelViewport::BindOptionCommands( FUICommandList& OutCommandList )
 
 	OutCommandList.MapAction(
 		ViewportActions.ToggleViewportToolbar,
-		FExecuteAction::CreateSP( this, &SLevelViewport::OnToggleShowFullToolbar ),
+		FExecuteAction::CreateSP( this, &SLevelViewport::ToggleViewportToolbarVisibility ),
 		FCanExecuteAction(),
-		FIsActionChecked::CreateSP( this, &SLevelViewport::ShouldShowFullToolbar )
+		FIsActionChecked::CreateSP( this, &SLevelViewport::IsViewportToolbarVisible )
 		);
 }
 
@@ -1878,6 +1913,13 @@ void SLevelViewport::BindStatCommand(const TSharedPtr<FUICommandInfo> InMenuItem
 
 void SLevelViewport::BindPIECommands(FUICommandList& OutCommandList)
 {
+	OutCommandList.MapAction(
+		FLevelViewportCommands::Get().ToggleViewportToolbar,
+		FExecuteAction::CreateSP( this, &SLevelViewport::ToggleViewportToolbarVisibility ),
+		FCanExecuteAction(),
+		FIsActionChecked::CreateSP( this, &SLevelViewport::IsViewportToolbarVisible )
+		);
+
 	OutCommandList.MapAction(FLevelViewportCommands::Get().UseDefaultShowFlags,
 		FExecuteAction::CreateSP(this, &SLevelViewport::ResetPIEShowFlags)
 	);
@@ -2009,6 +2051,46 @@ void SLevelViewport::RefreshPIEViewport()
 	}
 }
 
+TAutoConsoleVariable<int32> CVarLevelEditorPIEInputOverride(
+	TEXT("LevelEditorPIEInputOverride"),
+	1,
+	TEXT("Enable/disable overriding PIE gampeplay input with level editor PIE shorctuts.")
+);
+
+bool SLevelViewport::OnPIEViewportInputOverride(FInputKeyEventArgs& Input)
+{
+	if (CVarLevelEditorPIEInputOverride.GetValueOnGameThread() == 0)
+	{
+		return false;
+	}
+
+	if (Input.Event != IE_Pressed)
+	{
+		return false;
+	}
+	
+	FModifierKeysState Modifiers = FSlateApplication::Get().GetModifierKeys();
+	FInputChord CheckChord(Input.Key, EModifierKey::FromBools(
+		Modifiers.IsControlDown(),
+		Modifiers.IsAltDown(),
+		Modifiers.IsShiftDown(),
+		Modifiers.IsCommandDown()
+	));
+	
+	// This function will completely bypass any game input.
+	// The commands that can do this should be limited and carefully considered.
+	const TSharedPtr<FUICommandInfo>& Command = FLevelViewportCommands::Get().ToggleViewportToolbar;	
+	if (Command->HasActiveChord(CheckChord))
+	{
+		if (const FUIAction* Action = PIECommands->GetActionForCommand(Command))
+		{
+			return Action->Execute();
+		}
+	}
+	
+	return false;
+}
+
 const FSlateBrush* SLevelViewport::OnGetViewportBorderBrush() const
 {
 	const FSlateBrush* BorderBrush = nullptr;
@@ -2444,11 +2526,11 @@ TSharedPtr<SWidget> SLevelViewport::BuildViewportToolbar()
 	.Visibility_Lambda([this]
 	{
 		// TODO: Move this to GetToolBarVisibility() once the old toolbar is deprecated
-		if ((IsPlayInEditorViewportActive() && IsImmersive()) || GEngine->IsStereoscopic3D(ActiveViewport.Get()))
+		if (GEngine->IsStereoscopic3D(ActiveViewport.Get()))
 		{
 			return EVisibility::Collapsed;
 		}
-		return EVisibility::Visible;
+		return IsViewportToolbarVisible() ? EVisibility::Visible : EVisibility::Collapsed;
 	})
 	+ SWidgetSwitcher::Slot()
 	[
@@ -2614,29 +2696,11 @@ void SLevelViewport::OnPlaySettings()
 
 void SLevelViewport::OnToggleImmersive()
 {
-	if( ParentLayout.IsValid() ) 
-	{
-		bool bWantImmersive = !IsImmersive();
-		bool bWantMaximize = IsMaximized();
-
-		// We always want to animate in response to user-interactive toggling of maximized state
-		const bool bAllowAnimation = true;
-
-		FName ViewportName = ConfigKey;
-		if (!ViewportName.IsNone())
-		{
-			ParentLayout.Pin()->RequestMaximizeViewport( ViewportName, bWantMaximize, bWantImmersive, bAllowAnimation );
-		}
+	const bool bWantImmersive = !IsImmersive();
+	// We always want to animate in response to user-interactive toggling of maximized state
+	constexpr bool bAllowAnimation = true;
 
-		if (bWantImmersive && LevelViewportClient->IsVisualizeCalibrationMaterialEnabled()) 
-		{
-			bShowToolbarAndControls = false;
-		}
-		else
-		{
-			bShowToolbarAndControls = true;
-		}
-	}
+	MakeImmersive(bWantImmersive, bAllowAnimation);
 }
 
 void SLevelViewport::OnToggleSidebarTabs()
@@ -3018,7 +3082,7 @@ void SLevelViewport::SaveConfig(const FString& ConfigName) const
 	
 		ViewportInstanceSettings.bShowOnScreenStats = LevelViewportClient->ShouldShowStats();
 		ViewportInstanceSettings.FarViewPlane = LevelViewportClient->GetFarClipPlaneOverride();
-		ViewportInstanceSettings.bShowFullToolbar = bShowFullToolbar;
+		ViewportInstanceSettings.bShowFullToolbar = bShowEditorToolbar;
 
 		if(GetDefault<ULevelEditorViewportSettings>()->bSaveEngineStats)
 		{
@@ -3507,6 +3571,15 @@ void SLevelViewport::MakeImmersive( const bool bWantImmersive, const bool bAllow
 		{
 			ParentLayout.Pin()->RequestMaximizeViewport( ViewportName, bWantMaximize, bWantImmersive, bAllowAnimation );
 		}
+		
+		if (bWantImmersive && LevelViewportClient->IsVisualizeCalibrationMaterialEnabled()) 
+		{
+			bShowToolbarAndControls = false;
+		}
+		else
+		{
+			bShowToolbarAndControls = true;
+		}
 	}
 }
 
@@ -4797,7 +4870,31 @@ void SLevelViewport::StartPlayInEditorSession(UGameViewportClient* InPlayClient,
 	LevelViewportClient->PrepareCameraForPIE();
 
 	PlayClient = InPlayClient;
-
+	
+	// Set up binding overrides that can handle input before any player input
+	FOverrideInputKeyHandler& PlayInputKeyOverride = PlayClient->OnOverrideInputKey();
+	if (PlayInputKeyOverride.IsBound())
+	{
+		// Carry forward the previous binding and allow it to take precedence
+		FOverrideInputKeyHandler OldHandler = PlayInputKeyOverride;
+		PlayInputKeyOverride.BindLambda([WeakThis = SharedThis(this).ToWeakPtr(), OldHandler](FInputKeyEventArgs& Args)
+		{
+			if (OldHandler.IsBound() && OldHandler.Execute(Args))
+			{
+				return true;
+			}
+			if (TSharedPtr<SLevelViewport> Viewport = WeakThis.Pin())
+			{
+				return Viewport->OnPIEViewportInputOverride(Args);
+			}
+			return false;
+		});
+	}
+	else
+	{
+		PlayInputKeyOverride.BindSP(this, &SLevelViewport::OnPIEViewportInputOverride);
+	}
+	
 	// Here we will swap the editor viewport client out for the client for the play in editor session
 	InactiveViewport = ActiveViewport;
 	// Store the content in the viewport widget (editor tool bar etc) so we can show the game UI content if it has any
diff --git a/Engine/Source/Editor/LevelEditor/Public/LevelEditorActions.h b/Engine/Source/Editor/LevelEditor/Public/LevelEditorActions.h
--- a/Engine/Source/Editor/LevelEditor/Public/LevelEditorActions.h
+++ b/Engine/Source/Editor/LevelEditor/Public/LevelEditorActions.h
@@ -1352,8 +1352,11 @@ public:
 	static void OnEnableVertexSnap();
 	static bool OnIsVertexSnapEnabled();
 
-	static void OnToggleHideViewportUI();
-	static bool IsViewportUIHidden();
+	static void OnToggleShowViewportToolbar();
+	static bool IsViewportToolbarVisible();
+
+	static void OnToggleShowViewportUI();
+	static bool IsViewportUIVisible();
 
 	static bool IsEditorModeActive( FEditorModeID EditorMode );
 
diff --git a/Engine/Source/Editor/LevelEditor/Public/SLevelViewport.h b/Engine/Source/Editor/LevelEditor/Public/SLevelViewport.h
--- a/Engine/Source/Editor/LevelEditor/Public/SLevelViewport.h
+++ b/Engine/Source/Editor/LevelEditor/Public/SLevelViewport.h
@@ -416,7 +416,7 @@ public:
 	EVisibility GetOptionsMenuVisibility() const;
 
 	/** Get the visibility for items considered to be part of the 'full' viewport toolbar */
-	EVisibility GetFullToolbarVisibility() const { return (bShowToolbarAndControls && bShowFullToolbar) ? EVisibility::SelfHitTestInvisible : EVisibility::Collapsed; }
+	EVisibility GetFullToolbarVisibility() const { return (bShowToolbarAndControls && bShowEditorToolbar) ? EVisibility::SelfHitTestInvisible : EVisibility::Collapsed; }
 
 	/** Unpin and close all actor preview windows */
 	void RemoveAllPreviews(const bool bRemoveFromDesktopViewport = true);
@@ -466,6 +466,16 @@ public:
 	 */
 	FName GetConfigKey() const { return ConfigKey; }
 
+	/**
+	 * Toggles the current viewport toolbar visibility. 
+	 */
+	void ToggleViewportToolbarVisibility();
+
+	/**
+	 * Whether the viewport toolbar is currently visible.
+	 */
+	bool IsViewportToolbarVisible() const; 
+
 protected:
 	/** SEditorViewport interface */
 	virtual TSharedRef<FEditorViewportClient> MakeEditorViewportClient() override;
@@ -759,6 +769,11 @@ private:
 	 * Refreshes the PIE viewport.
 	 */
 	void RefreshPIEViewport();
+
+	/**
+	 * A handler for processing input events before gameplay does.
+	 */
+	bool OnPIEViewportInputOverride(FInputKeyEventArgs& Input);
 	
 	/**
 	 * Called to build the viewport context menu when the user is Drag Dropping from the Content Browser
@@ -836,13 +851,7 @@ private:
 
 	/** Gets the active scene viewport for the game */
 	FSceneViewport* GetGameSceneViewport() const;
-
-	/** Called when the user toggles the full toolbar */
-	void OnToggleShowFullToolbar() { bShowFullToolbar = !bShowFullToolbar; }
-
-	/** Check whether we should display the full toolbar or not */
-	bool ShouldShowFullToolbar() const { return bShowFullToolbar; }
-
+	
 	/** Handle any level viewport changes on entering PIE or simulate */
 	void TransitionToPIE(bool bIsSimulating);
 
@@ -1031,8 +1040,15 @@ private:
 	/** Whether to show toolbar or buttons */
 	bool bShowToolbarAndControls;
 
-	/** Whether to show a full toolbar, or a compact one */
-	bool bShowFullToolbar;
+	/** Whether to show the editor toolbar */
+	bool bShowEditorToolbar = true;
+	
+	/** Whether to show the in-PIE toolbar */
+	bool bShowPIEToolbar = true;
+	
+	/** Whether to show the in-PIE toolbar in immersive mode */
+	bool bShowImmersivePIEToolbar = false;
+	
 	TSharedPtr<class SWidget> InViewportMenuWrapper;
 	bool bIsInViewportMenuShowing;
 	bool bIsInViewportMenuInitialized;
