Adjust the viewport sizing menu to take maximized status into account.
- Make the single pane option maximize the current viewport.

#rnx

diff --git a/Engine/Source/Editor/EditorStyle/Private/StarshipStyle.cpp b/Engine/Source/Editor/EditorStyle/Private/StarshipStyle.cpp
--- a/Engine/Source/Editor/EditorStyle/Private/StarshipStyle.cpp
+++ b/Engine/Source/Editor/EditorStyle/Private/StarshipStyle.cpp
@@ -2831,13 +2831,13 @@ void FStarshipEditorStyle::FStyle::SetupViewportStyles()
 
 			// A raised button keeps the gray background here, with rounded corners on the right.
 			{
-				FButtonStyle ButtonStyleRaised = FButtonStyle(ButtonStyleEnd)
+				FCheckBoxStyle ToggleButtonStyleRaised = FCheckBoxStyle(ToggleButtonStyleEnd)
 					// The icon cannot be perfectly centered, so it needs a 1px shove to the right
 					// to match the design spec.
-					.SetNormalPadding(FMargin(7.0f, 4.0f, 6.0f, 4.0f))
-					.SetPressedPadding(FMargin(7.0f, 4.0f, 6.0f, 4.0f));
+					.SetPadding(FMargin(7.0f, 4.0f, 6.0f, 4.0f));
+					
 				const FToolBarStyle ViewportSizingRaisedStyle = FToolBarStyle(ViewportToolbarStyle)
-					.SetButtonStyle(ButtonStyleRaised)
+					.SetToggleButtonStyle(ToggleButtonStyleRaised)
 					.SetButtonPadding(FMargin(0.f, 0.f, 4.f, 0.f))
 					.SetComboButtonPadding(FMargin(2.0f, 0.f));
 
diff --git a/Engine/Source/Editor/LevelEditor/Private/LevelViewportTabContent.cpp b/Engine/Source/Editor/LevelEditor/Private/LevelViewportTabContent.cpp
--- a/Engine/Source/Editor/LevelEditor/Private/LevelViewportTabContent.cpp
+++ b/Engine/Source/Editor/LevelEditor/Private/LevelViewportTabContent.cpp
@@ -175,5 +175,16 @@ bool FLevelViewportTabContent::IsViewportTypeWithinLayoutEqual(FName InConfigKey
 
 bool FLevelViewportTabContent::IsViewportConfigurationChecked(FName InLayoutType) const
 {
-	return IsViewportConfigurationSet(InLayoutType);
+	// Only visually indicate that a configuration is set if the viewport is not currently maximized 
+	return IsViewportConfigurationSet(InLayoutType) && !IsLayoutMaximized();
 }
+
+bool FLevelViewportTabContent::IsLayoutMaximized() const
+{
+	if (TSharedPtr<FLevelViewportLayout> Layout = StaticCastSharedPtr<FLevelViewportLayout>(ActiveViewportLayout))
+	{
+		return Layout->IsMaximized();
+	}
+	return false;
+}
+
diff --git a/Engine/Source/Editor/LevelEditor/Private/SLevelViewport.cpp b/Engine/Source/Editor/LevelEditor/Private/SLevelViewport.cpp
--- a/Engine/Source/Editor/LevelEditor/Private/SLevelViewport.cpp
+++ b/Engine/Source/Editor/LevelEditor/Private/SLevelViewport.cpp
@@ -1775,6 +1775,14 @@ void SLevelViewport::BindViewCommands( FUICommandList& OutCommandList )
 		if (TSharedPtr<FEditorViewportTabContent> ParentTabContentPinned = LayoutPinned->GetParentTabContent().Pin())
 		{
 			ParentTabContentPinned->BindViewportLayoutCommands(OutCommandList, ConfigKey);
+			
+			// Override single panel configuration so that it acts like a one-way "maximize" button
+			OutCommandList.MapAction(
+				ViewportActions.ViewportConfig_OnePane,
+				FExecuteAction::CreateSP(this, &SLevelViewport::MakeMaximized, true, true),
+				FCanExecuteAction(),
+				FIsActionChecked::CreateSP(this, &SLevelViewport::IsMaximized)
+				);
 		}
 	}
 
@@ -3559,6 +3567,18 @@ FReply SLevelViewport::OnToggleMaximize()
 	return FReply::Handled();
 }
 
+void SLevelViewport::MakeMaximized(bool bWantMaximized, bool bAllowAnimation)
+{
+	if (TSharedPtr<FLevelViewportLayout> Layout = ParentLayout.Pin())
+	{
+		const FName ViewportName = ConfigKey;
+		if (!ViewportName.IsNone())
+		{
+			const bool bWantImmersive = IsImmersive();
+			Layout->RequestMaximizeViewport( ViewportName, bWantMaximized, bWantImmersive, bAllowAnimation );
+		}
+	}
+}
 
 void SLevelViewport::MakeImmersive( const bool bWantImmersive, const bool bAllowAnimation )
 {
diff --git a/Engine/Source/Editor/LevelEditor/Private/ViewportToolbar/LevelEditorSubmenus.cpp b/Engine/Source/Editor/LevelEditor/Private/ViewportToolbar/LevelEditorSubmenus.cpp
--- a/Engine/Source/Editor/LevelEditor/Private/ViewportToolbar/LevelEditorSubmenus.cpp
+++ b/Engine/Source/Editor/LevelEditor/Private/ViewportToolbar/LevelEditorSubmenus.cpp
@@ -2750,41 +2750,8 @@ FToolMenuEntry CreateViewportSizingSubmenu()
 										);
 									}
 								);
-
-								FToolUIAction Action;
-								Action.ExecuteAction = FToolMenuExecuteAction::CreateLambda(
-									[](const FToolMenuContext& Context) -> void
-									{
-										if (ULevelViewportContext* const LevelViewportContext =
-												Context.FindContext<ULevelViewportContext>())
-										{
-											if (const TSharedPtr<::SLevelViewport> LevelViewport =
-													LevelViewportContext->LevelViewport.Pin())
-											{
-												LevelViewport->OnToggleMaximize();
-											}
-										}
-									}
-								);
-								Action.CanExecuteAction = FToolMenuCanExecuteAction::CreateLambda(
-									[](const FToolMenuContext& Context) -> bool
-									{
-										if (ULevelViewportContext* const LevelViewportContext =
-												Context.FindContext<ULevelViewportContext>())
-										{
-											if (const TSharedPtr<::SLevelViewport> LevelViewport =
-													LevelViewportContext->LevelViewport.Pin())
-											{
-												return LevelViewport->CanMaximize();
-											}
-										}
-
-										return false;
-									}
-								);
-
-								FToolMenuEntry& MaximizeRestoreEntry =
-									InnerSection.AddMenuEntry("MaximizeRestore", Label, Tooltip, Icon, Action);
+								
+								FToolMenuEntry& MaximizeRestoreEntry = InnerSection.AddMenuEntry(FLevelViewportCommands::Get().ToggleMaximize, Label, Tooltip, Icon);
 								MaximizeRestoreEntry.SetShowInToolbarTopLevel(true);
 								MaximizeRestoreEntry.ToolBarData.ResizeParams.AllowClipping = false;
 								MaximizeRestoreEntry.StyleNameOverride = FName("ViewportToolbarViewportSizingSubmenu");
diff --git a/Engine/Source/Editor/LevelEditor/Public/LevelViewportLayout.h b/Engine/Source/Editor/LevelEditor/Public/LevelViewportLayout.h
--- a/Engine/Source/Editor/LevelEditor/Public/LevelViewportLayout.h
+++ b/Engine/Source/Editor/LevelEditor/Public/LevelViewportLayout.h
@@ -102,6 +102,11 @@ public:
 	*/
 	bool IsMaximizeSupported() const { return bIsMaximizeSupported; }
 
+	/**
+	 * @return true if any viewport is maximized. 
+	 */
+	bool IsMaximized() const { return bIsMaximized;}
+
 	/** 
 	 * Checks to see if the specified level viewport is maximized
 	 *
diff --git a/Engine/Source/Editor/LevelEditor/Public/LevelViewportTabContent.h b/Engine/Source/Editor/LevelEditor/Public/LevelViewportTabContent.h
--- a/Engine/Source/Editor/LevelEditor/Public/LevelViewportTabContent.h
+++ b/Engine/Source/Editor/LevelEditor/Public/LevelViewportTabContent.h
@@ -36,4 +36,5 @@ private:
 	void OnUIActionSetViewportTypeWithinLayout(FName InConfigKey, FName InLayoutType);
 	bool IsViewportTypeWithinLayoutEqual(FName InConfigName, FName InLayoutType) const;
 	bool IsViewportConfigurationChecked(FName InLayoutType) const;
+	bool IsLayoutMaximized() const;
 };
diff --git a/Engine/Source/Editor/LevelEditor/Public/SLevelViewport.h b/Engine/Source/Editor/LevelEditor/Public/SLevelViewport.h
--- a/Engine/Source/Editor/LevelEditor/Public/SLevelViewport.h
+++ b/Engine/Source/Editor/LevelEditor/Public/SLevelViewport.h
@@ -152,6 +152,13 @@ public:
 	 */
 	bool CanMaximize() const;
 
+	/**
+	 * Attempt to switch this viewport into a maximized or normal mode
+	 * @param bWantMaximized Whether to witch to maximized or normal
+	 * @param bAllowAnimation True to allow animation when transitioning, otherwise false
+	 */
+	void MakeMaximized(bool bWantMaximized, bool bAllowAnimation = true);
+
 	/**
 	 * Attempts to switch this viewport into immersive mode
 	 *
