Differentiate PIE LOCTEXT keys for viewport toolbar menus.

#rnx

diff --git a/Engine/Source/Editor/LevelEditor/Private/ViewportToolbar/LevelEditorSubmenus.cpp b/Engine/Source/Editor/LevelEditor/Private/ViewportToolbar/LevelEditorSubmenus.cpp
--- a/Engine/Source/Editor/LevelEditor/Private/ViewportToolbar/LevelEditorSubmenus.cpp
+++ b/Engine/Source/Editor/LevelEditor/Private/ViewportToolbar/LevelEditorSubmenus.cpp
@@ -2043,8 +2043,8 @@ FToolMenuEntry CreatePIEShowSubmenu()
 {
 	FToolMenuEntry Entry = FToolMenuEntry::InitSubMenu(
 		"Show",
-		LOCTEXT("ShowSubmenuLabel", "Show"),
-		LOCTEXT("ShowSubmenuTooltip", "Show flags related to the game viewport"),
+		LOCTEXT("PIEShowSubmenuLabel", "Show"),
+		LOCTEXT("PIEShowSubmenuTooltip", "Show flags related to the game viewport"),
 		FNewToolMenuDelegate::CreateLambda(
 			[](UToolMenu* InMenu) -> void
 			{
@@ -2244,8 +2244,8 @@ FToolMenuEntry CreatePIEPerformanceAndScalabilitySubmenu()
 {
 	FToolMenuEntry Entry = FToolMenuEntry::InitSubMenu(
 		"PerformanceAndScalability",
-		LOCTEXT("PerformanceAndScalabilityLabel", "Performance & Scalability"),
-		LOCTEXT("PerformanceAndScalabilityTooltip", "Performance and scalability tools tied to this viewport."),
+		LOCTEXT("PIEPerformanceAndScalabilityLabel", "Performance & Scalability"),
+		LOCTEXT("PIEPerformanceAndScalabilityTooltip", "Performance and scalability tools tied to this viewport."),
 		FNewToolMenuDelegate::CreateLambda(
 			[](UToolMenu* Submenu) -> void
 			{
@@ -2629,19 +2629,19 @@ FToolMenuEntry CreatePIESettingsSubmenu()
 {
 	FToolMenuEntry Entry = FToolMenuEntry::InitSubMenu(
 		"Settings",
-		LOCTEXT("SettingsSubmenuLabel", "Settings"),
-		LOCTEXT("SettingsSubmenuTooltip", "Viewport-related settings"),
+		LOCTEXT("PIESettingsSubmenuLabel", "Settings"),
+		LOCTEXT("PIESettingsSubmenuTooltip", "Viewport-related settings"),
 		FNewToolMenuDelegate::CreateLambda(
 			[](UToolMenu* Submenu) -> void
 			{
 				{
 					if (ULevelViewportContext* Context = Submenu->FindContext<ULevelViewportContext>())
 					{
-						FToolMenuSection& SettingsSection = Submenu->FindOrAddSection("Settings", LOCTEXT("SettingsSectionLabel", "Settings"));
+						FToolMenuSection& SettingsSection = Submenu->FindOrAddSection("Settings", LOCTEXT("PIESettingsSectionLabel", "Settings"));
 						SettingsSection.AddEntry(FToolMenuEntry::InitWidget(
 							"Volume",
 							BuildPIEVolumeControlCustomWidget(Context->LevelViewport),
-							LOCTEXT("VolumeControlLabel", "Volume")
+							LOCTEXT("PIEVolumeControlLabel", "Volume")
 						));
 					}
 					
