Make the level viewport's SExtensionPanel legacy extensions not show up in the overflow menu.

#rnx

diff --git a/Engine/Source/Editor/LevelEditor/Private/SLevelViewport.cpp b/Engine/Source/Editor/LevelEditor/Private/SLevelViewport.cpp
--- a/Engine/Source/Editor/LevelEditor/Private/SLevelViewport.cpp
+++ b/Engine/Source/Editor/LevelEditor/Private/SLevelViewport.cpp
@@ -2421,7 +2421,7 @@ TSharedPtr<SWidget> SLevelViewport::BuildViewportToolbar()
 					
 					auto MakeExtensionPanel = [ExtensionContext, &Section](FName ExtensionId)
 					{
-						Section.AddEntry(FToolMenuEntry::InitWidget(
+						FToolMenuEntry& Entry = Section.AddEntry(FToolMenuEntry::InitWidget(
 							ExtensionId,
 							SNew(SExtensionPanel)
 								.ExtensionPanelID(ExtensionId)
@@ -2429,6 +2429,8 @@ TSharedPtr<SWidget> SLevelViewport::BuildViewportToolbar()
 							FText::GetEmpty(),
 							true, false, true
 						));
+						
+						Entry.ToolBarData.ResizeParams.VisibleInOverflow = false;
 					};
 					
 					MakeExtensionPanel("LevelViewportToolBar.LeftExtension");
