Make "Disable Preview" option in the PreviewPlatform menu set bPreviewFeatureLevelActive=false in order to not show Preview Platform specific overlays

diff --git a/Engine/Source/Editor/LevelEditor/Private/LevelEditor.cpp b/Engine/Source/Editor/LevelEditor/Private/LevelEditor.cpp
--- a/Engine/Source/Editor/LevelEditor/Private/LevelEditor.cpp
+++ b/Engine/Source/Editor/LevelEditor/Private/LevelEditor.cpp
@@ -1905,7 +1905,7 @@ void FLevelEditorModule::BindGlobalLevelEditorCommands()
 			{
 				if (bIsDefaultShaderPlatform)
 				{
-					return FPreviewPlatformInfo(GMaxRHIFeatureLevel, GMaxRHIShaderPlatform, NAME_None, NAME_None, NAME_None, true, NAME_None);
+					return FPreviewPlatformInfo(GMaxRHIFeatureLevel, GMaxRHIShaderPlatform, NAME_None, NAME_None, NAME_None, false, NAME_None);
 				}
 
 				const ERHIFeatureLevel::Type FeatureLevel = GetMaxSupportedFeatureLevel(ShaderPlatform);
