Fix issue where `FLevelViewportCommands::Get().ToggleViewportToolbar` could no longer be resolved after switching between Cinematic & Default viewports.

This appears to be an instance of UE-273630. Fixed by ensuring the command is bound to the command list directly tied to the menus so that it doesn't need to search through (potentially stale) parents/children.

#rnx

diff --git a/Engine/Source/Editor/LevelEditor/Private/LevelEditor.cpp b/Engine/Source/Editor/LevelEditor/Private/LevelEditor.cpp
--- a/Engine/Source/Editor/LevelEditor/Private/LevelEditor.cpp
+++ b/Engine/Source/Editor/LevelEditor/Private/LevelEditor.cpp
@@ -1791,13 +1791,7 @@ void FLevelEditorModule::BindGlobalLevelEditorCommands()
 		FCanExecuteAction(),
 		FIsActionChecked::CreateStatic(FLevelEditorActionCallbacks::ScaleGridSnap_IsChecked)
 		);
-		
-	ActionList.MapAction(
-		FLevelViewportCommands::Get().ToggleViewportToolbar,
-		FExecuteAction::CreateStatic( &FLevelEditorActionCallbacks::OnToggleShowViewportToolbar),
-		FCanExecuteAction(),
-		FIsActionChecked::CreateStatic( &FLevelEditorActionCallbacks::IsViewportToolbarVisible )
-		);
+	
 	ActionList.MapAction(
 		Commands.ToggleHideViewportUI,
 		FExecuteAction::CreateStatic( &FLevelEditorActionCallbacks::OnToggleShowViewportUI ),
diff --git a/Engine/Source/Editor/LevelEditor/Private/SLevelEditor.cpp b/Engine/Source/Editor/LevelEditor/Private/SLevelEditor.cpp
--- a/Engine/Source/Editor/LevelEditor/Private/SLevelEditor.cpp
+++ b/Engine/Source/Editor/LevelEditor/Private/SLevelEditor.cpp
@@ -79,6 +79,7 @@
 #include "SLevelViewport.h"
 #include "LevelEditorOutlinerSettings.h"
 #include "IWorldHierarchy.h"
+#include "LevelViewportActions.h"
 #include "Tools/EdModeInteractiveToolsContext.h"
 
 #define LOCTEXT_NAMESPACE "SLevelEditor"
@@ -193,6 +194,13 @@ void SLevelEditor::BindCommands()
 		FEditorViewportCommands::Get().FocusOutlinerToContextFolder,
 		FExecuteAction::CreateStatic( &FLevelEditorActionCallbacks::OnFocusOutlinerToContextFolder, TWeakPtr< SLevelEditor >( SharedThis( this ) ) )
 		);
+		
+	LevelEditorCommands->MapAction(
+		FLevelViewportCommands::Get().ToggleViewportToolbar,
+		FExecuteAction::CreateStatic( &FLevelEditorActionCallbacks::OnToggleShowViewportToolbar ),
+		FCanExecuteAction(),
+		FIsActionChecked::CreateStatic( &FLevelEditorActionCallbacks::IsViewportToolbarVisible )
+		);
 
 	if (FPlayWorldCommands::GlobalPlayWorldActions.IsValid())
 	{
