Added support for showing Verse paths to the recent and favorite level menus.

diff --git a/Engine/Source/Editor/LevelEditor/Private/LevelEditorMenu.cpp b/Engine/Source/Editor/LevelEditor/Private/LevelEditorMenu.cpp
--- a/Engine/Source/Editor/LevelEditor/Private/LevelEditorMenu.cpp
+++ b/Engine/Source/Editor/LevelEditor/Private/LevelEditorMenu.cpp
@@ -1,6 +1,10 @@
 // Copyright Epic Games, Inc. All Rights Reserved.
 
 #include "LevelEditorMenu.h"
+#include "AssetToolsModule.h"
+#include "AssetRegistry/IAssetRegistry.h"
+#include "Containers/VersePath.h"
+#include "Engine/World.h"
 #include "Modules/ModuleManager.h"
 #include "Framework/Application/SlateApplication.h"
 #include "Widgets/Text/STextBlock.h"
@@ -15,6 +19,7 @@
 #include "Interfaces/IMainFrameModule.h"
 #include "MRUFavoritesList.h"
 #include "Framework/Commands/GenericCommands.h"
+#include "IAssetTools.h"
 #include "IDocumentation.h"
 #include "ToolMenus.h"
 #include "LevelEditorMenuContext.h"
@@ -66,6 +71,35 @@ void FLevelEditorMenu::RegisterLevelEditorMenus()
 			SaveSection.AddMenuEntry( FLevelEditorCommands::Get().SaveAs ).InsertPosition = InsertPos;
 		}
 
+		static FText GetLevelPath(const FString& PackageName)
+		{
+			if (FAssetToolsModule::GetModule().Get().ShowingContentVersePath())
+			{
+				if (const IAssetRegistry* AssetRegistry = IAssetRegistry::Get())
+				{
+					// Minic logic in UEditorEngine::Map_Load which finds the first UWorld in the package.
+					TArray<FAssetData> AssetDatas;
+					if (AssetRegistry->GetAssetsByPackageName(FName(PackageName), AssetDatas))
+					{
+						for (const FAssetData& AssetData : AssetDatas)
+						{
+							if (AssetData.AssetClassPath == UWorld::StaticClass()->GetClassPathName())
+							{
+								UE::Core::FVersePath VersePath = AssetData.GetVersePath();
+								if (VersePath.IsValid())
+								{
+									return FText::FromString(MoveTemp(VersePath).ToString());
+								}
+								break;
+							}
+						}
+					}
+				}
+			}
+
+			return FText::FromString(PackageName);
+		}
+
 		static void FillFileRecentAndFavoriteFileItems()
 		{
 			UToolMenu* Menu = UToolMenus::Get()->ExtendMenu("LevelEditor.MainMenu.File");
@@ -126,9 +160,9 @@ void FLevelEditorMenu::RegisterLevelEditorMenus()
 						for (int32 CurFavoriteIndex = 0; CurFavoriteIndex < AllowedFavorites; ++CurFavoriteIndex)
 						{
 							TSharedPtr< FUICommandInfo > OpenFavoriteFile = FLevelEditorCommands::Get().OpenFavoriteFileCommands[CurFavoriteIndex];
-							const FString CurFavorite = FPaths::GetBaseFilename(MRUFavorites.GetFavoritesItem(CurFavoriteIndex));
-							const FText ToolTip = FText::Format(LOCTEXT("FavoriteLevelToolTip", "Opens favorite level: {0}"), FText::FromString(CurFavorite));
-							const FText Label = FText::FromString(FPaths::GetBaseFilename(CurFavorite));
+							const FString CurFavorite = MRUFavorites.GetFavoritesItem(CurFavoriteIndex);
+							const FText ToolTip = FText::Format(LOCTEXT("FavoriteLevelToolTip", "Opens favorite level: {0}"), GetLevelPath(CurFavorite));
+							const FText Label = FText::FromString(FPackageName::GetShortName(CurFavorite));
 
 							InMenu->FindOrAddSection("Favorite").AddMenuEntry(OpenFavoriteFile, Label, ToolTip).Name = NAME_None;
 						}
@@ -182,8 +216,8 @@ void FLevelEditorMenu::RegisterLevelEditorMenus()
 
 							const FString& CurRecent = MRUFavorites.GetMRUItem(CurRecentIndex);
 
-							const FText ToolTip = FText::Format(LOCTEXT("RecentLevelToolTip", "Opens recent level: {0}"), FText::FromString(CurRecent));
-							const FText Label = FText::FromString(FPaths::GetBaseFilename(CurRecent));
+							const FText ToolTip = FText::Format(LOCTEXT("RecentLevelToolTip", "Opens recent level: {0}"), GetLevelPath(CurRecent));
+							const FText Label = FText::FromString(FPackageName::GetShortName(CurRecent));
 
 							Section.AddMenuEntry(OpenRecentFile, Label, ToolTip).Name = NAME_None;
 						}
