Remove the gap on the right edge of the level viewport overflow menu.
- Hide the entire "LevelEditor.LevelEditorToolBar.SettingsToolBar" menu.
- Allow the "LevelEditor.LevelEditorToolBar.User" menu to expand fully.

#rnx

diff --git a/Engine/Source/Editor/LevelEditor/Private/LevelEditorToolBar.cpp b/Engine/Source/Editor/LevelEditor/Private/LevelEditorToolBar.cpp
--- a/Engine/Source/Editor/LevelEditor/Private/LevelEditorToolBar.cpp
+++ b/Engine/Source/Editor/LevelEditor/Private/LevelEditorToolBar.cpp
@@ -1650,7 +1650,7 @@ TSharedRef< SWidget > FLevelEditorToolBar::MakeLevelEditorToolBar( const TShared
 				
 			]
 			+ SHorizontalBox::Slot()
-			.HAlign(HAlign_Left)
+			.HAlign(HAlign_Fill)
 			.FillWidth(1.0f)
 			[
 				SNew(SBorder)
@@ -1667,6 +1667,10 @@ TSharedRef< SWidget > FLevelEditorToolBar::MakeLevelEditorToolBar( const TShared
 			.Padding(0.0f, 0.0f, 7.0f, 0.0f)
 			[
 				SNew(SBorder)
+				.Visibility_Lambda([]
+				{
+					return UE::Editor::Private::CVarLevelEditorToolbarSettingsValue != 0 ? EVisibility::Visible : EVisibility::Collapsed;
+				})
 				.Padding(0.f)
 				.BorderImage(FAppStyle::Get().GetBrush("NoBorder"))
 				.IsEnabled(FSlateApplication::Get().GetNormalExecutionAttribute())
