Update custom viewport toolbar Show submenus across various editors to use the new UE::UnrealEd::CreateShowSubmenu() utility function.

#rnx

diff --git a/Engine/Plugins/ChaosClothAssetEditor/Source/ChaosClothAssetEditor/Private/ChaosClothAsset/SClothEditorRestSpaceViewport.cpp b/Engine/Plugins/ChaosClothAssetEditor/Source/ChaosClothAssetEditor/Private/ChaosClothAsset/SClothEditorRestSpaceViewport.cpp
--- a/Engine/Plugins/ChaosClothAssetEditor/Source/ChaosClothAssetEditor/Private/ChaosClothAsset/SClothEditorRestSpaceViewport.cpp
+++ b/Engine/Plugins/ChaosClothAssetEditor/Source/ChaosClothAssetEditor/Private/ChaosClothAsset/SClothEditorRestSpaceViewport.cpp
@@ -335,27 +335,17 @@ TSharedPtr<SWidget> SChaosClothAssetEditorRestSpaceViewport::BuildViewportToolba
 			
 			{
 				// Show Menu
-				FToolMenuEntry& ShowEntry = RightSection.AddSubMenu(
-					"Show",
-					LOCTEXT("ShowSubmenuLabel", "Show"),
-					LOCTEXT("ShowSubmenuTooltip", "Show flags related to the current viewport"),
-					FNewToolMenuDelegate::CreateLambda(
-						[](UToolMenu* InMenu) -> void
-						{
-							FToolMenuSection& Section = InMenu->FindOrAddSection("Cloth", LOCTEXT("ClothShowMenu", "Chaos Cloth"));
-							Section.AddMenuEntry(FChaosClothAssetEditorCommands::Get().ToggleMeshStats);
-							Section.AddMenuEntry(FChaosClothAssetEditorCommands::Get().ToggleConstructionViewSeams);
-							Section.AddMenuEntry(FChaosClothAssetEditorCommands::Get().ToggleConstructionViewSeamsCollapse);
-							Section.AddMenuEntry(FChaosClothAssetEditorCommands::Get().TogglePatternColor);
-							Section.AddMenuEntry(FChaosClothAssetEditorCommands::Get().ToggleConstructionViewSurfaceNormals);
-						}
-					),
-					false,
-					FSlateIcon(FAppStyle::GetAppStyleSetName(), "Level.VisibleHighlightIcon16x")
-				);
-				
-				ShowEntry.ToolBarData.LabelOverride = FText();
-				ShowEntry.ToolBarData.ResizeParams.ClippingPriority = 800;
+				RightSection.AddEntry(UE::UnrealEd::CreateShowSubmenu(
+					FNewToolMenuDelegate::CreateLambda([](UToolMenu* InMenu)
+					{
+						FToolMenuSection& Section = InMenu->FindOrAddSection("Cloth", LOCTEXT("ClothShowMenu", "Chaos Cloth"));
+						Section.AddMenuEntry(FChaosClothAssetEditorCommands::Get().ToggleMeshStats);
+						Section.AddMenuEntry(FChaosClothAssetEditorCommands::Get().ToggleConstructionViewSeams);
+						Section.AddMenuEntry(FChaosClothAssetEditorCommands::Get().ToggleConstructionViewSeamsCollapse);
+						Section.AddMenuEntry(FChaosClothAssetEditorCommands::Get().TogglePatternColor);
+						Section.AddMenuEntry(FChaosClothAssetEditorCommands::Get().ToggleConstructionViewSurfaceNormals);
+					})
+				));
 			}
 			
 			RightSection.AddEntry(CreateDynamicSimulationMenuItem());
diff --git a/Engine/Plugins/Experimental/Dataflow/Source/DataflowEditor/Private/Dataflow/DataflowConstructionViewport.cpp b/Engine/Plugins/Experimental/Dataflow/Source/DataflowEditor/Private/Dataflow/DataflowConstructionViewport.cpp
--- a/Engine/Plugins/Experimental/Dataflow/Source/DataflowEditor/Private/Dataflow/DataflowConstructionViewport.cpp
+++ b/Engine/Plugins/Experimental/Dataflow/Source/DataflowEditor/Private/Dataflow/DataflowConstructionViewport.cpp
@@ -88,10 +88,7 @@ TSharedPtr<SWidget> SDataflowConstructionViewport::BuildViewportToolbar()
 				
 					{
 						// Show Menu
-						FToolMenuEntry& Entry = Section.AddSubMenu(
-							"Show",
-							LOCTEXT("ShowSubmenuLabel", "Show"),
-							LOCTEXT("ShowSubmenuTooltip", "Display options for the construction viewport."),
+						Section.AddEntry(UE::UnrealEd::CreateShowSubmenu(
 							FNewMenuDelegate::CreateLambda([WeakViewport = ConstructionViewport.ToWeakPtr()](FMenuBuilder& Menu)
 							{
 								if (const TSharedPtr<SDataflowConstructionViewport> PinnedEditorViewport = WeakViewport.Pin())
@@ -105,10 +102,7 @@ TSharedPtr<SWidget> SDataflowConstructionViewport::BuildViewportToolbar()
 									}
 								}
 							})
-						);
-						Entry.Icon = FSlateIcon(FAppStyle::GetAppStyleSetName(), "Level.VisibleHighlightIcon16x");
-						Entry.ToolBarData.LabelOverride = FText();
-						Entry.ToolBarData.ResizeParams.ClippingPriority = 800;
+						));
 					}
 					
 					{
diff --git a/Engine/Plugins/FX/Niagara/Source/NiagaraEditor/Private/Widgets/NiagaraSystemEditorViewportToolbarSections.cpp b/Engine/Plugins/FX/Niagara/Source/NiagaraEditor/Private/Widgets/NiagaraSystemEditorViewportToolbarSections.cpp
--- a/Engine/Plugins/FX/Niagara/Source/NiagaraEditor/Private/Widgets/NiagaraSystemEditorViewportToolbarSections.cpp
+++ b/Engine/Plugins/FX/Niagara/Source/NiagaraEditor/Private/Widgets/NiagaraSystemEditorViewportToolbarSections.cpp
@@ -16,22 +16,7 @@
 
 FToolMenuEntry UE::NiagaraSystemEditor::CreateShowSubmenu()
 {
-	FToolMenuEntry Entry = FToolMenuEntry::InitSubMenu(
-		"Show",
-		LOCTEXT("ShowSubmenuLabel", "Show"),
-		LOCTEXT("ShowSubmenuTooltip", "Show options"),
-		FNewToolMenuDelegate::CreateLambda(
-			[](UToolMenu* Submenu) -> void
-			{
-				UE::NiagaraSystemEditor::FillShowSubmenu(Submenu);
-			}
-		)
-	);
-	Entry.Icon = FSlateIcon(FAppStyle::GetAppStyleSetName(), "Level.VisibleHighlightIcon16x");
-	Entry.ToolBarData.LabelOverride = FText();
-	Entry.ToolBarData.ResizeParams.ClippingPriority = 800;
-
-	return Entry;
+	return UE::UnrealEd::CreateShowSubmenu(FNewToolMenuDelegate::CreateStatic(&UE::NiagaraSystemEditor::FillShowSubmenu, true));
 }
 
 void UE::NiagaraSystemEditor::FillShowSubmenu(UToolMenu* InMenu, bool bInShowViewportStatsToggle)
diff --git a/Engine/Plugins/Runtime/nDisplay/Source/DisplayClusterLightCardEditor/Private/Viewport/DisplayClusterLightcardEditorViewport.cpp b/Engine/Plugins/Runtime/nDisplay/Source/DisplayClusterLightCardEditor/Private/Viewport/DisplayClusterLightcardEditorViewport.cpp
--- a/Engine/Plugins/Runtime/nDisplay/Source/DisplayClusterLightCardEditor/Private/Viewport/DisplayClusterLightcardEditorViewport.cpp
+++ b/Engine/Plugins/Runtime/nDisplay/Source/DisplayClusterLightCardEditor/Private/Viewport/DisplayClusterLightcardEditorViewport.cpp
@@ -30,6 +30,7 @@
 #include "ToolMenu.h"
 #include "ToolMenus.h"
 #include "ToolMenuSection.h"
+#include "ViewportToolbar/UnrealEdViewportToolbar.h"
 #include "Widgets/Input/SCheckBox.h"
 
 
@@ -133,35 +134,26 @@ public:
 
 			// Show (eye) dropdown.
 			{
-				FToolMenuEntry& ShowEntry = RightSection.AddSubMenu(
-					"Show",
-					LOCTEXT("ShowSubmenuLabel", "Show"),
-					LOCTEXT("ShowSubmenuTooltip", "Show flags related to the current viewport"),
-					FNewToolMenuDelegate::CreateLambda(
-						[](UToolMenu* InMenu)
-						{
-							FToolMenuSection& SubSection = InMenu->FindOrAddSection("ShowFlags", LOCTEXT("ShowFlagsSection", "Show Flags"));
-
-							SubSection.AddMenuEntry(
-								FDisplayClusterLightCardEditorCommands::Get().ToggleAllLabels,
-								TAttribute<FText>(),
-								TAttribute<FText>(),
-								TAttribute<FSlateIcon>(FSlateIcon(FDisplayClusterLightCardEditorStyle::Get().GetStyleSetName(), "DisplayClusterLightCardEditor.LabelSymbol"))
-							);
-
-							SubSection.AddMenuEntry(
-								FDisplayClusterLightCardEditorCommands::Get().ToggleIconVisibility,
-								TAttribute<FText>(),
-								TAttribute<FText>(),
-								TAttribute<FSlateIcon>(FSlateIcon(FDisplayClusterLightCardEditorStyle::Get().GetStyleSetName(), "DisplayClusterLightCardEditor.IconSymbol"))
-							);
-						}
-					),
-					false,
-					FSlateIcon(FAppStyle::GetAppStyleSetName(), "Level.VisibleHighlightIcon16x")
-				);
-				ShowEntry.ToolBarData.LabelOverride = FText();
-				ShowEntry.ToolBarData.ResizeParams.ClippingPriority = 800;
+				RightSection.AddEntry(UE::UnrealEd::CreateShowSubmenu(FNewToolMenuDelegate::CreateLambda(
+					[](UToolMenu* InMenu)
+					{
+						FToolMenuSection& SubSection = InMenu->FindOrAddSection("ShowFlags", LOCTEXT("ShowFlagsSection", "Show Flags"));
+
+						SubSection.AddMenuEntry(
+							FDisplayClusterLightCardEditorCommands::Get().ToggleAllLabels,
+							TAttribute<FText>(),
+							TAttribute<FText>(),
+							TAttribute<FSlateIcon>(FSlateIcon(FDisplayClusterLightCardEditorStyle::Get().GetStyleSetName(), "DisplayClusterLightCardEditor.LabelSymbol"))
+						);
+
+						SubSection.AddMenuEntry(
+							FDisplayClusterLightCardEditorCommands::Get().ToggleIconVisibility,
+							TAttribute<FText>(),
+							TAttribute<FText>(),
+							TAttribute<FSlateIcon>(FSlateIcon(FDisplayClusterLightCardEditorStyle::Get().GetStyleSetName(), "DisplayClusterLightCardEditor.IconSymbol"))
+						);
+					})
+				));
 			}
 
 			// DrawLightCard button as a widget (it is more straight forward to customize the style of a Widget than a ToolBarButton).
diff --git a/Engine/Source/Editor/Kismet/Private/SSCSEditorViewport.cpp b/Engine/Source/Editor/Kismet/Private/SSCSEditorViewport.cpp
--- a/Engine/Source/Editor/Kismet/Private/SSCSEditorViewport.cpp
+++ b/Engine/Source/Editor/Kismet/Private/SSCSEditorViewport.cpp
@@ -415,23 +415,14 @@ TSharedPtr<SWidget> SSCSEditorViewport::BuildViewportToolbar()
 
 			// Add the "Show" submenu.
 			{
-				FToolMenuEntry ShowSubmenu = FToolMenuEntry::InitSubMenu(
-					"Show",
-					LOCTEXT("ShowLabel", "Show"),
-					LOCTEXT("ShowTooltip", "Show or hide elements from the viewport"),
-					FNewToolMenuDelegate::CreateLambda(
-						[](UToolMenu* Submenu) -> void
-						{
-							FToolMenuSection& UnnamedSection = Submenu->FindOrAddSection(NAME_None);
-							UnnamedSection.AddMenuEntry(FBlueprintEditorCommands::Get().ShowFloor);
-							UnnamedSection.AddMenuEntry(FBlueprintEditorCommands::Get().ShowGrid);
-						}
-					)
-				);
-				ShowSubmenu.ToolBarData.LabelOverride = FText();
-				ShowSubmenu.Icon = FSlateIcon(FAppStyle::GetAppStyleSetName(), "Level.VisibleHighlightIcon16x");
-				ShowSubmenu.ToolBarData.ResizeParams.ClippingPriority = 800;
-				RightSection.AddEntry(ShowSubmenu);
+				RightSection.AddEntry(UE::UnrealEd::CreateShowSubmenu(FNewToolMenuDelegate::CreateLambda(
+					[](UToolMenu* Submenu) -> void
+					{
+						FToolMenuSection& UnnamedSection = Submenu->FindOrAddSection(NAME_None);
+						UnnamedSection.AddMenuEntry(FBlueprintEditorCommands::Get().ShowFloor);
+						UnnamedSection.AddMenuEntry(FBlueprintEditorCommands::Get().ShowGrid);
+					})
+				));
 			}
 
 			// Add the "Performance & Scalability" submenu.
diff --git a/Engine/Source/Editor/LevelEditor/Private/ViewportToolbar/LevelEditorSubmenus.cpp b/Engine/Source/Editor/LevelEditor/Private/ViewportToolbar/LevelEditorSubmenus.cpp
--- a/Engine/Source/Editor/LevelEditor/Private/ViewportToolbar/LevelEditorSubmenus.cpp
+++ b/Engine/Source/Editor/LevelEditor/Private/ViewportToolbar/LevelEditorSubmenus.cpp
@@ -1949,118 +1949,100 @@ FToolMenuEntry CreateShowStatsSubmenu(bool bInAddToggleStatsCheckbox, TAttribute
 
 FToolMenuEntry CreateShowSubmenu()
 {
-	FToolMenuEntry Entry = FToolMenuEntry::InitSubMenu(
-		"Show",
-		LOCTEXT("ShowSubmenuLabel", "Show"),
-		LOCTEXT("ShowSubmenuTooltip", "Show flags related to the current viewport"),
-		FNewToolMenuDelegate::CreateLambda(
-			[](UToolMenu* InMenu) -> void
+	return UE::UnrealEd::CreateShowSubmenu(
+		FNewToolMenuDelegate::CreateLambda([](UToolMenu* InMenu) -> void
+		{
 			{
-				{
-					FToolMenuSection& UnnamedSection = InMenu->FindOrAddSection(NAME_None);
+				FToolMenuSection& UnnamedSection = InMenu->FindOrAddSection(NAME_None);
 
-					UnnamedSection.AddMenuEntry(FLevelViewportCommands::Get().UseDefaultShowFlags);
+				UnnamedSection.AddMenuEntry(FLevelViewportCommands::Get().UseDefaultShowFlags);
 
-					UnnamedSection.AddSeparator("ViewportStatsSeparator");
+				UnnamedSection.AddSeparator("ViewportStatsSeparator");
 
 #if STATS
-					// Override the label of the stats submenu for the new viewport toolbar.
-					UnnamedSection.AddEntry(
-						UE::LevelEditor::CreateShowStatsSubmenu(true, LOCTEXT("ViewportStatsLabel", "Viewport Stats"))
-					);
+				// Override the label of the stats submenu for the new viewport toolbar.
+				UnnamedSection.AddEntry(
+					UE::LevelEditor::CreateShowStatsSubmenu(true, LOCTEXT("ViewportStatsLabel", "Viewport Stats"))
+				);
 #endif
-				}
+			}
 
-				// Starting from commonly used flags
-				UE::UnrealEd::AddDefaultShowFlags(InMenu);
+			// Starting from commonly used flags
+			UE::UnrealEd::AddDefaultShowFlags(InMenu);
 
-				// Add Level Editor specific entries to the All Show Flags Section
-				{
-					FToolMenuSection& AllShowFlagsSection =
-						InMenu->FindOrAddSection("AllShowFlags", LOCTEXT("AllShowFlagsLabel", "All Show Flags"));
+			// Add Level Editor specific entries to the All Show Flags Section
+			{
+				FToolMenuSection& AllShowFlagsSection =
+					InMenu->FindOrAddSection("AllShowFlags", LOCTEXT("AllShowFlagsLabel", "All Show Flags"));
 
-					// Show Foliage
-					{
-						FToolMenuEntry ShowFoliageSubmenu = CreateShowFoliageSubmenu();
-						ShowFoliageSubmenu.Label = LOCTEXT("ShowFoliageLabel", "Foliage");
-						ShowFoliageSubmenu.InsertPosition.Position = EToolMenuInsertType::First;
-						AllShowFlagsSection.AddEntry(ShowFoliageSubmenu);
-					}
+				// Show Foliage
+				{
+					FToolMenuEntry ShowFoliageSubmenu = CreateShowFoliageSubmenu();
+					ShowFoliageSubmenu.Label = LOCTEXT("ShowFoliageLabel", "Foliage");
+					ShowFoliageSubmenu.InsertPosition.Position = EToolMenuInsertType::First;
+					AllShowFlagsSection.AddEntry(ShowFoliageSubmenu);
+				}
 
-					// Show HLODs
-					{
-						FToolMenuEntry ShowHLODSubmenu = CreateShowHLODsSubmenu();
-						ShowHLODSubmenu.InsertPosition.Position = EToolMenuInsertType::First;
-						AllShowFlagsSection.AddEntry(ShowHLODSubmenu);
-					}
+				// Show HLODs
+				{
+					FToolMenuEntry ShowHLODSubmenu = CreateShowHLODsSubmenu();
+					ShowHLODSubmenu.InsertPosition.Position = EToolMenuInsertType::First;
+					AllShowFlagsSection.AddEntry(ShowHLODSubmenu);
+				}
 
-					// Show Layers
-					{
-						FToolMenuEntry ShowLayersSubmenu = CreateShowLayersSubmenu();
-						ShowLayersSubmenu.InsertPosition.Position = EToolMenuInsertType::First;
-						AllShowFlagsSection.AddEntry(ShowLayersSubmenu);
-					}
+				// Show Layers
+				{
+					FToolMenuEntry ShowLayersSubmenu = CreateShowLayersSubmenu();
+					ShowLayersSubmenu.InsertPosition.Position = EToolMenuInsertType::First;
+					AllShowFlagsSection.AddEntry(ShowLayersSubmenu);
+				}
 
-					// Show Sprites
-					{
-						FToolMenuEntry ShowSpriteSubmenu = CreateShowSpritesSubmenu();
-						ShowSpriteSubmenu.InsertPosition.Position = EToolMenuInsertType::First;
-						AllShowFlagsSection.AddEntry(ShowSpriteSubmenu);
-					}
+				// Show Sprites
+				{
+					FToolMenuEntry ShowSpriteSubmenu = CreateShowSpritesSubmenu();
+					ShowSpriteSubmenu.InsertPosition.Position = EToolMenuInsertType::First;
+					AllShowFlagsSection.AddEntry(ShowSpriteSubmenu);
+				}
 
-					// Show Volumes
-					{
-						FToolMenuEntry ShowVolumesSubmenu = CreateShowVolumesSubmenu();
-						ShowVolumesSubmenu.InsertPosition.Position = EToolMenuInsertType::First;
-						AllShowFlagsSection.AddEntry(ShowVolumesSubmenu);
-					}
+				// Show Volumes
+				{
+					FToolMenuEntry ShowVolumesSubmenu = CreateShowVolumesSubmenu();
+					ShowVolumesSubmenu.InsertPosition.Position = EToolMenuInsertType::First;
+					AllShowFlagsSection.AddEntry(ShowVolumesSubmenu);
 				}
+			}
 
-				// Adds show flags sections for backward compatibility with the old viewport toolbar.
-				//  If your entries end up in this section, you should move it to the new "CommonShowFlags" section instead.
-				InMenu->FindOrAddSection(
-					"ShowFlagsMenuSectionCommon",
-					LOCTEXT("ShowFlagsMenuSectionCommonLabel", "Common Show Flags (Deprecated section)")
-				);
+			// Adds show flags sections for backward compatibility with the old viewport toolbar.
+			//  If your entries end up in this section, you should move it to the new "CommonShowFlags" section instead.
+			InMenu->FindOrAddSection(
+				"ShowFlagsMenuSectionCommon",
+				LOCTEXT("ShowFlagsMenuSectionCommonLabel", "Common Show Flags (Deprecated section)")
+			);
 
-				// If your entries end up in these sections, you should move them to the above "AllShowFlags" section instead.
-				InMenu->FindOrAddSection(
-					"LevelViewportShowFlags", LOCTEXT("LevelViewportShowFlagsLabel", "All Show Flags (Deprecated section)")
-				);
-				InMenu->FindOrAddSection(
-					"LevelViewportEditorShow", LOCTEXT("LevelViewportEditorShowLabel", "Editor (Deprecated section)")
-				);
-			}
-		)
-	);
-	Entry.Icon = FSlateIcon(FAppStyle::GetAppStyleSetName(), "Level.VisibleHighlightIcon16x");
-	Entry.ToolBarData.LabelOverride = FText();
-	Entry.ToolBarData.ResizeParams.ClippingPriority = 800;
-	return Entry;
+			// If your entries end up in these sections, you should move them to the above "AllShowFlags" section instead.
+			InMenu->FindOrAddSection(
+				"LevelViewportShowFlags", LOCTEXT("LevelViewportShowFlagsLabel", "All Show Flags (Deprecated section)")
+			);
+			InMenu->FindOrAddSection(
+				"LevelViewportEditorShow", LOCTEXT("LevelViewportEditorShowLabel", "Editor (Deprecated section)")
+			);
+		}
+	));
 }
 
 FToolMenuEntry CreatePIEShowSubmenu()
 {
-	FToolMenuEntry Entry = FToolMenuEntry::InitSubMenu(
-		"Show",
-		LOCTEXT("PIEShowSubmenuLabel", "Show"),
-		LOCTEXT("PIEShowSubmenuTooltip", "Show flags related to the game viewport"),
-		FNewToolMenuDelegate::CreateLambda(
-			[](UToolMenu* InMenu) -> void
+	return UE::UnrealEd::CreateShowSubmenu(FNewToolMenuDelegate::CreateLambda(
+		[](UToolMenu* InMenu) -> void
+		{
 			{
-				{
-					FToolMenuSection& UnnamedSection = InMenu->FindOrAddSection(NAME_None);
-					UnnamedSection.AddMenuEntry(FLevelViewportCommands::Get().UseDefaultShowFlags);
-				}
-
-				UE::UnrealEd::AddDefaultShowFlags(InMenu);
+				FToolMenuSection& UnnamedSection = InMenu->FindOrAddSection(NAME_None);
+				UnnamedSection.AddMenuEntry(FLevelViewportCommands::Get().UseDefaultShowFlags);
 			}
-		)
-	);
-	Entry.Icon = FSlateIcon(FAppStyle::GetAppStyleSetName(), "Level.VisibleHighlightIcon16x");
-	Entry.ToolBarData.LabelOverride = FText();
-	Entry.ToolBarData.ResizeParams.ClippingPriority = 800;
-	return Entry;
+
+			UE::UnrealEd::AddDefaultShowFlags(InMenu);
+		}
+	));
 }
 
 FToolMenuEntry CreateFeatureLevelPreviewSubmenu()
diff --git a/Engine/Source/Editor/Persona/Private/ViewportToolbar/AnimationEditorMenus.cpp b/Engine/Source/Editor/Persona/Private/ViewportToolbar/AnimationEditorMenus.cpp
--- a/Engine/Source/Editor/Persona/Private/ViewportToolbar/AnimationEditorMenus.cpp
+++ b/Engine/Source/Editor/Persona/Private/ViewportToolbar/AnimationEditorMenus.cpp
@@ -507,24 +507,13 @@ TSharedRef<SWidget> UE::AnimationEditor::MakeFollowBoneWidget(
 
 FToolMenuEntry UE::AnimationEditor::CreateShowSubmenu()
 {
-	FToolMenuEntry Entry = FToolMenuEntry::InitSubMenu(
-		"Show",
-		LOCTEXT("ShowSubmenuLabel", "Show"),
-		LOCTEXT("ShowSubmenuTooltip", "Show options"),
-		FNewToolMenuDelegate::CreateLambda(
-			[](UToolMenu* Submenu) -> void
-			{
-				AddSceneElementsSection(Submenu);
-				UE::AnimationEditor::FillShowSubmenu(Submenu);
-			}
-		)
-	);
-
-	Entry.Icon = FSlateIcon(FAppStyle::GetAppStyleSetName(), "Level.VisibleHighlightIcon16x");
-	Entry.ToolBarData.LabelOverride = FText();
-	Entry.ToolBarData.ResizeParams.ClippingPriority = 800;
-
-	return Entry;
+	return UE::UnrealEd::CreateShowSubmenu(FNewToolMenuDelegate::CreateLambda(
+		[](UToolMenu* Submenu) -> void
+		{
+			AddSceneElementsSection(Submenu);
+			UE::AnimationEditor::FillShowSubmenu(Submenu);
+		}
+	));
 }
 
 void UE::AnimationEditor::FillShowSubmenu(UToolMenu* InMenu, bool bInShowViewportStatsToggle)
diff --git a/Engine/Source/Editor/StaticMeshEditor/Private/StaticMeshEditorViewportToolbarSections.cpp b/Engine/Source/Editor/StaticMeshEditor/Private/StaticMeshEditorViewportToolbarSections.cpp
--- a/Engine/Source/Editor/StaticMeshEditor/Private/StaticMeshEditorViewportToolbarSections.cpp
+++ b/Engine/Source/Editor/StaticMeshEditor/Private/StaticMeshEditorViewportToolbarSections.cpp
@@ -52,22 +52,7 @@ FToolMenuEntry UE::StaticMeshEditor::CreateLODSubmenu()
 
 FToolMenuEntry UE::StaticMeshEditor::CreateShowSubmenu()
 {
-	FToolMenuEntry Entry = FToolMenuEntry::InitSubMenu(
-		"Show",
-		LOCTEXT("ShowSubmenuLabel", "Show"),
-		FText(),
-		FNewToolMenuDelegate::CreateLambda(
-			[](UToolMenu* Submenu) -> void
-			{
-				FillShowSubmenu(Submenu);
-			}
-		)
-	);
-	Entry.Icon = FSlateIcon(FAppStyle::GetAppStyleSetName(), "Level.VisibleHighlightIcon16x");
-	Entry.ToolBarData.LabelOverride = FText();
-	Entry.ToolBarData.ResizeParams.ClippingPriority = 800;
-
-	return Entry;
+	return UE::UnrealEd::CreateShowSubmenu(FNewToolMenuDelegate::CreateLambda(&FillShowSubmenu));
 }
 
 TSharedRef<SWidget> UE::StaticMeshEditor::GenerateLODMenuWidget(const TSharedPtr<SStaticMeshEditorViewport>& InStaticMeshEditorViewport
diff --git a/Engine/Source/Editor/UnrealEd/Private/SCommonEditorViewportToolbarBase.cpp b/Engine/Source/Editor/UnrealEd/Private/SCommonEditorViewportToolbarBase.cpp
--- a/Engine/Source/Editor/UnrealEd/Private/SCommonEditorViewportToolbarBase.cpp
+++ b/Engine/Source/Editor/UnrealEd/Private/SCommonEditorViewportToolbarBase.cpp
@@ -412,10 +412,7 @@ void SCommonEditorViewportToolbarBase::Construct(const FArguments& InArgs, TShar
 					}
 					UToolMenus::Get()->RegisterMenu(UToolMenus::JoinMenuPaths(ViewportToolbarName, "Show"), "ViewportToolbarBase.Show");
 				
-					FToolMenuEntry& ShowMenu = RightSection.AddSubMenu(
-						"Show",
-						LOCTEXT("ShowSubmenuLabel", "Show"),
-						LOCTEXT("ShowSubmenuTooltip", "Show flags related to the current viewport"),
+					RightSection.AddEntry(UE::UnrealEd::CreateShowSubmenu(
 						FNewToolMenuDelegate::CreateLambda([](UToolMenu* Submenu)
 						{
 							Submenu->AddDynamicSection("Flags", FNewToolMenuDelegate::CreateLambda([](UToolMenu* Menu)
@@ -425,7 +422,7 @@ void SCommonEditorViewportToolbarBase::Construct(const FArguments& InArgs, TShar
 								{
 									return;
 								}
-								
+									
 								if (!ToolbarWidget->bIsGeneratingToolMenuWidget)
 								{
 									if (TSharedPtr<SWidget> LegacyWidget = ToolbarWidget->MakeLegacyShowMenu())
@@ -447,10 +444,7 @@ void SCommonEditorViewportToolbarBase::Construct(const FArguments& InArgs, TShar
 								}
 							}));
 						})
-					);
-					ShowMenu.Icon = FSlateIcon(FAppStyle::GetAppStyleSetName(), "Level.VisibleHighlightIcon16x");
-					ShowMenu.ToolBarData.LabelOverride = FText();
-					ShowMenu.ToolBarData.ResizeParams.ClippingPriority = 800;
+					));
 				}
 				
 				RightSection.AddEntry(UE::UnrealEd::CreatePerformanceAndScalabilitySubmenu());
