Fixed localized string

#rnx

diff --git a/Engine/Source/Editor/LevelEditor/Private/LevelEditorActions.cpp b/Engine/Source/Editor/LevelEditor/Private/LevelEditorActions.cpp
--- a/Engine/Source/Editor/LevelEditor/Private/LevelEditorActions.cpp
+++ b/Engine/Source/Editor/LevelEditor/Private/LevelEditorActions.cpp
@@ -4373,7 +4373,7 @@ void FLevelEditorCommands::RegisterCommands()
 	UI_COMMAND( ToggleLODViewLocking, "Enable LOD View Locking", "If enabled viewports of the same type will use the same LOD", EUserInterfaceActionType::ToggleButton, FInputChord() );
 	UI_COMMAND( LevelStreamingVolumePrevis, "Enable Automatic Level Streaming", "If enabled, the viewport will stream in levels automatically when the camera is moved", EUserInterfaceActionType::ToggleButton, FInputChord() );
 	UI_COMMAND( EnableActorSnap, "Actor", "If enabled, actors will snap to the location of other actors when they are within distance", EUserInterfaceActionType::ToggleButton, FInputChord(EModifierKey::Control | EModifierKey::Shift, EKeys::K) );
-	UI_COMMAND( EnableVertexSnap, "Vertex", "If enabled, actors will snap to the location of the nearest vertex on another actor in the direction of " "movement", EUserInterfaceActionType::ToggleButton, FInputChord() );
+	UI_COMMAND( EnableVertexSnap, "Vertex", "If enabled, actors will snap to the location of the nearest vertex on another actor in the direction of movement", EUserInterfaceActionType::ToggleButton, FInputChord() );
 	UI_COMMAND( ToggleHideViewportUI, "Show Viewport UI", "Sets the visibility of all overlaid viewport UI widgets.", EUserInterfaceActionType::ToggleButton, FInputChord() );
 
 	//if (FParse::Param( FCommandLine::Get(), TEXT( "editortoolbox" ) ))
