FLevelEditorModule::StartupModule - Register FGenericCommands before the call to BindGlobalLevelEditorCommands to ensure the singleton instance.

diff --git a/Engine/Source/Editor/LevelEditor/Private/LevelEditor.cpp b/Engine/Source/Editor/LevelEditor/Private/LevelEditor.cpp
--- a/Engine/Source/Editor/LevelEditor/Private/LevelEditor.cpp
+++ b/Engine/Source/Editor/LevelEditor/Private/LevelEditor.cpp
@@ -273,6 +273,7 @@ void FLevelEditorModule::StartupModule()
 	NotificationBarExtensibilityManager = MakeShared<FExtensibilityManager>();
 
 	// Note this must come before any tab spawning because that can create the SLevelEditor and attempt to map commands
+	FGenericCommands::Register();
 	FGlobalEditorCommonCommands::Register();
 	FEditorViewportCommands::Register();
 	FLevelViewportCommands::Register();
