Fixed crash when dragging and dropping cooked blueprint assets from the content browser by adding a short lived TStrongObjectPtr

[RN] Fixed crash when dragging and dropping cooked blueprint assets from the content browser

diff --git a/Engine/Source/Editor/LevelEditor/Private/SLevelViewport.cpp b/Engine/Source/Editor/LevelEditor/Private/SLevelViewport.cpp
--- a/Engine/Source/Editor/LevelEditor/Private/SLevelViewport.cpp
+++ b/Engine/Source/Editor/LevelEditor/Private/SLevelViewport.cpp
@@ -1147,6 +1147,9 @@ bool SLevelViewport::HandlePlaceDraggedObjects(const FGeometry& MyGeometry, cons
 		bool bDropSuccessful = false;
 
 		// Make sure the drop preview is destroyed
+		// Note that this can run gc, better hold onto our objects!
+		TArray<TStrongObjectPtr<UObject>> DroppedObjectsLocked;
+		Algo::Transform(DroppedObjects, DroppedObjectsLocked, [](UObject* Obj){ return TStrongObjectPtr<UObject>(Obj); }); // Transform to accomodate explicit ctor
 		LevelViewportClient->DestroyDropPreviewElements();
 
 		if( !bShowDropContextMenu || !bCreateDropPreview )
