[Build Menu Actions] Added a name to the Dynamic Menu Sections + Entry of the Build Types menu. This will will prevent the duplication taking place when a new layout was loaded.

diff --git a/Engine/Source/Editor/LevelEditor/Private/LevelEditorMenu.cpp b/Engine/Source/Editor/LevelEditor/Private/LevelEditorMenu.cpp
--- a/Engine/Source/Editor/LevelEditor/Private/LevelEditorMenu.cpp
+++ b/Engine/Source/Editor/LevelEditor/Private/LevelEditorMenu.cpp
@@ -670,7 +670,7 @@ void FLevelEditorMenu::RegisterBuildMenu()
 
 	// Add section for external build types
 	{
-		Menu->AddDynamicSection(NAME_None, FNewToolMenuDelegate::CreateLambda([](UToolMenu* InMenu)
+		Menu->AddDynamicSection("ExternalBuildTypes", FNewToolMenuDelegate::CreateLambda([](UToolMenu* InMenu)
 		{
 			TArray<FName> BuildTypeNames;
 			FEditorBuildUtils::GetBuildTypes(BuildTypeNames);
@@ -738,7 +738,7 @@ void FLevelEditorMenu::RegisterBuildMenu()
 	}
 
 	{
-		Menu->AddDynamicSection(NAME_None, FNewToolMenuDelegate::CreateLambda([](UToolMenu* InMenu)
+		Menu->AddDynamicSection("WorldPartitionBuildTypes", FNewToolMenuDelegate::CreateLambda([](UToolMenu* InMenu)
 		{
 			if (GWorld != nullptr && GWorld->GetWorldPartition() != nullptr)
 			{
@@ -752,7 +752,7 @@ void FLevelEditorMenu::RegisterBuildMenu()
 
 	{
 		// The day we only support World Partitioned worlds, we can remove this section.
-		Menu->AddDynamicSection(NAME_None, FNewToolMenuDelegate::CreateLambda([](UToolMenu* InMenu)
+		Menu->AddDynamicSection("NonWorldPartitionBuildTypes", FNewToolMenuDelegate::CreateLambda([](UToolMenu* InMenu)
 		{
 			if (GWorld != nullptr && GWorld->GetWorldPartition() == nullptr)
 			{
