[Viewport Toolbar] Rename "Show Transform Widget" to "Show Transform Gizmo"

diff --git a/Engine/Source/Editor/LevelEditor/Private/LevelEditorActions.cpp b/Engine/Source/Editor/LevelEditor/Private/LevelEditorActions.cpp
--- a/Engine/Source/Editor/LevelEditor/Private/LevelEditorActions.cpp
+++ b/Engine/Source/Editor/LevelEditor/Private/LevelEditorActions.cpp
@@ -4361,7 +4361,7 @@ void FLevelEditorCommands::RegisterCommands()
 	UI_COMMAND( CreateBlankBlueprintClass, "New Empty Blueprint Class...", "Create a new Blueprint Class", EUserInterfaceActionType::Button, FInputChord() );
 	UI_COMMAND( ConvertSelectionToBlueprint, "Convert Selection to Blueprint Class...", "Replace all of the selected actors with a new Blueprint Class", EUserInterfaceActionType::Button, FInputChord() );
 
-	UI_COMMAND( ShowTransformWidget, "Show Transform Widget", "Toggles the visibility of the transform widgets", EUserInterfaceActionType::ToggleButton, FInputChord() );
+	UI_COMMAND( ShowTransformWidget, "Show Transform Gizmo", "Toggles the visibility of the transform widgets", EUserInterfaceActionType::ToggleButton, FInputChord() );
 	UI_COMMAND( AllowTranslucentSelection, "Allow Translucent Selection", "Allows translucent objects to be selected", EUserInterfaceActionType::ToggleButton, FInputChord(EKeys::T) );
 	UI_COMMAND( AllowGroupSelection, "Allow Group Selection", "Allows actor groups to be selected", EUserInterfaceActionType::ToggleButton, FInputChord(EModifierKey::Control|EModifierKey::Shift, EKeys::G) );
 	UI_COMMAND( StrictBoxSelect, "Strict Box Selection", "When enabled an object must be entirely encompassed by the selection box when marquee box selecting", EUserInterfaceActionType::ToggleButton, FInputChord() );
