[Viewport Toolbar] Keep necessary Camera Movement entries visible when switching to Ortographic view

Also remove unused CreateCameraSpeedMenu function in SCSEditor private namespace

diff --git a/Engine/Source/Editor/Kismet/Private/SSCSEditorViewport.cpp b/Engine/Source/Editor/Kismet/Private/SSCSEditorViewport.cpp
--- a/Engine/Source/Editor/Kismet/Private/SSCSEditorViewport.cpp
+++ b/Engine/Source/Editor/Kismet/Private/SSCSEditorViewport.cpp
@@ -59,14 +59,6 @@ struct FGeometry;
 namespace UE::SCSEditor::Private
 {
 
-void CreateCameraSpeedMenu(UToolMenu* InMenu, TWeakPtr<FSCSEditorViewportClient> InWeakViewportClient)
-{
-	FToolMenuSection& PositioningSection =
-		InMenu->FindOrAddSection("Positioning", LOCTEXT("PositioningLabel", "Positioning"));
-
-	PositioningSection.AddEntry(UE::UnrealEd::CreateCameraSpeedMenu());
-}
-
 bool IsViewModeSupported(EViewModeIndex InViewModeIndex)
 {
 	switch (InViewModeIndex)
diff --git a/Engine/Source/Editor/LevelEditor/Private/ViewportToolbar/LevelEditorSubmenus.cpp b/Engine/Source/Editor/LevelEditor/Private/ViewportToolbar/LevelEditorSubmenus.cpp
--- a/Engine/Source/Editor/LevelEditor/Private/ViewportToolbar/LevelEditorSubmenus.cpp
+++ b/Engine/Source/Editor/LevelEditor/Private/ViewportToolbar/LevelEditorSubmenus.cpp
@@ -613,29 +613,45 @@ FToolMenuEntry CreateCameraMovementSubmenu(const TWeakPtr<::SLevelViewport>& Lev
 		LOCTEXT("CameraMovementSubmenu", "Camera Movement"),
 		LOCTEXT("CameraMovementSubmenu_ToolTip", "Camera movement options"),
 		FNewToolMenuDelegate::CreateLambda(
-			[](UToolMenu* InMenu)
+			[LevelViewportWeak](UToolMenu* InMenu)
 			{
 				{
-					FToolMenuSection& CameraMovementSection =
-						InMenu->FindOrAddSection("CameraMovement", LOCTEXT("CameraMovementLabel", "Camera Movement"));
+					// Camera Movement
+					{
+						FToolMenuSection& CameraMovementSection =
+							InMenu->FindOrAddSection("CameraMovement", LOCTEXT("CameraMovementLabel", "Camera Movement"));
 
-					CameraMovementSection.AddEntry(UE::UnrealEd::CreateCameraSpeedMenu());
+						// This entry is created as perspective view only
+						CameraMovementSection.AddEntry(UE::UnrealEd::CreateCameraSpeedMenu());
 
-					// Frame Selection
-					{
-						FToolMenuEntry& FrameEntry = CameraMovementSection.AddMenuEntry(FEditorViewportCommands::Get().FocusViewportToSelection);
-						FrameEntry.Icon = FSlateIcon(FAppStyle::GetAppStyleSetName(), "Icons.FrameActor");
-					}
+						// Frame Selection (both perspective and orthographic)
+						{
+							FToolMenuEntry& FrameEntry =
+								CameraMovementSection.AddMenuEntry(FEditorViewportCommands::Get().FocusViewportToSelection);
+							FrameEntry.Icon = FSlateIcon(FAppStyle::GetAppStyleSetName(), "Icons.FrameActor");
+						}
 
-					CameraMovementSection.AddMenuEntry(FLevelEditorCommands::Get().SnapCameraToObject);
-					CameraMovementSection.AddMenuEntry(FLevelEditorCommands::Get().SnapObjectToCamera);
-					CameraMovementSection.AddMenuEntry(FLevelEditorCommands::Get().OrbitCameraAroundSelection);
+						// The following entries are perspective view only
+						{
+							FToolMenuSection& PerspectiveOnlySection =
+								InMenu->FindOrAddSection("CameraMovement_Perspective");
 
-					FToolMenuSection& OrthoSection =
-						InMenu->FindOrAddSection("Orthographic", LOCTEXT("OrthographicSectionLabel", "Orthographic"));
+							PerspectiveOnlySection.AddMenuEntry(FLevelEditorCommands::Get().SnapCameraToObject);
+							PerspectiveOnlySection.AddMenuEntry(FLevelEditorCommands::Get().SnapObjectToCamera);
+							PerspectiveOnlySection.AddMenuEntry(FLevelEditorCommands::Get().OrbitCameraAroundSelection);
 
-					OrthoSection.AddMenuEntry(FLevelEditorCommands::Get().LinkOrthographicViewports);
-					OrthoSection.AddMenuEntry(FLevelEditorCommands::Get().OrthoZoomToCursor);
+							PerspectiveOnlySection.Visibility = GetIsPerspectiveAttribute(LevelViewportWeak);
+						}
+					}
+
+					// Orthographic
+					{
+						FToolMenuSection& OrthoSection =
+							InMenu->FindOrAddSection("Orthographic", LOCTEXT("OrthographicSectionLabel", "Orthographic"));
+
+						OrthoSection.AddMenuEntry(FLevelEditorCommands::Get().LinkOrthographicViewports);
+						OrthoSection.AddMenuEntry(FLevelEditorCommands::Get().OrthoZoomToCursor);
+					}
 				}
 			}
 		),
@@ -643,8 +659,6 @@ FToolMenuEntry CreateCameraMovementSubmenu(const TWeakPtr<::SLevelViewport>& Lev
 		FSlateIcon(FSlateIcon(FAppStyle::GetAppStyleSetName(), "ClassIcon.CameraComponent"))
 	);
 
-	Entry.Visibility = GetIsPerspectiveAttribute(LevelViewportWeak);
-
 	return Entry;
 }
 
diff --git a/Engine/Source/Editor/UnrealEd/Private/ViewportToolbar/UnrealEdViewportToolbar.cpp b/Engine/Source/Editor/UnrealEd/Private/ViewportToolbar/UnrealEdViewportToolbar.cpp
--- a/Engine/Source/Editor/UnrealEd/Private/ViewportToolbar/UnrealEdViewportToolbar.cpp
+++ b/Engine/Source/Editor/UnrealEd/Private/ViewportToolbar/UnrealEdViewportToolbar.cpp
@@ -3204,8 +3204,6 @@ void PopulateCameraMenu(UToolMenu* InMenu, const FViewportCameraMenuOptions& InO
 			FToolMenuEntry& FrameEntry = MovementSection.AddMenuEntry(FEditorViewportCommands::Get().FocusViewportToSelection);
 			FrameEntry.Icon = FSlateIcon(FAppStyle::GetAppStyleSetName(), "Icons.FrameActor");
 		}
-
-		MovementSection.Visibility = UE::UnrealEd::GetIsPerspectiveAttribute(EditorViewport->GetViewportClient());
 	}
 
 	// View
@@ -4241,24 +4239,25 @@ FToolMenuEntry CreateCameraSpeedMenu()
 						}
 					);
 
-					// Don't show camera speed raised entry if viewport is orthographic (or missing)
-					Entry.SetShowInToolbarTopLevel(
-						TAttribute<bool>::CreateLambda(
-							[WeakViewport = Context->Viewport]() -> bool
+					TAttribute<bool> IsPerspectiveLambda = TAttribute<bool>::CreateLambda(
+						[WeakViewport = Context->Viewport]() -> bool
+						{
+							if (const TSharedPtr<SEditorViewport>& ViewportPinned = WeakViewport.Pin())
 							{
-								if (const TSharedPtr<SEditorViewport>& ViewportPinned = WeakViewport.Pin())
+								if (const TSharedPtr<FEditorViewportClient>& ViewportClient =
+										ViewportPinned->GetViewportClient())
 								{
-									if (const TSharedPtr<FEditorViewportClient>& ViewportClient =
-											ViewportPinned->GetViewportClient())
-									{
-										return ViewportClient->IsPerspective();
-									}
+									return ViewportClient->IsPerspective();
 								}
-
-								return false;
 							}
-						)
+
+							return false;
+						}
 					);
+
+					// Don't show camera speed entries (in-menu and raised) if viewport is orthographic (or missing)
+					Entry.SetShowInToolbarTopLevel(true);
+					Entry.Visibility = IsPerspectiveLambda;
 				}
 
 				Entry.ToolBarData.PlacementOverride = MenuPlacement_BelowRightAnchor;
