[Editor Viewports] Fix a crash when spawning objects caused by a minimized viewport being set as the active global viewport client

diff --git a/Engine/Source/Editor/LevelEditor/Private/LevelViewportTabContent.cpp b/Engine/Source/Editor/LevelEditor/Private/LevelViewportTabContent.cpp
--- a/Engine/Source/Editor/LevelEditor/Private/LevelViewportTabContent.cpp
+++ b/Engine/Source/Editor/LevelEditor/Private/LevelViewportTabContent.cpp
@@ -15,7 +15,8 @@
 #include "ToolMenus.h"
 #include "ViewportTypeDefinition.h"
 #include "Templates/SharedPointer.h"
-
+#include "SLevelViewport.h"
+#include "LevelViewportLayoutEntity.h"
 
 // FLevelViewportTabContent ///////////////////////////
 
@@ -109,25 +110,41 @@ void FLevelViewportTabContent::OnLayoutStartChange(bool bSwitchingLayouts)
 
 void FLevelViewportTabContent::OnLayoutChanged()
 {
-	// Set the global level editor to the first, valid perspective viewport found
-	if (GEditor)
+	// Set the global level editor viewport when the layout changes
+	TSharedPtr<FLevelViewportLayout> Layout = StaticCastSharedPtr<FLevelViewportLayout>(ActiveViewportLayout);
+	const bool bIsLayoutMaximized = Layout->IsMaximizeSupported() && Layout->IsMaximized();
+	FLevelEditorViewportClient* ClientToSetAsCurrent = nullptr;
+
+	for (const TPair<FName, TSharedPtr<IEditorViewportLayoutEntity>>& Viewport : ActiveViewportLayout->GetViewports())
 	{
-		const TArray<FLevelEditorViewportClient*>& LevelViewportClients = GEditor->GetLevelViewportClients();
-		for (FLevelEditorViewportClient* LevelViewport : LevelViewportClients)
+		ILevelViewportLayoutEntity* LevelLayoutEntity = static_cast<ILevelViewportLayoutEntity*>(Viewport.Value.Get());
+		// Prefer the maximized viewport in the layout, as it will be the only one visible
+		if (bIsLayoutMaximized)
 		{
-			if (LevelViewport->IsPerspective())
+			if (Layout->IsViewportMaximized(Viewport.Key))
 			{
-				LevelViewport->SetCurrentViewport();
+				ClientToSetAsCurrent = &LevelLayoutEntity->GetLevelViewportClient();
 				break;
 			}
 		}
-		// Otherwise just make sure it's set to something
-		if (!GCurrentLevelEditingViewportClient && LevelViewportClients.Num())
+		// Or the first perspective viewport we encounter in the layout
+		else if (LevelLayoutEntity->GetLevelViewportClient().IsPerspective())
+		{
+			ClientToSetAsCurrent = &LevelLayoutEntity->GetLevelViewportClient();
+			break;
+		}
+		// Or use the first viewport if we're not maximized and there is no perspective viewport
+		else if (!ClientToSetAsCurrent)
 		{
-			GCurrentLevelEditingViewportClient = LevelViewportClients[0];
+			ClientToSetAsCurrent = &LevelLayoutEntity->GetLevelViewportClient();
 		}
 	}
 
+	if (ensure(ClientToSetAsCurrent))
+	{
+		ClientToSetAsCurrent->SetCurrentViewport();
+	}
+
 	FLevelEditorModule& LevelEditorModule = FModuleManager::LoadModuleChecked<FLevelEditorModule>("LevelEditor");
 	LevelEditorModule.OnTabContentChanged().Broadcast();
 }
