Updated headers using UnrealCodeFixup to make sure dllstorage is on methods/staticvars instead of types. Part 5/n

diff --git a/Engine/Source/Editor/KismetWidgets/Public/SPinTypeSelector.h b/Engine/Source/Editor/KismetWidgets/Public/SPinTypeSelector.h
--- a/Engine/Source/Editor/KismetWidgets/Public/SPinTypeSelector.h
+++ b/Engine/Source/Editor/KismetWidgets/Public/SPinTypeSelector.h
@@ -37,6 +37,8 @@
 #include "Widgets/Views/STreeView.h"
 #include "Algo/LevenshteinDistance.h"
 
+#define UE_API KISMETWIDGETS_API
+
 class ITableRow;
 class SComboButton;
 class SMenuOwner;
@@ -91,12 +93,12 @@ struct FObjectReferenceType;
 typedef TSharedPtr<struct FObjectReferenceType> FObjectReferenceListItem;
 
 /** Widget for modifying the type for a variable or pin */
-class KISMETWIDGETS_API SPinTypeSelector : public SCompoundWidget
+class SPinTypeSelector : public SCompoundWidget
 {
 public:
-	static TSharedRef<SWidget> ConstructPinTypeImage(const FSlateBrush* PrimaryIcon, const FSlateColor& PrimaryColor, const FSlateBrush* SecondaryIcon, const FSlateColor& SecondaryColor, TSharedPtr<SToolTip> InToolTip);
-	static TSharedRef<SWidget> ConstructPinTypeImage(TAttribute<const FSlateBrush*> PrimaryIcon, TAttribute<FSlateColor> PrimaryColor, TAttribute<const FSlateBrush*> SecondaryIcon, TAttribute<FSlateColor> SecondaryColor );
-	static TSharedRef<SWidget> ConstructPinTypeImage(UEdGraphPin* Pin);
+	static UE_API TSharedRef<SWidget> ConstructPinTypeImage(const FSlateBrush* PrimaryIcon, const FSlateColor& PrimaryColor, const FSlateBrush* SecondaryIcon, const FSlateColor& SecondaryColor, TSharedPtr<SToolTip> InToolTip);
+	static UE_API TSharedRef<SWidget> ConstructPinTypeImage(TAttribute<const FSlateBrush*> PrimaryIcon, TAttribute<FSlateColor> PrimaryColor, TAttribute<const FSlateBrush*> SecondaryIcon, TAttribute<FSlateColor> SecondaryColor );
+	static UE_API TSharedRef<SWidget> ConstructPinTypeImage(UEdGraphPin* Pin);
 
 	/** Which type of selector should be used: compact or full mode, or not a selector at all, but just the type image. */
 	enum class ESelectorType : uint8
@@ -138,35 +140,35 @@ public:
 		SLATE_ARGUMENT(TArray<TSharedPtr<class IPinTypeSelectorFilter>>, CustomFilters)
 	SLATE_END_ARGS()
 public:
-	void Construct(const FArguments& InArgs, FGetPinTypeTree GetPinTypeTreeFunc);
+	UE_API void Construct(const FArguments& InArgs, FGetPinTypeTree GetPinTypeTreeFunc);
 
 	// SWidget interface
-	virtual FReply OnMouseButtonDown( const FGeometry& MyGeometry, const FPointerEvent& MouseEvent ) override;
-	virtual FReply OnMouseButtonUp( const FGeometry& MyGeometry, const FPointerEvent& MouseEvent ) override;
-	virtual void OnMouseLeave( const FPointerEvent& MouseEvent ) override;
+	UE_API virtual FReply OnMouseButtonDown( const FGeometry& MyGeometry, const FPointerEvent& MouseEvent ) override;
+	UE_API virtual FReply OnMouseButtonUp( const FGeometry& MyGeometry, const FPointerEvent& MouseEvent ) override;
+	UE_API virtual void OnMouseLeave( const FPointerEvent& MouseEvent ) override;
 	// End of SWidget interface
 
 protected:
 	/** Gets the icon (value, array, set, or map) for the type being manipulated */
-	const FSlateBrush* GetTypeIconImage() const;
+	UE_API const FSlateBrush* GetTypeIconImage() const;
 
 	/** Gets the secondary icon (for maps, otherwise null) for the type being manipulated */
-	const FSlateBrush* GetSecondaryTypeIconImage() const;
+	UE_API const FSlateBrush* GetSecondaryTypeIconImage() const;
 
 	/** Gets the type-specific color for the type being manipulated */
-	FSlateColor GetTypeIconColor() const;
+	UE_API FSlateColor GetTypeIconColor() const;
 
 	/** Gets the secondary type-specific color for the type being manipulated */
-	FSlateColor GetSecondaryTypeIconColor() const;
+	UE_API FSlateColor GetSecondaryTypeIconColor() const;
 
 	/** Gets a succinct type description for the type being manipulated */
-	FText GetTypeDescription(bool bIncludeSubcategory = false) const;
+	UE_API FText GetTypeDescription(bool bIncludeSubcategory = false) const;
 
 	/** Gets the secondary type description. E.g. the value type for TMaps */
-	FText GetSecondaryTypeDescription(bool bIncludeSubcategory = false) const;
+	UE_API FText GetSecondaryTypeDescription(bool bIncludeSubcategory = false) const;
 
 	/** Gets a combined description of the primary, container, and secondary types. E.g. "Map of Strings to Floats" */
-	FText GetCombinedTypeDescription(bool bIncludeSubcategory = false) const;
+	UE_API FText GetCombinedTypeDescription(bool bIncludeSubcategory = false) const;
 
 	TSharedPtr<SComboButton>		TypeComboButton;
 	TSharedPtr<SComboButton>		SecondaryTypeComboButton;
@@ -228,14 +230,14 @@ protected:
 	TArray<TSharedPtr<class IPinTypeSelectorFilter>> CustomFilters;
 
 	/** Array checkbox support functions */
-	ECheckBoxState IsArrayChecked() const;
-	void OnArrayCheckStateChanged(ECheckBoxState NewState);
+	UE_API ECheckBoxState IsArrayChecked() const;
+	UE_API void OnArrayCheckStateChanged(ECheckBoxState NewState);
 
 	/** Toggles the variable type as an array */
-	void OnArrayStateToggled();
+	UE_API void OnArrayStateToggled();
 
 	/** Updates the variable container type: */
-	void OnContainerTypeSelectionChanged(EPinContainerType PinContainerType);
+	UE_API void OnContainerTypeSelectionChanged(EPinContainerType PinContainerType);
 
 	/** Array containing the unfiltered list of all supported types this pin could possibly have */
 	TArray<FPinTypeTreeItem>		TypeTreeRoot;
@@ -243,44 +245,44 @@ protected:
 	TArray<FPinTypeTreeItem>		FilteredTypeTreeRoot;
 
 	/** Treeview support functions */
-	virtual TSharedRef<ITableRow> GenerateTypeTreeRow(FPinTypeTreeItem InItem, const TSharedRef<STableViewBase>& OwnerTree, bool bForSecondaryType);
-	void OnTypeSelectionChanged(FPinTypeTreeItem Selection, ESelectInfo::Type SelectInfo, bool bForSecondaryType);
-	void GetTypeChildren(FPinTypeTreeItem InItem, TArray<FPinTypeTreeItem>& OutChildren);
+	UE_API virtual TSharedRef<ITableRow> GenerateTypeTreeRow(FPinTypeTreeItem InItem, const TSharedRef<STableViewBase>& OwnerTree, bool bForSecondaryType);
+	UE_API void OnTypeSelectionChanged(FPinTypeTreeItem Selection, ESelectInfo::Type SelectInfo, bool bForSecondaryType);
+	UE_API void GetTypeChildren(FPinTypeTreeItem InItem, TArray<FPinTypeTreeItem>& OutChildren);
 
 	/** Listview support functions for sub-menu */
-	TSharedRef<ITableRow> GenerateObjectReferenceTreeRow(FObjectReferenceListItem InItem, const TSharedRef<STableViewBase>& OwnerTree);
-	void OnObjectReferenceSelectionChanged(FObjectReferenceListItem InItem, ESelectInfo::Type SelectInfo, bool bForSecondaryType);
+	UE_API TSharedRef<ITableRow> GenerateObjectReferenceTreeRow(FObjectReferenceListItem InItem, const TSharedRef<STableViewBase>& OwnerTree);
+	UE_API void OnObjectReferenceSelectionChanged(FObjectReferenceListItem InItem, ESelectInfo::Type SelectInfo, bool bForSecondaryType);
 
 	/** Reference to the menu content that's displayed when the type button is clicked on */
 	TSharedPtr<SMenuOwner> MenuContent;
-	virtual TSharedRef<SWidget>	GetMenuContent(bool bForSecondaryType);
-	TSharedRef<SWidget> GetPinContainerTypeMenuContent();
+	UE_API virtual TSharedRef<SWidget>	GetMenuContent(bool bForSecondaryType);
+	UE_API TSharedRef<SWidget> GetPinContainerTypeMenuContent();
 
 	/** Type searching support */
 	FText SearchText;
-	void OnFilterTextChanged(const FText& NewText);
-	void OnFilterTextCommitted(const FText& NewText, ETextCommit::Type CommitInfo);
+	UE_API void OnFilterTextChanged(const FText& NewText);
+	UE_API void OnFilterTextCommitted(const FText& NewText, ETextCommit::Type CommitInfo);
 
 	/** Helper to generate the filtered list of types, based on the supported types of the schema */
-	void FilterUnsupportedTypes(TArray<FPinTypeTreeItem>& UnfilteredList);
+	UE_API void FilterUnsupportedTypes(TArray<FPinTypeTreeItem>& UnfilteredList);
 
 	/** Helper to generate the filtered list of types, based on the search string matching */
-	bool GetChildrenMatchingSearch(const FText& SearchText, const TArray<FPinTypeTreeItem>& UnfilteredList, TArray<FPinTypeTreeItem>& OutFilteredList, FTopLevenshteinResult<FPinTypeTreeItem>& OutTopLevenshteinResult);
+	UE_API bool GetChildrenMatchingSearch(const FText& SearchText, const TArray<FPinTypeTreeItem>& UnfilteredList, TArray<FPinTypeTreeItem>& OutFilteredList, FTopLevenshteinResult<FPinTypeTreeItem>& OutTopLevenshteinResult);
 
 	/** Callback to get the tooltip text for the pin type combo box */
-	FText GetToolTipForComboBoxType() const;
+	UE_API FText GetToolTipForComboBoxType() const;
 
 	/** Callback to get the tooltip text for the secondary pin type combo box */
-	FText GetToolTipForComboBoxSecondaryType() const;
+	UE_API FText GetToolTipForComboBoxSecondaryType() const;
 
 	/** Callback to get the tooltip for the array button widget */
-	FText GetToolTipForArrayWidget() const;
+	UE_API FText GetToolTipForArrayWidget() const;
 
 	/** Callback to get the tooltip for the container type dropdown widget */
-	FText GetToolTipForContainerWidget() const;
+	UE_API FText GetToolTipForContainerWidget() const;
 
 	/** Callback to get the display text for the total pin type item count */
-	FText GetPinTypeItemCountText() const;
+	UE_API FText GetPinTypeItemCountText() const;
 
 	/**
 	 * Helper function to create widget for the sub-menu
@@ -290,13 +292,13 @@ protected:
 	 * @param InIconBrush			Brush icon to use for the menu entry item
 	 * @param InTooltip				The simple tooltip to use for the menu item, an advanced tooltip link will be auto-generated based on the PinCategory
 	 */
-	TSharedRef<SWidget> CreateObjectReferenceWidget(FPinTypeTreeItem InItem, FEdGraphPinType& InPinType, const FSlateBrush* InIconBrush, FText InSimpleTooltip) const;
+	UE_API TSharedRef<SWidget> CreateObjectReferenceWidget(FPinTypeTreeItem InItem, FEdGraphPinType& InPinType, const FSlateBrush* InIconBrush, FText InSimpleTooltip) const;
 
 	/** Gets the allowable object types for an tree item, used for building the sub-menu */
-	TSharedRef< SWidget > GetAllowedObjectTypes(FPinTypeTreeItem InItem, bool bForSecondaryType);
+	UE_API TSharedRef< SWidget > GetAllowedObjectTypes(FPinTypeTreeItem InItem, bool bForSecondaryType);
 
 	/** Fills the provided array with generated list items based on a pin type's allowed reference types, which could then be used as the item source for a list view */
-	void GenerateAllowedObjectTypesList(TArray<FObjectReferenceListItem>& OutList, FPinTypeTreeItem InItem, bool bForSecondaryType) const;
+	UE_API void GenerateAllowedObjectTypesList(TArray<FObjectReferenceListItem>& OutList, FPinTypeTreeItem InItem, bool bForSecondaryType) const;
 	
 	/**
 	 * When a pin type is selected, handle it
@@ -304,10 +306,10 @@ protected:
 	 * @param InItem				Item selected
 	 * @param InPinCategory			This is the PinType's category, must be provided separately as the PinType in the tree item is always Object Types for any object related type.
 	 */
-	void OnSelectPinType(FPinTypeTreeItem InItem, FName InPinCategory, bool bForSecondaryType);
+	UE_API void OnSelectPinType(FPinTypeTreeItem InItem, FName InPinCategory, bool bForSecondaryType);
 
 	/** Called whenever the custom filter options are changed. */
-	void OnCustomFilterChanged();
+	UE_API void OnCustomFilterChanged();
 };
 
 
@@ -315,7 +317,7 @@ protected:
 // IPinTypeSelectorFilter
 
 /** An interface for implementing a custom pin type filter for the selector widget. */
-class KISMETWIDGETS_API IPinTypeSelectorFilter
+class IPinTypeSelectorFilter
 {
 public:
 	virtual ~IPinTypeSelectorFilter() {}
@@ -331,4 +333,6 @@ public:
 
 	/** (Optional) - Override this method to return a widget that allows the user to toggle filter options on/off, etc. */
 	virtual TSharedPtr<SWidget> GetFilterOptionsWidget() { return TSharedPtr<SWidget>(); }
-};
\ No newline at end of file
+};
+
+#undef UE_API
diff --git a/Engine/Source/Editor/KismetWidgets/Public/SScrubControlPanel.h b/Engine/Source/Editor/KismetWidgets/Public/SScrubControlPanel.h
--- a/Engine/Source/Editor/KismetWidgets/Public/SScrubControlPanel.h
+++ b/Engine/Source/Editor/KismetWidgets/Public/SScrubControlPanel.h
@@ -16,7 +16,9 @@
 #include "Widgets/DeclarativeSyntaxSupport.h"
 #include "Widgets/SCompoundWidget.h"
 
-class KISMETWIDGETS_API SScrubControlPanel : public SCompoundWidget
+#define UE_API KISMETWIDGETS_API
+
+class SScrubControlPanel : public SCompoundWidget
 {
 public:
 	SLATE_BEGIN_ARGS( SScrubControlPanel )
@@ -104,12 +106,12 @@ public:
 	 * 
 	 * @param InArgs   A declaration from which to construct the widget
 	 */
-	void Construct( const FArguments& InArgs );
+	UE_API void Construct( const FArguments& InArgs );
 
 	virtual ~SScrubControlPanel() {}
 
 protected:
-	EVisibility GetRealtimeControlVisibility(bool bIsControlForRealtimeStreamingMode) const;
+	UE_API EVisibility GetRealtimeControlVisibility(bool bIsControlForRealtimeStreamingMode) const;
 
 protected:
 	TSharedPtr<class SScrubWidget>	ScrubWidget;
@@ -117,3 +119,4 @@ protected:
 	TAttribute<bool> IsRealtimeStreamingMode;
 };
 
+#undef UE_API
diff --git a/Engine/Source/Editor/KismetWidgets/Public/SScrubWidget.h b/Engine/Source/Editor/KismetWidgets/Public/SScrubWidget.h
--- a/Engine/Source/Editor/KismetWidgets/Public/SScrubWidget.h
+++ b/Engine/Source/Editor/KismetWidgets/Public/SScrubWidget.h
@@ -16,6 +16,8 @@
 #include "Widgets/DeclarativeSyntaxSupport.h"
 #include "Widgets/SCompoundWidget.h"
 
+#define UE_API KISMETWIDGETS_API
+
 class FPaintArgs;
 class FSlateRect;
 class FSlateWindowElementList;
@@ -37,7 +39,7 @@ DECLARE_DELEGATE_TwoParams( FOnScrubBarDrag, int32, float)
 DECLARE_DELEGATE_TwoParams( FOnScrubBarCommit, int32, float)
 DECLARE_DELEGATE_OneParam( FOnReZeroAnimSequence, int32 )
 
-class KISMETWIDGETS_API SScrubWidget : public SCompoundWidget
+class SScrubWidget : public SCompoundWidget
 {
 public:
 
@@ -94,20 +96,20 @@ public:
 	 * 
 	 * @param InArgs   A declaration from which to construct the widget
 	 */
-	void Construct( const FArguments& InArgs );
+	UE_API void Construct( const FArguments& InArgs );
 
 	// SWidget interface
-	virtual FVector2D ComputeDesiredSize(float) const override;
-	virtual int32 OnPaint( const FPaintArgs& Args, const FGeometry& AllottedGeometry, const FSlateRect& MyCullingRect, FSlateWindowElementList& OutDrawElements, int32 LayerId, const FWidgetStyle& InWidgetStyle, bool bParentEnabled ) const override;
-	virtual FReply OnMouseButtonDown( const FGeometry& MyGeometry, const FPointerEvent& MouseEvent ) override;
-	virtual FReply OnMouseButtonUp( const FGeometry& MyGeometry, const FPointerEvent& MouseEvent ) override;
-	virtual FReply OnMouseMove( const FGeometry& MyGeometry, const FPointerEvent& MouseEvent ) override;
-	virtual FReply OnMouseWheel( const FGeometry& MyGeometry, const FPointerEvent& MouseEvent ) override;
-	virtual FCursorReply OnCursorQuery( const FGeometry& MyGeometry, const FPointerEvent& CursorEvent ) const override;
+	UE_API virtual FVector2D ComputeDesiredSize(float) const override;
+	UE_API virtual int32 OnPaint( const FPaintArgs& Args, const FGeometry& AllottedGeometry, const FSlateRect& MyCullingRect, FSlateWindowElementList& OutDrawElements, int32 LayerId, const FWidgetStyle& InWidgetStyle, bool bParentEnabled ) const override;
+	UE_API virtual FReply OnMouseButtonDown( const FGeometry& MyGeometry, const FPointerEvent& MouseEvent ) override;
+	UE_API virtual FReply OnMouseButtonUp( const FGeometry& MyGeometry, const FPointerEvent& MouseEvent ) override;
+	UE_API virtual FReply OnMouseMove( const FGeometry& MyGeometry, const FPointerEvent& MouseEvent ) override;
+	UE_API virtual FReply OnMouseWheel( const FGeometry& MyGeometry, const FPointerEvent& MouseEvent ) override;
+	UE_API virtual FCursorReply OnCursorQuery( const FGeometry& MyGeometry, const FPointerEvent& CursorEvent ) const override;
 
 	UE_DEPRECATED(5.6, "Use the helper function accepting FVector2f instead of FVector2D")
-	static int32 GetDivider(float InputMinX, float InputMaxX, FVector2D WidgetSize, float SequenceLength, int32 NumFrames);
-	static int32 GetDivider(float InputMinX, float InputMaxX, FVector2f WidgetSize, float SequenceLength, int32 NumFrames);
+	static UE_API int32 GetDivider(float InputMinX, float InputMaxX, FVector2D WidgetSize, float SequenceLength, int32 NumFrames);
+	static UE_API int32 GetDivider(float InputMinX, float InputMaxX, FVector2f WidgetSize, float SequenceLength, int32 NumFrames);
 protected:
 
 	/**
@@ -117,27 +119,27 @@ protected:
 	 * @param bSliderClamp			true to clamp the slider value to the nearest delta value
 	 * @param bCommittedFromText	true if the value was committed from the text box and not the slider	
 	 */
-	void CommitValue( float NewValue, bool bSliderClamp, bool bCommittedFromText );
+	UE_API void CommitValue( float NewValue, bool bSliderClamp, bool bCommittedFromText );
 
 private:
 
 	/** Function to create context menu to display anim sequence editing options */
-	void CreateContextMenu( float CurrentFrameTime, const FPointerEvent& MouseEvent );
+	UE_API void CreateContextMenu( float CurrentFrameTime, const FPointerEvent& MouseEvent );
 
 	/** Function to crop animation sequence before/after selected frame */
-	void OnSequenceCroppedCalled( bool bFromStart, float CurrentFrameTime );
+	UE_API void OnSequenceCroppedCalled( bool bFromStart, float CurrentFrameTime );
 
 	/** Function to crop animation sequence before/after selected frame */
-	void OnSequenceAddedCalled(bool bBefore, int32 CurrentFrameNumber);
+	UE_API void OnSequenceAddedCalled(bool bBefore, int32 CurrentFrameNumber);
 
 	/** Function to append frames in the beginning or at the end*/
-	void OnSequenceAppendedCalled(const FText & InNewGroupText, ETextCommit::Type CommitInfo, bool bBegin);
+	UE_API void OnSequenceAppendedCalled(const FText & InNewGroupText, ETextCommit::Type CommitInfo, bool bBegin);
 
 	/** Function to ask how many frames to append */
-	void OnShowPopupOfAppendAnimation(bool bBegin);
+	UE_API void OnShowPopupOfAppendAnimation(bool bBegin);
 
 	/** Function to zero out translation of the selected frame */
-	void OnReZeroCalled(int32 FrameIndex);
+	UE_API void OnReZeroCalled(int32 FrameIndex);
 
 	TAttribute<float> ValueAttribute;
 	FOnFloatValueChanged OnValueChanged;
@@ -176,3 +178,4 @@ private:
 	bool DraggingBar;
 };
 
+#undef UE_API
diff --git a/Engine/Source/Editor/KismetWidgets/Public/SSingleObjectDetailsPanel.h b/Engine/Source/Editor/KismetWidgets/Public/SSingleObjectDetailsPanel.h
--- a/Engine/Source/Editor/KismetWidgets/Public/SSingleObjectDetailsPanel.h
+++ b/Engine/Source/Editor/KismetWidgets/Public/SSingleObjectDetailsPanel.h
@@ -12,6 +12,8 @@
 #include "Widgets/SCompoundWidget.h"
 #include "Widgets/SWidget.h"
 
+#define UE_API KISMETWIDGETS_API
+
 class IDetailsView;
 class SWidget;
 class UObject;
@@ -20,7 +22,7 @@ struct FGeometry;
 /////////////////////////////////////////////////////
 // SSingleObjectDetailsPanel
 
-class KISMETWIDGETS_API SSingleObjectDetailsPanel : public SCompoundWidget
+class SSingleObjectDetailsPanel : public SCompoundWidget
 {
 public:
 	SLATE_BEGIN_ARGS(SSingleObjectDetailsPanel) {}
@@ -28,21 +30,21 @@ public:
 		SLATE_ARGUMENT(TSharedPtr<class FTabManager>, HostTabManager)
 	SLATE_END_ARGS()
 
-	void Construct(const FArguments& InArgs, bool bAutomaticallyObserveViaGetObjectToObserve = true, bool bAllowSearch = false);
+	UE_API void Construct(const FArguments& InArgs, bool bAutomaticallyObserveViaGetObjectToObserve = true, bool bAllowSearch = false);
 
 	// SWidget interface
-	virtual void Tick( const FGeometry& AllottedGeometry, const double InCurrentTime, const float InDeltaTime ) override;
+	UE_API virtual void Tick( const FGeometry& AllottedGeometry, const double InCurrentTime, const float InDeltaTime ) override;
 	// End of SWidget interface
 
 protected:
 	// Should be implemented by derived classes to provide the object being observed
-	virtual UObject* GetObjectToObserve() const;
+	UE_API virtual UObject* GetObjectToObserve() const;
 
 	// 
-	void SetPropertyWindowContents(TArray<UObject*> Objects);
+	UE_API void SetPropertyWindowContents(TArray<UObject*> Objects);
 
 	//
-	virtual TSharedRef<SWidget> PopulateSlot(TSharedRef<SWidget> PropertyEditorWidget);
+	UE_API virtual TSharedRef<SWidget> PopulateSlot(TSharedRef<SWidget> PropertyEditorWidget);
 
 	// Property viewing widget
 	TSharedPtr<class IDetailsView> PropertyView;
@@ -57,3 +59,4 @@ private:
 };
 
 
+#undef UE_API
diff --git a/Engine/Source/Editor/LandscapeEditor/Public/LandscapeImportHelper.h b/Engine/Source/Editor/LandscapeEditor/Public/LandscapeImportHelper.h
--- a/Engine/Source/Editor/LandscapeEditor/Public/LandscapeImportHelper.h
+++ b/Engine/Source/Editor/LandscapeEditor/Public/LandscapeImportHelper.h
@@ -9,6 +9,8 @@
 #include "Templates/EnableIf.h"
 #include "LandscapeImportHelper.generated.h"
 
+#define UE_API LANDSCAPEEDITOR_API
+
 USTRUCT()
 struct FLandscapeImportFileDescriptor
 {
@@ -107,7 +109,7 @@ struct FLandscapeImportDescriptor
 	FVector Scale;
 };
 
-class LANDSCAPEEDITOR_API FLandscapeImportHelper
+class FLandscapeImportHelper
 {
 public:
 	/**
@@ -117,7 +119,7 @@ public:
 	 * @param OutMessage In case the result of the operation is  an error or warning, this will contain the reason.
 	 * @return If the operation was a Success/Error/Warning. 
 	 */
-	static ELandscapeImportResult GetHeightmapImportDescriptor(const FString& FilePath, bool bSingleFile, bool bFlipYAxis, FLandscapeImportDescriptor& OutImportDescriptor, FText& OutMessage);
+	static UE_API ELandscapeImportResult GetHeightmapImportDescriptor(const FString& FilePath, bool bSingleFile, bool bFlipYAxis, FLandscapeImportDescriptor& OutImportDescriptor, FText& OutMessage);
 
 	/**
 	 * @param FilePath File to import when bSingleFile is true.
@@ -127,18 +129,18 @@ public:
 	 * @param OutMessage In case the result of the operation is  an error or warning, this will contain the reason.
 	 * @return If the operation was a Success/Error/Warning.
 	 */
-	static ELandscapeImportResult GetWeightmapImportDescriptor(const FString& FilePath, bool bSingleFile, bool bFlipYAxis, FName LayerName, FLandscapeImportDescriptor& OutImportDescriptor, FText& OutMessage);
+	static UE_API ELandscapeImportResult GetWeightmapImportDescriptor(const FString& FilePath, bool bSingleFile, bool bFlipYAxis, FName LayerName, FLandscapeImportDescriptor& OutImportDescriptor, FText& OutMessage);
 
-	static ELandscapeImportResult GetHeightmapImportData(const FLandscapeImportDescriptor& ImportDescriptor, int32 DescriptorIndex, TArray<uint16>& OutData, FText& OutMessage);
-	static ELandscapeImportResult GetWeightmapImportData(const FLandscapeImportDescriptor& ImportDescriptor, int32 DescriptorIndex, FName LayerName, TArray<uint8>& OutData, FText& OutMessage);
+	static UE_API ELandscapeImportResult GetHeightmapImportData(const FLandscapeImportDescriptor& ImportDescriptor, int32 DescriptorIndex, TArray<uint16>& OutData, FText& OutMessage);
+	static UE_API ELandscapeImportResult GetWeightmapImportData(const FLandscapeImportDescriptor& ImportDescriptor, int32 DescriptorIndex, FName LayerName, TArray<uint8>& OutData, FText& OutMessage);
 
-	static void TransformWeightmapImportData(const TArray<uint8>& InData, TArray<uint8>& OutData, const FLandscapeImportResolution& CurrentResolution, const FLandscapeImportResolution& RequiredResolution, ELandscapeImportTransformType TransformType = ELandscapeImportTransformType::ExpandCentered, FIntPoint Offset = FIntPoint(0,0));
-	static void TransformHeightmapImportData(const TArray<uint16>& InData, TArray<uint16>& OutData, const FLandscapeImportResolution& CurrentResolution, const FLandscapeImportResolution& RequiredResolution, ELandscapeImportTransformType TransformType = ELandscapeImportTransformType::ExpandCentered, FIntPoint Offset = FIntPoint(0, 0));
+	static UE_API void TransformWeightmapImportData(const TArray<uint8>& InData, TArray<uint8>& OutData, const FLandscapeImportResolution& CurrentResolution, const FLandscapeImportResolution& RequiredResolution, ELandscapeImportTransformType TransformType = ELandscapeImportTransformType::ExpandCentered, FIntPoint Offset = FIntPoint(0,0));
+	static UE_API void TransformHeightmapImportData(const TArray<uint16>& InData, TArray<uint16>& OutData, const FLandscapeImportResolution& CurrentResolution, const FLandscapeImportResolution& RequiredResolution, ELandscapeImportTransformType TransformType = ELandscapeImportTransformType::ExpandCentered, FIntPoint Offset = FIntPoint(0, 0));
 
-	static void ChooseBestComponentSizeForImport(int32 Width, int32 Height, int32& InOutQuadsPerSection, int32& InOutSectionsPerComponent, FIntPoint& OutComponentCount);
+	static UE_API void ChooseBestComponentSizeForImport(int32 Width, int32 Height, int32& InOutQuadsPerSection, int32& InOutSectionsPerComponent, FIntPoint& OutComponentCount);
 
-	static bool ExtractCoordinates(const FString& BaseFilename, FIntPoint& OutCoord, FString& OutBaseFilePattern);
-	static void GetMatchingFiles(const FString& FilePathPattern, TArray<FString>& OutFilePaths);
+	static UE_API bool ExtractCoordinates(const FString& BaseFilename, FIntPoint& OutCoord, FString& OutBaseFilePattern);
+	static UE_API void GetMatchingFiles(const FString& FilePathPattern, TArray<FString>& OutFilePaths);
 
 
 	template<typename T>
@@ -173,4 +175,6 @@ public:
 	{
 		return TransformWeightmapImportData(InData, OutData, CurrentResolution, RequiredResolution, TransformType, Offset);
 	}
-};
\ No newline at end of file
+};
+
+#undef UE_API
diff --git a/Engine/Source/Editor/LandscapeEditorUtilities/Public/LandscapeBlueprintBrush.h b/Engine/Source/Editor/LandscapeEditorUtilities/Public/LandscapeBlueprintBrush.h
--- a/Engine/Source/Editor/LandscapeEditorUtilities/Public/LandscapeBlueprintBrush.h
+++ b/Engine/Source/Editor/LandscapeEditorUtilities/Public/LandscapeBlueprintBrush.h
@@ -9,8 +9,8 @@
 
 #include "LandscapeBlueprintBrush.generated.h"
 
-UCLASS(Abstract, Blueprintable, hidecategories = (Replication, Input, LOD, Actor, Rendering), showcategories = (Cooking))
-class LANDSCAPEEDITORUTILITIES_API ALandscapeBlueprintBrush : public ALandscapeBlueprintBrushBase
+UCLASS(MinimalAPI, Abstract, Blueprintable, hidecategories = (Replication, Input, LOD, Actor, Rendering), showcategories = (Cooking))
+class ALandscapeBlueprintBrush : public ALandscapeBlueprintBrushBase
 {
 	GENERATED_UCLASS_BODY()
-};
\ No newline at end of file
+};
diff --git a/Engine/Source/Editor/Layers/Public/LayersDragDropOp.h b/Engine/Source/Editor/Layers/Public/LayersDragDropOp.h
--- a/Engine/Source/Editor/Layers/Public/LayersDragDropOp.h
+++ b/Engine/Source/Editor/Layers/Public/LayersDragDropOp.h
@@ -7,10 +7,12 @@
 #include "Input/DragAndDrop.h"
 #include "UObject/NameTypes.h"
 
+#define UE_API LAYERS_API
+
 class FDragDropEvent;
 
 /** Drag/drop operation for dragging layers in the editor */
-class LAYERS_API FLayersDragDropOp : public FDecoratedDragDropOp
+class FLayersDragDropOp : public FDecoratedDragDropOp
 {
 public:
 	DRAG_DROP_OPERATOR_TYPE(FLayersDragDropOp, FDragDropOperation)
@@ -18,5 +20,7 @@ public:
 	/** The names of the layers being dragged */
 	TArray<FName> Layers;
 
-	virtual void Construct() override;
+	UE_API virtual void Construct() override;
 };
+
+#undef UE_API
diff --git a/Engine/Source/Editor/LevelEditor/Private/LevelViewportLayoutEntity.h b/Engine/Source/Editor/LevelEditor/Private/LevelViewportLayoutEntity.h
--- a/Engine/Source/Editor/LevelEditor/Private/LevelViewportLayoutEntity.h
+++ b/Engine/Source/Editor/LevelEditor/Private/LevelViewportLayoutEntity.h
@@ -6,28 +6,32 @@
 #include "Widgets/SWidget.h"
 #include "LevelViewportLayout.h"
 
+#define UE_API LEVELEDITOR_API
+
 class FLevelEditorViewportClient;
 class SLevelViewport;
 
-class LEVELEDITOR_API FLevelViewportLayoutEntity : public ILevelViewportLayoutEntity
+class FLevelViewportLayoutEntity : public ILevelViewportLayoutEntity
 {
 public:
 
-	FLevelViewportLayoutEntity(TSharedPtr<SAssetEditorViewport> InLevelViewport);
-	virtual TSharedRef<SWidget> AsWidget() const override;
-	virtual TSharedPtr<SLevelViewport> AsLevelViewport() const override;
+	UE_API FLevelViewportLayoutEntity(TSharedPtr<SAssetEditorViewport> InLevelViewport);
+	UE_API virtual TSharedRef<SWidget> AsWidget() const override;
+	UE_API virtual TSharedPtr<SLevelViewport> AsLevelViewport() const override;
 
-	bool IsPlayInEditorViewportActive() const;
-	void RegisterGameViewportIfPIE();
-	void SetKeyboardFocus();
-	void OnLayoutDestroyed();
-	void SaveConfig(const FString& ConfigSection);
-	FLevelEditorViewportClient& GetLevelViewportClient() const;
-	FName GetType() const;
-	void TakeHighResScreenShot() const;
+	UE_API bool IsPlayInEditorViewportActive() const;
+	UE_API void RegisterGameViewportIfPIE();
+	UE_API void SetKeyboardFocus();
+	UE_API void OnLayoutDestroyed();
+	UE_API void SaveConfig(const FString& ConfigSection);
+	UE_API FLevelEditorViewportClient& GetLevelViewportClient() const;
+	UE_API FName GetType() const;
+	UE_API void TakeHighResScreenShot() const;
 
 protected:
 
 	/** This entity's level viewport */
 	TSharedRef<SLevelViewport> LevelViewport;
 };
+
+#undef UE_API
diff --git a/Engine/Source/Editor/LevelEditor/Private/ViewportToolbar/LevelViewportContext.h b/Engine/Source/Editor/LevelEditor/Private/ViewportToolbar/LevelViewportContext.h
--- a/Engine/Source/Editor/LevelEditor/Private/ViewportToolbar/LevelViewportContext.h
+++ b/Engine/Source/Editor/LevelEditor/Private/ViewportToolbar/LevelViewportContext.h
@@ -8,8 +8,10 @@
 
 #include "LevelViewportContext.generated.h"
 
-UCLASS()
-class LEVELEDITOR_API ULevelViewportContext : public UObject
+#define UE_API LEVELEDITOR_API
+
+UCLASS(MinimalAPI)
+class ULevelViewportContext : public UObject
 {
 	GENERATED_BODY()
 
@@ -19,13 +21,15 @@ public:
 
 // Customizes ULevelViewportToolBarContext for backwards compatability
 // TODO: Remove this once the old toolbars are deprecated.
-UCLASS()
-class LEVELEDITOR_API ULegacyLevelViewportToolbarContext : public ULevelViewportToolBarContext
+UCLASS(MinimalAPI)
+class ULegacyLevelViewportToolbarContext : public ULevelViewportToolBarContext
 {
 	GENERATED_BODY()
 	
 public:
 	TWeakPtr<SLevelViewport> LevelViewport;
 	
-	virtual FLevelEditorViewportClient* GetLevelViewportClient() const override;
+	UE_API virtual FLevelEditorViewportClient* GetLevelViewportClient() const override;
 };
+
+#undef UE_API
diff --git a/Engine/Source/Editor/LevelEditor/Public/Elements/Actor/ActorElementLevelEditorViewportInteractionCustomization.h b/Engine/Source/Editor/LevelEditor/Public/Elements/Actor/ActorElementLevelEditorViewportInteractionCustomization.h
--- a/Engine/Source/Editor/LevelEditor/Public/Elements/Actor/ActorElementLevelEditorViewportInteractionCustomization.h
+++ b/Engine/Source/Editor/LevelEditor/Public/Elements/Actor/ActorElementLevelEditorViewportInteractionCustomization.h
@@ -5,14 +5,18 @@
 #include "Elements/Actor/ActorElementEditorViewportInteractionCustomization.h"
 #include "Elements/Framework/TypedElementAssetEditorLevelEditorViewportClientMixin.h"
 
-class LEVELEDITOR_API FActorElementLevelEditorViewportInteractionCustomization : public FActorElementEditorViewportInteractionCustomization, public FTypedElementAssetEditorLevelEditorViewportClientMixin
+#define UE_API LEVELEDITOR_API
+
+class FActorElementLevelEditorViewportInteractionCustomization : public FActorElementEditorViewportInteractionCustomization, public FTypedElementAssetEditorLevelEditorViewportClientMixin
 {
 public:
-	virtual void GizmoManipulationStarted(const TTypedElement<ITypedElementWorldInterface>& InElementWorldHandle, const UE::Widget::EWidgetMode InWidgetMode) override;
-	virtual void GizmoManipulationDeltaUpdate(const TTypedElement<ITypedElementWorldInterface>& InElementWorldHandle, const UE::Widget::EWidgetMode InWidgetMode, const EAxisList::Type InDragAxis, const FInputDeviceState& InInputState, const FTransform& InDeltaTransform, const FVector& InPivotLocation) override;
-	virtual void GizmoManipulationStopped(const TTypedElement<ITypedElementWorldInterface>& InElementWorldHandle, const UE::Widget::EWidgetMode InWidgetMode, const ETypedElementViewportInteractionGizmoManipulationType InManipulationType) override;
-	virtual void PostGizmoManipulationStopped(TArrayView<const FTypedElementHandle> InElementHandles, const UE::Widget::EWidgetMode InWidgetMode) override;
+	UE_API virtual void GizmoManipulationStarted(const TTypedElement<ITypedElementWorldInterface>& InElementWorldHandle, const UE::Widget::EWidgetMode InWidgetMode) override;
+	UE_API virtual void GizmoManipulationDeltaUpdate(const TTypedElement<ITypedElementWorldInterface>& InElementWorldHandle, const UE::Widget::EWidgetMode InWidgetMode, const EAxisList::Type InDragAxis, const FInputDeviceState& InInputState, const FTransform& InDeltaTransform, const FVector& InPivotLocation) override;
+	UE_API virtual void GizmoManipulationStopped(const TTypedElement<ITypedElementWorldInterface>& InElementWorldHandle, const UE::Widget::EWidgetMode InWidgetMode, const ETypedElementViewportInteractionGizmoManipulationType InManipulationType) override;
+	UE_API virtual void PostGizmoManipulationStopped(TArrayView<const FTypedElementHandle> InElementHandles, const UE::Widget::EWidgetMode InWidgetMode) override;
 
 private:
-	void ModifyScale(AActor* InActor, const EAxisList::Type InDragAxis, FVector& ScaleDelta, bool bCheckSmallExtent);
+	UE_API void ModifyScale(AActor* InActor, const EAxisList::Type InDragAxis, FVector& ScaleDelta, bool bCheckSmallExtent);
 };
+
+#undef UE_API
diff --git a/Engine/Source/Editor/LevelEditor/Public/Elements/Component/ComponentElementLevelEditorViewportInteractionCustomization.h b/Engine/Source/Editor/LevelEditor/Public/Elements/Component/ComponentElementLevelEditorViewportInteractionCustomization.h
--- a/Engine/Source/Editor/LevelEditor/Public/Elements/Component/ComponentElementLevelEditorViewportInteractionCustomization.h
+++ b/Engine/Source/Editor/LevelEditor/Public/Elements/Component/ComponentElementLevelEditorViewportInteractionCustomization.h
@@ -5,18 +5,22 @@
 #include "Elements/Component/ComponentElementEditorViewportInteractionCustomization.h"
 #include "Elements/Framework/TypedElementAssetEditorLevelEditorViewportClientMixin.h"
 
+#define UE_API LEVELEDITOR_API
+
 class FLevelEditorViewportClient;
 
-class LEVELEDITOR_API FComponentElementLevelEditorViewportInteractionCustomization : public FComponentElementEditorViewportInteractionCustomization, public FTypedElementAssetEditorLevelEditorViewportClientMixin
+class FComponentElementLevelEditorViewportInteractionCustomization : public FComponentElementEditorViewportInteractionCustomization, public FTypedElementAssetEditorLevelEditorViewportClientMixin
 {
 public:
-	virtual void GizmoManipulationStarted(const TTypedElement<ITypedElementWorldInterface>& InElementWorldHandle, const UE::Widget::EWidgetMode InWidgetMode) override;
-	virtual void GizmoManipulationDeltaUpdate(const TTypedElement<ITypedElementWorldInterface>& InElementWorldHandle, const UE::Widget::EWidgetMode InWidgetMode, const EAxisList::Type InDragAxis, const FInputDeviceState& InInputState, const FTransform& InDeltaTransform, const FVector& InPivotLocation) override;
-	virtual void GizmoManipulationStopped(const TTypedElement<ITypedElementWorldInterface>& InElementWorldHandle, const UE::Widget::EWidgetMode InWidgetMode, const ETypedElementViewportInteractionGizmoManipulationType InManipulationType) override;
+	UE_API virtual void GizmoManipulationStarted(const TTypedElement<ITypedElementWorldInterface>& InElementWorldHandle, const UE::Widget::EWidgetMode InWidgetMode) override;
+	UE_API virtual void GizmoManipulationDeltaUpdate(const TTypedElement<ITypedElementWorldInterface>& InElementWorldHandle, const UE::Widget::EWidgetMode InWidgetMode, const EAxisList::Type InDragAxis, const FInputDeviceState& InInputState, const FTransform& InDeltaTransform, const FVector& InPivotLocation) override;
+	UE_API virtual void GizmoManipulationStopped(const TTypedElement<ITypedElementWorldInterface>& InElementWorldHandle, const UE::Widget::EWidgetMode InWidgetMode, const ETypedElementViewportInteractionGizmoManipulationType InManipulationType) override;
 
-	static void ValidateScale(const FVector& InOriginalPreDragScale, const EAxisList::Type InDragAxis, const FVector& InCurrentScale, const FVector& InBoxExtent, FVector& InOutScaleDelta, bool bInCheckSmallExtent);
-	static FProperty* GetEditTransformProperty(const UE::Widget::EWidgetMode InWidgetMode);
+	static UE_API void ValidateScale(const FVector& InOriginalPreDragScale, const EAxisList::Type InDragAxis, const FVector& InCurrentScale, const FVector& InBoxExtent, FVector& InOutScaleDelta, bool bInCheckSmallExtent);
+	static UE_API FProperty* GetEditTransformProperty(const UE::Widget::EWidgetMode InWidgetMode);
 
 private:
-	void ModifyScale(USceneComponent* InComponent, const EAxisList::Type InDragAxis, FVector& ScaleDelta);
+	UE_API void ModifyScale(USceneComponent* InComponent, const EAxisList::Type InDragAxis, FVector& ScaleDelta);
 };
+
+#undef UE_API
diff --git a/Engine/Source/Editor/LevelEditor/Public/LevelEditor.h b/Engine/Source/Editor/LevelEditor/Public/LevelEditor.h
--- a/Engine/Source/Editor/LevelEditor/Public/LevelEditor.h
+++ b/Engine/Source/Editor/LevelEditor/Public/LevelEditor.h
@@ -14,6 +14,8 @@
 #include "ViewportTypeDefinition.h"
 #include "LevelEditorOutlinerSettings.h"
 
+#define UE_API LEVELEDITOR_API
+
 class AActor;
 class IAssetViewport;
 class SDockTab;
@@ -28,38 +30,38 @@ extern const FName LevelEditorApp;
 
 DECLARE_DELEGATE_RetVal_OneParam(bool, FAreObjectsEditable, const TArray<TWeakObjectPtr<UObject>>&);
 
-class LEVELEDITOR_API LevelEditorTabIds
+class LevelEditorTabIds
 {
 public:
-	static const FName LevelEditorViewport;
-	static const FName LevelEditorViewport_Clone1;
-	static const FName LevelEditorViewport_Clone2;
-	static const FName LevelEditorViewport_Clone3;
-	static const FName LevelEditorViewport_Clone4;
-	static const FName LevelEditorToolBox;
-	static const FName LevelEditorSelectionDetails;
-	static const FName LevelEditorSelectionDetails2;
-	static const FName LevelEditorSelectionDetails3;
-	static const FName LevelEditorSelectionDetails4;
-	static const FName PlacementBrowser;
-	static const FName LevelEditorBuildAndSubmit;
-	static const FName LevelEditorSceneOutliner;
-	static const FName LevelEditorSceneOutliner2;
-	static const FName LevelEditorSceneOutliner3;
-	static const FName LevelEditorSceneOutliner4;
-	static const FName LevelEditorStatsViewer;
-	static const FName LevelEditorLayerBrowser;
-	static const FName LevelEditorDataLayerBrowser;
-	static const FName Sequencer;
-	static const FName SequencerGraphEditor;
-	static const FName WorldSettings;
-	static const FName WorldBrowserComposition;
-	static const FName WorldBrowserPartitionEditor;
-	static const FName WorldBrowserHierarchy;
-	static const FName WorldBrowserDetails;
-	static const FName LevelEditorHierarchicalLODOutliner;
-	static const FName OutputLog;
-	static const FName LevelEditorEnvironmentLightingViewer;
+	static UE_API const FName LevelEditorViewport;
+	static UE_API const FName LevelEditorViewport_Clone1;
+	static UE_API const FName LevelEditorViewport_Clone2;
+	static UE_API const FName LevelEditorViewport_Clone3;
+	static UE_API const FName LevelEditorViewport_Clone4;
+	static UE_API const FName LevelEditorToolBox;
+	static UE_API const FName LevelEditorSelectionDetails;
+	static UE_API const FName LevelEditorSelectionDetails2;
+	static UE_API const FName LevelEditorSelectionDetails3;
+	static UE_API const FName LevelEditorSelectionDetails4;
+	static UE_API const FName PlacementBrowser;
+	static UE_API const FName LevelEditorBuildAndSubmit;
+	static UE_API const FName LevelEditorSceneOutliner;
+	static UE_API const FName LevelEditorSceneOutliner2;
+	static UE_API const FName LevelEditorSceneOutliner3;
+	static UE_API const FName LevelEditorSceneOutliner4;
+	static UE_API const FName LevelEditorStatsViewer;
+	static UE_API const FName LevelEditorLayerBrowser;
+	static UE_API const FName LevelEditorDataLayerBrowser;
+	static UE_API const FName Sequencer;
+	static UE_API const FName SequencerGraphEditor;
+	static UE_API const FName WorldSettings;
+	static UE_API const FName WorldBrowserComposition;
+	static UE_API const FName WorldBrowserPartitionEditor;
+	static UE_API const FName WorldBrowserHierarchy;
+	static UE_API const FName WorldBrowserDetails;
+	static UE_API const FName LevelEditorHierarchicalLODOutliner;
+	static UE_API const FName OutputLog;
+	static UE_API const FName LevelEditorEnvironmentLightingViewer;
 };
 
 /**
@@ -527,3 +529,5 @@ private:
 	/** The settings used to create any Outliners in the Level Editor */
 	TSharedPtr<FLevelEditorOutlinerSettings> OutlinerSettings;
 };
+
+#undef UE_API
diff --git a/Engine/Source/Editor/LevelEditor/Public/LevelEditorActions.h b/Engine/Source/Editor/LevelEditor/Public/LevelEditorActions.h
--- a/Engine/Source/Editor/LevelEditor/Public/LevelEditorActions.h
+++ b/Engine/Source/Editor/LevelEditor/Public/LevelEditorActions.h
@@ -17,6 +17,8 @@
 #include "Elements/Framework/TypedElementHandle.h"
 #include "Elements/Interfaces/TypedElementWorldInterface.h"
 
+#define UE_API LEVELEDITOR_API
+
 class SLevelViewport;
 class FLightingBuildOptions;
 class SLevelEditor;
@@ -680,7 +682,7 @@ public:
 /**
  * Implementation of various level editor action callback functions
  */
-class LEVELEDITOR_API FLevelEditorActionCallbacks
+class FLevelEditorActionCallbacks
 {
 public:
 
@@ -688,373 +690,373 @@ public:
 	 * The default can execute action for all commands unless they override it
 	 * By default commands cannot be executed if the application is in K2 debug mode.
 	 */
-	static bool DefaultCanExecuteAction();
+	static UE_API bool DefaultCanExecuteAction();
 
 	/** Opens the global documentation homepage */
-	static void BrowseDocumentation();
+	static UE_API void BrowseDocumentation();
 
 	/** Opens the viewport controls page*/
-	static void BrowseViewportControls();
+	static UE_API void BrowseViewportControls();
 
 	/** Creates a new level */
-	static void NewLevel();
-	static void NewLevel(bool& bOutLevelCreated);
+	static UE_API void NewLevel();
+	static UE_API void NewLevel(bool& bOutLevelCreated);
 	DECLARE_DELEGATE_OneParam(FNewLevelOverride, bool& /*bOutLevelCreated*/);
-	static FNewLevelOverride NewLevelOverride;
-	static bool NewLevel_CanExecute();
+	static UE_API FNewLevelOverride NewLevelOverride;
+	static UE_API bool NewLevel_CanExecute();
 
 	/** Opens an existing level */
-	static void OpenLevel();
-	static bool OpenLevel_CanExecute();
+	static UE_API void OpenLevel();
+	static UE_API bool OpenLevel_CanExecute();
 
 	/** Toggles VR mode */
-	static void ToggleVR();
-	static bool ToggleVR_CanExecute();
-	static bool ToggleVR_IsButtonActive();
-	static bool ToggleVR_IsChecked();
+	static UE_API void ToggleVR();
+	static UE_API bool ToggleVR_CanExecute();
+	static UE_API bool ToggleVR_IsButtonActive();
+	static UE_API bool ToggleVR_IsChecked();
 
 	/** Opens delta transform */
-	static void DeltaTransform();
+	static UE_API void DeltaTransform();
 
 	/**
 	 * Opens a recent file
 	 *
 	 * @param	RecentFileIndex		Index into our MRU list of recent files that can be opened
 	 */
-	static void OpenRecentFile( int32 RecentFileIndex );
+	static UE_API void OpenRecentFile( int32 RecentFileIndex );
 
 	/** Clear the list of recent files. */
-	static void ClearRecentFiles();
+	static UE_API void ClearRecentFiles();
 
 	/**
 	 * Opens a favorite file
 	 *
 	 * @param	FavoriteFileIndex		Index into our list of favorite files that can be opened
 	 */
-	static void OpenFavoriteFile( int32 FavoriteFileIndex );
+	static UE_API void OpenFavoriteFile( int32 FavoriteFileIndex );
 
-	static void ToggleFavorite();
+	static UE_API void ToggleFavorite();
 
 	/**
 	 * Remove a favorite file from the favorites list
 	 *
 	 * @param	FavoriteFileIndex		Index into our list of favorite files to be removed
 	 */
-	static void RemoveFavorite( int32 FavoriteFileIndex );
+	static UE_API void RemoveFavorite( int32 FavoriteFileIndex );
 
-	static bool ToggleFavorite_CanExecute();
-	static bool ToggleFavorite_IsChecked();
+	static UE_API bool ToggleFavorite_CanExecute();
+	static UE_API bool ToggleFavorite_IsChecked();
 
 	/** Determine whether the level can be saved at this moment */
-	static bool CanSaveWorld();
-	static bool CanSaveUnpartitionedWorld();
+	static UE_API bool CanSaveWorld();
+	static UE_API bool CanSaveUnpartitionedWorld();
 
 	/** Save the current level as... */
-	static bool CanSaveCurrentAs();
-	static void SaveCurrentAs();
+	static UE_API bool CanSaveCurrentAs();
+	static UE_API void SaveCurrentAs();
 
 	/** Saves the current map */
-	static void Save();
+	static UE_API void Save();
 
 	/** Saves all unsaved maps (but not packages) */
-	static void SaveAllLevels();
+	static UE_API void SaveAllLevels();
 
 	/** Browses to the current map */
-	static void Browse();
-	static bool CanBrowse();
+	static UE_API void Browse();
+	static UE_API bool CanBrowse();
 
 
 	/**
 	 * Called when import scene is selected
 	 */
-	static void ImportScene_Clicked();
+	static UE_API void ImportScene_Clicked();
 	
 	/**
 	* Called When Preview Json is selected in the Platforms Preview Sub Menu
 	*/
-	static void PreviewJson_Clicked(FName PlatformName, FName PreviewShaderPlatformName, FString JsonFile);
+	static UE_API void PreviewJson_Clicked(FName PlatformName, FName PreviewShaderPlatformName, FString JsonFile);
 
 	/**
 	* Is Preview Json visible in the Platforms Preview Sub Menu
 	*/
-	static bool IsPreviewJsonVisible(FName PlatformName);
+	static UE_API bool IsPreviewJsonVisible(FName PlatformName);
 
 	/**
 	* Called When Generate Preview Json is selected in the Platforms Preview Sub Menu
 	*/
-	static void GeneratePreviewJson_Clicked(FString PlatformName);
+	static UE_API void GeneratePreviewJson_Clicked(FString PlatformName);
 	
 	/**
 	* Is Generate Preview Json visible in the Platforms Preview Sub Menu
 	*/
-	static bool IsGeneratePreviewJsonVisible(FName PlatformName);
+	static UE_API bool IsGeneratePreviewJsonVisible(FName PlatformName);
 
 	/**
 	 * Called when export all is selected
 	 */
-	static void ExportAll_Clicked();
+	static UE_API void ExportAll_Clicked();
 
 
 	/**
 	 * Called when export selected is clicked
 	 */
-	static void ExportSelected_Clicked();
+	static UE_API void ExportSelected_Clicked();
 
 
 	/**
 	 * @return	True if the export selected option is available to execute
 	 */
-	static bool ExportSelected_CanExecute();
+	static UE_API bool ExportSelected_CanExecute();
 
 
-	static void ConfigureLightingBuildOptions( const FLightingBuildOptions& Options );
+	static UE_API void ConfigureLightingBuildOptions( const FLightingBuildOptions& Options );
 
-	static bool CanBuildLighting();
-	static bool CanBuildReflectionCaptures();
+	static UE_API bool CanBuildLighting();
+	static UE_API bool CanBuildReflectionCaptures();
 
 	/**
 	 * Build callbacks
 	 */
-	static void Build_Execute();
-	static bool Build_CanExecute();
-	static void BuildAndSubmitToSourceControl_Execute();
-	static void BuildLightingOnly_Execute();
-	static bool BuildLighting_CanExecute();
-	static void BuildReflectionCapturesOnly_Execute();
-	static bool BuildReflectionCapturesOnly_CanExecute();
-	static void BuildLightingOnly_VisibilityOnly_Execute();
-	static bool LightingBuildOptions_UseErrorColoring_IsChecked();
-	static void LightingBuildOptions_UseErrorColoring_Toggled();
-	static bool LightingBuildOptions_ShowLightingStats_IsChecked();
-	static void LightingBuildOptions_ShowLightingStats_Toggled();
-	static void BuildGeometryOnly_Execute();
-	static void BuildGeometryOnly_OnlyCurrentLevel_Execute();
-	static void BuildPathsOnly_Execute();
-	static bool IsWorldPartitionEnabled();
-	static bool IsWorldPartitionStreamingEnabled();
-	static void BuildHLODs_Execute();
-	static void BuildMinimap_Execute();
-	static void BuildLandscapeSplineMeshes_Execute();
-	static void BuildTextureStreamingOnly_Execute();
-	static void BuildVirtualTextureOnly_Execute();
-	static void BuildAllLandscape_Execute();
-	static bool BuildExternalType_CanExecute( int32 Index );
-	static void BuildExternalType_Execute( int32 Index );
-	static void SetLightingQuality( ELightingBuildQuality NewQuality );
-	static bool IsLightingQualityChecked( ELightingBuildQuality TestQuality );
-	static float GetLightingDensityIdeal();
-	static void SetLightingDensityIdeal( float Value );
-	static float GetLightingDensityMaximum();
-	static void SetLightingDensityMaximum( float Value );
-	static float GetLightingDensityColorScale();
-	static void SetLightingDensityColorScale( float Value );
-	static float GetLightingDensityGrayscaleScale();
-	static void SetLightingDensityGrayscaleScale( float Value );
-	static void SetLightingDensityRenderGrayscale();
-	static bool IsLightingDensityRenderGrayscaleChecked();
-	static void SetLightingResolutionStaticMeshes( ECheckBoxState NewCheckedState );
-	static ECheckBoxState IsLightingResolutionStaticMeshesChecked();
-	static void SetLightingResolutionBSPSurfaces( ECheckBoxState NewCheckedState );
-	static ECheckBoxState IsLightingResolutionBSPSurfacesChecked();
-	static void SetLightingResolutionLevel( FLightmapResRatioAdjustSettings::AdjustLevels NewLevel );
-	static bool IsLightingResolutionLevelChecked( FLightmapResRatioAdjustSettings::AdjustLevels TestLevel );
-	static void SetLightingResolutionSelectedObjectsOnly();
-	static bool IsLightingResolutionSelectedObjectsOnlyChecked();
-	static float GetLightingResolutionMinSMs();
-	static void SetLightingResolutionMinSMs( float Value );
-	static float GetLightingResolutionMaxSMs();
-	static void SetLightingResolutionMaxSMs( float Value );
-	static float GetLightingResolutionMinBSPs();
-	static void SetLightingResolutionMinBSPs( float Value );
-	static float GetLightingResolutionMaxBSPs();
-	static void SetLightingResolutionMaxBSPs( float Value );
-	static int32 GetLightingResolutionRatio();
-	static void SetLightingResolutionRatio( int32 Value );
-	static void SetLightingResolutionRatioCommit( int32 Value, ETextCommit::Type CommitInfo);
-	static void ShowLightingStaticMeshInfo();
-	static void AttachToActor(AActor* ParentActorPtr );
-	static void AttachToSocketSelection(FName SocketName, AActor* ParentActorPtr);
-	static void SetMaterialQualityLevel( EMaterialQualityLevel::Type NewQualityLevel );
-	static bool IsMaterialQualityLevelChecked( EMaterialQualityLevel::Type TestQualityLevel );
-	static void ToggleFeatureLevelPreview();
-	static bool IsFeatureLevelPreviewEnabled();
-	static bool IsFeatureLevelPreviewActive();
-	static bool IsPreviewModeButtonVisible();
-	static void SetPreviewPlatform(FPreviewPlatformInfo NewPreviewPlatform);
-	static bool CanExecutePreviewPlatform(FPreviewPlatformInfo NewPreviewPlatform);
-	static bool IsPreviewPlatformChecked(FPreviewPlatformInfo NewPreviewPlatform);
-	static void GeometryCollection_SelectAllGeometry();
-	static void GeometryCollection_SelectNone();
-	static void GeometryCollection_SelectInverseGeometry();
-	static bool GeometryCollection_IsChecked();
-
-	static void ToggleAllowArcballRotation();
-	static bool IsAllowArcballRotationChecked();
-
-	static void ToggleAllowScreenspaceRotation();
-	static bool IsAllowScreenspaceRotationChecked();
-
-	static void ToggleEnableViewportHoverFeedback();
-	static bool IsEnableViewportHoverFeedbackChecked();
-
-	static void TogglePreviewSelectedCameras(const TWeakPtr<SLevelViewport>& InLevelViewportWeak = {});
-	static bool IsPreviewSelectedCamerasChecked();
-
-	static void ToggleOrbitCameraAroundSelection();
-	static bool IsOrbitCameraAroundSelectionChecked();
-
-	static void ToggleLinkOrthographicViewports();
-	static bool IsLinkOrthographicViewportsChecked();
-
-	static void ToggleOrthoZoomToCursor();
-	static bool IsOrthoZoomToCursorChecked();
-
-	static void ToggleInvertMiddleMousePan();
-	static bool IsInvertMiddleMousePanChecked();
+	static UE_API void Build_Execute();
+	static UE_API bool Build_CanExecute();
+	static UE_API void BuildAndSubmitToSourceControl_Execute();
+	static UE_API void BuildLightingOnly_Execute();
+	static UE_API bool BuildLighting_CanExecute();
+	static UE_API void BuildReflectionCapturesOnly_Execute();
+	static UE_API bool BuildReflectionCapturesOnly_CanExecute();
+	static UE_API void BuildLightingOnly_VisibilityOnly_Execute();
+	static UE_API bool LightingBuildOptions_UseErrorColoring_IsChecked();
+	static UE_API void LightingBuildOptions_UseErrorColoring_Toggled();
+	static UE_API bool LightingBuildOptions_ShowLightingStats_IsChecked();
+	static UE_API void LightingBuildOptions_ShowLightingStats_Toggled();
+	static UE_API void BuildGeometryOnly_Execute();
+	static UE_API void BuildGeometryOnly_OnlyCurrentLevel_Execute();
+	static UE_API void BuildPathsOnly_Execute();
+	static UE_API bool IsWorldPartitionEnabled();
+	static UE_API bool IsWorldPartitionStreamingEnabled();
+	static UE_API void BuildHLODs_Execute();
+	static UE_API void BuildMinimap_Execute();
+	static UE_API void BuildLandscapeSplineMeshes_Execute();
+	static UE_API void BuildTextureStreamingOnly_Execute();
+	static UE_API void BuildVirtualTextureOnly_Execute();
+	static UE_API void BuildAllLandscape_Execute();
+	static UE_API bool BuildExternalType_CanExecute( int32 Index );
+	static UE_API void BuildExternalType_Execute( int32 Index );
+	static UE_API void SetLightingQuality( ELightingBuildQuality NewQuality );
+	static UE_API bool IsLightingQualityChecked( ELightingBuildQuality TestQuality );
+	static UE_API float GetLightingDensityIdeal();
+	static UE_API void SetLightingDensityIdeal( float Value );
+	static UE_API float GetLightingDensityMaximum();
+	static UE_API void SetLightingDensityMaximum( float Value );
+	static UE_API float GetLightingDensityColorScale();
+	static UE_API void SetLightingDensityColorScale( float Value );
+	static UE_API float GetLightingDensityGrayscaleScale();
+	static UE_API void SetLightingDensityGrayscaleScale( float Value );
+	static UE_API void SetLightingDensityRenderGrayscale();
+	static UE_API bool IsLightingDensityRenderGrayscaleChecked();
+	static UE_API void SetLightingResolutionStaticMeshes( ECheckBoxState NewCheckedState );
+	static UE_API ECheckBoxState IsLightingResolutionStaticMeshesChecked();
+	static UE_API void SetLightingResolutionBSPSurfaces( ECheckBoxState NewCheckedState );
+	static UE_API ECheckBoxState IsLightingResolutionBSPSurfacesChecked();
+	static UE_API void SetLightingResolutionLevel( FLightmapResRatioAdjustSettings::AdjustLevels NewLevel );
+	static UE_API bool IsLightingResolutionLevelChecked( FLightmapResRatioAdjustSettings::AdjustLevels TestLevel );
+	static UE_API void SetLightingResolutionSelectedObjectsOnly();
+	static UE_API bool IsLightingResolutionSelectedObjectsOnlyChecked();
+	static UE_API float GetLightingResolutionMinSMs();
+	static UE_API void SetLightingResolutionMinSMs( float Value );
+	static UE_API float GetLightingResolutionMaxSMs();
+	static UE_API void SetLightingResolutionMaxSMs( float Value );
+	static UE_API float GetLightingResolutionMinBSPs();
+	static UE_API void SetLightingResolutionMinBSPs( float Value );
+	static UE_API float GetLightingResolutionMaxBSPs();
+	static UE_API void SetLightingResolutionMaxBSPs( float Value );
+	static UE_API int32 GetLightingResolutionRatio();
+	static UE_API void SetLightingResolutionRatio( int32 Value );
+	static UE_API void SetLightingResolutionRatioCommit( int32 Value, ETextCommit::Type CommitInfo);
+	static UE_API void ShowLightingStaticMeshInfo();
+	static UE_API void AttachToActor(AActor* ParentActorPtr );
+	static UE_API void AttachToSocketSelection(FName SocketName, AActor* ParentActorPtr);
+	static UE_API void SetMaterialQualityLevel( EMaterialQualityLevel::Type NewQualityLevel );
+	static UE_API bool IsMaterialQualityLevelChecked( EMaterialQualityLevel::Type TestQualityLevel );
+	static UE_API void ToggleFeatureLevelPreview();
+	static UE_API bool IsFeatureLevelPreviewEnabled();
+	static UE_API bool IsFeatureLevelPreviewActive();
+	static UE_API bool IsPreviewModeButtonVisible();
+	static UE_API void SetPreviewPlatform(FPreviewPlatformInfo NewPreviewPlatform);
+	static UE_API bool CanExecutePreviewPlatform(FPreviewPlatformInfo NewPreviewPlatform);
+	static UE_API bool IsPreviewPlatformChecked(FPreviewPlatformInfo NewPreviewPlatform);
+	static UE_API void GeometryCollection_SelectAllGeometry();
+	static UE_API void GeometryCollection_SelectNone();
+	static UE_API void GeometryCollection_SelectInverseGeometry();
+	static UE_API bool GeometryCollection_IsChecked();
+
+	static UE_API void ToggleAllowArcballRotation();
+	static UE_API bool IsAllowArcballRotationChecked();
+
+	static UE_API void ToggleAllowScreenspaceRotation();
+	static UE_API bool IsAllowScreenspaceRotationChecked();
+
+	static UE_API void ToggleEnableViewportHoverFeedback();
+	static UE_API bool IsEnableViewportHoverFeedbackChecked();
+
+	static UE_API void TogglePreviewSelectedCameras(const TWeakPtr<SLevelViewport>& InLevelViewportWeak = {});
+	static UE_API bool IsPreviewSelectedCamerasChecked();
+
+	static UE_API void ToggleOrbitCameraAroundSelection();
+	static UE_API bool IsOrbitCameraAroundSelectionChecked();
+
+	static UE_API void ToggleLinkOrthographicViewports();
+	static UE_API bool IsLinkOrthographicViewportsChecked();
+
+	static UE_API void ToggleOrthoZoomToCursor();
+	static UE_API bool IsOrthoZoomToCursorChecked();
+
+	static UE_API void ToggleInvertMiddleMousePan();
+	static UE_API bool IsInvertMiddleMousePanChecked();
 	
-	static void ToggleInvertOrbitYAxis();
-	static bool IsInvertOrbitYAxisChecked();
+	static UE_API void ToggleInvertOrbitYAxis();
+	static UE_API bool IsInvertOrbitYAxisChecked();
 	
-	static void ToggleInvertRightMouseDollyYAxis();
-	static bool IsInvertRightMouseDollyYAxisChecked();
+	static UE_API void ToggleInvertRightMouseDollyYAxis();
+	static UE_API bool IsInvertRightMouseDollyYAxisChecked();
 
 	/**
 	 * Called when the Scene Stats button is clicked.  Invokes the Primitive Stats dialog.
 	 */
-	static void ShowSceneStats();
+	static UE_API void ShowSceneStats();
 
 	/**
 	 * Called when the Texture Stats button is clicked.  Invokes the Texture Stats dialog.
 	 */
-	static void ShowTextureStats();
+	static UE_API void ShowTextureStats();
 
 	/**
 	 * Called when the Map Check button is clicked.  Invokes the Map Check dialog.
 	 */
-	static void MapCheck_Execute();
+	static UE_API void MapCheck_Execute();
 
 	/** @return True if actions that should only be visible when source code is thought to be available */
-	static bool CanShowSourceCodeActions();
+	static UE_API bool CanShowSourceCodeActions();
 
 	/**
 	 * Called when the recompile buttons are clicked.
 	 */
-	static void RecompileGameCode_Clicked();
-	static bool Recompile_CanExecute();
+	static UE_API void RecompileGameCode_Clicked();
+	static UE_API bool Recompile_CanExecute();
 
 #if WITH_LIVE_CODING
 	/**
 	 * Enables live coding mode
 	 */
-	static void LiveCoding_ToggleEnabled();
+	static UE_API void LiveCoding_ToggleEnabled();
 
 	/**
 	 * Determines if live coding is enabled
 	 */
-	static bool LiveCoding_IsEnabled();
+	static UE_API bool LiveCoding_IsEnabled();
 
 	/**
 	 * Starts live coding (in manual mode)
 	 */
-	static void LiveCoding_StartSession_Clicked();
+	static UE_API void LiveCoding_StartSession_Clicked();
 
 	/**
 	 * Determines whether we can manually start live coding for the current session
 	 */
-	static bool LiveCoding_CanStartSession();
+	static UE_API bool LiveCoding_CanStartSession();
 
 	/**
 	 * Shows the console
 	 */
-	static void LiveCoding_ShowConsole_Clicked();
+	static UE_API void LiveCoding_ShowConsole_Clicked();
 
 	/**
 	 * Determines whether the console can be shown
 	 */
-	static bool LiveCoding_CanShowConsole();
+	static UE_API bool LiveCoding_CanShowConsole();
 
 	/**
 	 * Shows the settings panel
 	 */
-	static void LiveCoding_Settings_Clicked();
+	static UE_API void LiveCoding_Settings_Clicked();
 #endif
 
 	/**
 	 * Called when requesting connection to source control
 	 */
-	static void ConnectToSourceControl_Clicked();
+	static UE_API void ConnectToSourceControl_Clicked();
 
 	/**
 	 * Called when Check Out Modified Files is clicked
 	 */
-	static void CheckOutModifiedFiles_Clicked();
-	static bool CheckOutModifiedFiles_CanExecute();
+	static UE_API void CheckOutModifiedFiles_Clicked();
+	static UE_API bool CheckOutModifiedFiles_CanExecute();
 
 	/**
 	 * Called when Submit to Source Control is clicked
 	 */
-	static void SubmitToSourceControl_Clicked();
-	static bool SubmitToSourceControl_CanExecute();
+	static UE_API void SubmitToSourceControl_Clicked();
+	static UE_API bool SubmitToSourceControl_CanExecute();
 
 	/**
 	 * Called when the FindInContentBrowser command is executed
 	 */
-	static void FindInContentBrowser_Clicked();
-	static bool FindInContentBrowser_CanExecute();
+	static UE_API void FindInContentBrowser_Clicked();
+	static UE_API bool FindInContentBrowser_CanExecute();
 
 	/** Called to when "Edit Asset" is clicked */
-	static void EditAsset_Clicked( const EToolkitMode::Type ToolkitMode, TWeakPtr< class SLevelEditor > LevelEditor, bool bAskMultiple );
-	static bool EditAsset_CanExecute();
+	static UE_API void EditAsset_Clicked( const EToolkitMode::Type ToolkitMode, TWeakPtr< class SLevelEditor > LevelEditor, bool bAskMultiple );
+	static UE_API bool EditAsset_CanExecute();
 
 	/** Called to when "Open Selection in Property Matrix" is clicked */
-	static void OpenSelectionInPropertyMatrix_Clicked();
-	static bool OpenSelectionInPropertyMatrix_IsVisible();
+	static UE_API void OpenSelectionInPropertyMatrix_Clicked();
+	static UE_API bool OpenSelectionInPropertyMatrix_IsVisible();
 
 
 	/** Called when 'detach' is clicked */
-	static void DetachActor_Clicked();
-	static bool DetachActor_CanExecute();
+	static UE_API void DetachActor_Clicked();
+	static UE_API bool DetachActor_CanExecute();
 
 	/** Called when attach selected actors is pressed */
-	static void AttachSelectedActors();
+	static UE_API void AttachSelectedActors();
 
 	/** Called when the actor picker needs to be used to select a new parent actor */
-	static void AttachActorIteractive();
+	static UE_API void AttachActorIteractive();
 
 	/** @return true if the selected actor can be attached to the given parent actor */
-	static bool IsAttachableActor( const AActor* const ParentActor );
+	static UE_API bool IsAttachableActor( const AActor* const ParentActor );
 
 	/** Called when create new outliner folder is clicked */
-	static void CreateNewOutlinerFolder_Clicked();
+	static UE_API void CreateNewOutlinerFolder_Clicked();
 
 	/** Called when the go here command is clicked 
 	 * 
 	 * @param Point	- Specified point to go to.  If null, a point will be calculated from current mouse position
 	 */
-	static void GoHere_Clicked( const FVector* Point );
+	static UE_API void GoHere_Clicked( const FVector* Point );
 
 	/** Called when selected actor can be used to start a play session */
-	static void PlayFromHere_Clicked(bool bFloatingWindow);
-	static bool PlayFromHere_IsVisible();
+	static UE_API void PlayFromHere_Clicked(bool bFloatingWindow);
+	static UE_API bool PlayFromHere_IsVisible();
 
 	/** Called when 'Go to Code for Actor' is clicked */
-	static void GoToCodeForActor_Clicked();
-	static bool GoToCodeForActor_CanExecute();
-	static bool GoToCodeForActor_IsVisible();
+	static UE_API void GoToCodeForActor_Clicked();
+	static UE_API bool GoToCodeForActor_CanExecute();
+	static UE_API bool GoToCodeForActor_IsVisible();
 
 	/** Called when 'Go to Documentation for Actor' is clicked */
-	static void GoToDocsForActor_Clicked();
+	static UE_API void GoToDocsForActor_Clicked();
 
 /**
 	 * Called when the LockActorMovement command is executed
 	 */
-	static void LockActorMovement_Clicked();
+	static UE_API void LockActorMovement_Clicked();
 
 		
 	/**
 	 * @return true if the lock actor menu option should appear checked
 	 */
-	static bool LockActorMovement_IsChecked();
+	static UE_API bool LockActorMovement_IsChecked();
 
 	/**
 	 * Called when the AddActor command is executed
@@ -1063,396 +1065,396 @@ public:
 	 * @param bUsePlacement		Whether to use the placement editor. If not, the actor will be placed at the last click.
 	 * @param ActorLocation		[opt] If NULL, positions the actor at the mouse location, otherwise the location specified. Default is true.
 	 */
-	static void AddActor_Clicked( UActorFactory* ActorFactory, FAssetData AssetData);
-	static AActor* AddActor( UActorFactory* ActorFactory, const FAssetData& AssetData, const FTransform* ActorLocation );
+	static UE_API void AddActor_Clicked( UActorFactory* ActorFactory, FAssetData AssetData);
+	static UE_API AActor* AddActor( UActorFactory* ActorFactory, const FAssetData& AssetData, const FTransform* ActorLocation );
 
 	/**
 	 * Called when the AddActor command is executed and a class is selected in the actor browser
 	 *
 	 * @param ActorClass		The class of the actor to add
 	 */
-	static void AddActorFromClass_Clicked( UClass* ActorClass );
-	static AActor* AddActorFromClass( UClass* ActorClass );
+	static UE_API void AddActorFromClass_Clicked( UClass* ActorClass );
+	static UE_API AActor* AddActorFromClass( UClass* ActorClass );
 
 	/**
 	 * Replaces currently selected actors with an actor from the given actor factory
 	 *
 	 * @param ActorFactory	The actor factory to use in replacement
 	 */
-	static void ReplaceActors_Clicked( UActorFactory* ActorFactory, FAssetData AssetData );
-	static AActor* ReplaceActors( UActorFactory* ActorFactory, const FAssetData& AssetData, bool bCopySourceProperties = true );
+	static UE_API void ReplaceActors_Clicked( UActorFactory* ActorFactory, FAssetData AssetData );
+	static UE_API AActor* ReplaceActors( UActorFactory* ActorFactory, const FAssetData& AssetData, bool bCopySourceProperties = true );
 
 	/**
 	 * Called when the ReplaceActor command is executed and a class is selected in the actor browser
 	 *
 	 * @param ActorClass	The class of the actor to replace
 	 */
-	static void ReplaceActorsFromClass_Clicked( UClass* ActorClass );
+	static UE_API void ReplaceActorsFromClass_Clicked( UClass* ActorClass );
 
 	/**
 	 * Called to check to see if the Edit commands can be executed
 	 *
 	 * @return true, if the operation can be performed
 	 */
-	static bool Duplicate_CanExecute();
-	static bool Delete_CanExecute();
-	static void Rename_Execute();
-	static bool Rename_CanExecute();
-	static bool Cut_CanExecute();
-	static bool Copy_CanExecute();
-	static bool Paste_CanExecute();
-	static bool PasteHere_CanExecute();
+	static UE_API bool Duplicate_CanExecute();
+	static UE_API bool Delete_CanExecute();
+	static UE_API void Rename_Execute();
+	static UE_API bool Rename_CanExecute();
+	static UE_API bool Cut_CanExecute();
+	static UE_API bool Copy_CanExecute();
+	static UE_API bool Paste_CanExecute();
+	static UE_API bool PasteHere_CanExecute();
 
 	/**
 	 * Called when many of the menu items in the level editor context menu are clicked
 	 *
 	 * @param Command	The command to execute
 	 */
-	static void ExecuteExecCommand( FString Command );
+	static UE_API void ExecuteExecCommand( FString Command );
 	
 	/**
 	 * Called when selecting all actors of the same class that is selected
 	 *
 	 * @param bArchetype	true to also check that the archetype is the same
 	 */
-	static void OnSelectAllActorsOfClass( bool bArchetype );
+	static UE_API void OnSelectAllActorsOfClass( bool bArchetype );
 
 	/** Called to see if all selected actors are the same class */
-	static bool CanSelectAllActorsOfClass();
+	static UE_API bool CanSelectAllActorsOfClass();
 
 	/** Called when selecting the actor that owns the currently selected component(s) */
-	static void OnSelectComponentOwnerActor();
+	static UE_API void OnSelectComponentOwnerActor();
 
 	/** Called to see if any components are selected */
-	static bool CanSelectComponentOwnerActor();
+	static UE_API bool CanSelectComponentOwnerActor();
 
 	/**
 	 * Called to select all lights
 	 */
-	static void OnSelectAllLights();
+	static UE_API void OnSelectAllLights();
 
 	/** Selects stationary lights that are exceeding the overlap limit. */
-	static void OnSelectStationaryLightsExceedingOverlap();
+	static UE_API void OnSelectStationaryLightsExceedingOverlap();
 
 	/**
 	* Called when selecting an Actor's (if available) owning HLOD cluster
 	*/
-	static void OnSelectOwningHLODCluster();
+	static UE_API void OnSelectOwningHLODCluster();
 	
 	/**
 	 * Called to change bsp surface alignment
 	 *
 	 * @param AlignmentMode	The new alignment mode
 	 */
-	static void OnSurfaceAlignment( ETexAlign AligmentMode );
+	static UE_API void OnSurfaceAlignment( ETexAlign AligmentMode );
 
 	/**
 	 * Called to apply a material to selected surfaces
 	 */
-	static void OnApplyMaterialToSurface();
+	static UE_API void OnApplyMaterialToSurface();
 
 	/**
 	 * Checks to see if the selected actors can be grouped
 	 *	@return true if it can execute.
 	 */
-	static bool GroupActors_CanExecute();
+	static UE_API bool GroupActors_CanExecute();
 	
 	/**
 	 * Called when the RegroupActor command is executed
 	 */
-	static void RegroupActor_Clicked();
+	static UE_API void RegroupActor_Clicked();
 		
 	/**
 	 * Called when the UngroupActor command is executed
 	 */
-	static void UngroupActor_Clicked();
+	static UE_API void UngroupActor_Clicked();
 		
 	/**
 	 * Called when the LockGroup command is executed
 	 */
-	static void LockGroup_Clicked();
+	static UE_API void LockGroup_Clicked();
 	
 	/**
 	 * Called when the UnlockGroup command is executed
 	 */
-	static void UnlockGroup_Clicked();
+	static UE_API void UnlockGroup_Clicked();
 	
 	/**
 	 * Called when the AddActorsToGroup command is executed
 	 */
-	static void AddActorsToGroup_Clicked();
+	static UE_API void AddActorsToGroup_Clicked();
 	
 	/**
 	 * Called when the RemoveActorsFromGroup command is executed
 	 */
-	static void RemoveActorsFromGroup_Clicked();
+	static UE_API void RemoveActorsFromGroup_Clicked();
 
 	/**
 	 * Called when the location grid snap is toggled off and on
 	 */
-	static void LocationGridSnap_Clicked();
+	static UE_API void LocationGridSnap_Clicked();
 
 	/**
 	 * @return Returns whether or not location grid snap is enabled
 	 */
-	static bool LocationGridSnap_IsChecked();
+	static UE_API bool LocationGridSnap_IsChecked();
 
 	/**
 	 * Called when the rotation grid snap is toggled off and on
 	 */
-	static void RotationGridSnap_Clicked();
+	static UE_API void RotationGridSnap_Clicked();
 
 	/**
 	 * @return Returns whether or not rotation grid snap is enabled
 	 */
-	static bool RotationGridSnap_IsChecked();
+	static UE_API bool RotationGridSnap_IsChecked();
 
 	/**
 	 * Called when the scale grid snap is toggled off and on
 	 */
-	static void ScaleGridSnap_Clicked();
+	static UE_API void ScaleGridSnap_Clicked();
 
 	/**
 	 * @return Returns whether or not scale grid snap is enabled
 	 */
-	static bool ScaleGridSnap_IsChecked();
+	static UE_API bool ScaleGridSnap_IsChecked();
 
 
 	/** Called when "Keep Simulation Changes" is clicked in the viewport right click menu */
-	static void OnKeepSimulationChanges();
+	static UE_API void OnKeepSimulationChanges();
 
 	/** @return Returns true if 'Keep Simulation Changes' can be used right now */
-	static bool CanExecuteKeepSimulationChanges();
+	static UE_API bool CanExecuteKeepSimulationChanges();
 		
 		
 	/**
 	 * Makes the currently selected actors level the current level
 	 * If multiple actors are selected they must all be in the same level
 	 */
-	static void OnMakeSelectedActorLevelCurrent();
+	static UE_API void OnMakeSelectedActorLevelCurrent();
 
 	/**
 	 * Moves the currently selected actors to the current level                   
 	 */
-	static void OnMoveSelectedToCurrentLevel();
+	static UE_API void OnMoveSelectedToCurrentLevel();
 
 	/** Finds the currently selected actor(s) level in the content browser */
-	static void OnFindActorLevelInContentBrowser();
+	static UE_API void OnFindActorLevelInContentBrowser();
 
 	/** @return Returns true if all selected actors are from the same level and hence can browse to it in the content browser */
-	static bool CanExecuteFindActorLevelInContentBrowser();
+	static UE_API bool CanExecuteFindActorLevelInContentBrowser();
 
 	/**
 	 * Selects the currently selected actor(s) levels in the level browser
 	 * Deselecting everything else first
 	 */
-	static void OnFindLevelsInLevelBrowser();
+	static UE_API void OnFindLevelsInLevelBrowser();
 
 	/**
 	 * Selects the currently selected actor(s) levels in the level browser
 	 */
-	static void OnSelectLevelInLevelBrowser();
+	static UE_API void OnSelectLevelInLevelBrowser();
 
 	/**
 	 * Deselects the currently selected actor(s) levels in the level browser
 	 */
-	static void OnDeselectLevelInLevelBrowser();
+	static UE_API void OnDeselectLevelInLevelBrowser();
 
 	/**
 	 * Finds references to the currently selected actor(s) in level scripts
 	 */
-	static void OnFindActorInLevelScript();
+	static UE_API void OnFindActorInLevelScript();
 
 	/** Select the world info actor and show the properties */
-	static void OnShowWorldProperties( TWeakPtr< SLevelEditor > LevelEditor );
+	static UE_API void OnShowWorldProperties( TWeakPtr< SLevelEditor > LevelEditor );
 
 	/** Focuses the outliner on the selected actors */
-	static void OnFocusOutlinerToSelection(TWeakPtr<SLevelEditor> LevelEditor);
+	static UE_API void OnFocusOutlinerToSelection(TWeakPtr<SLevelEditor> LevelEditor);
 
 	/** Focuses the outliner on the context folder if it's currently set */
-	static void OnFocusOutlinerToContextFolder(TWeakPtr<SLevelEditor> LevelEditor);
+	static UE_API void OnFocusOutlinerToContextFolder(TWeakPtr<SLevelEditor> LevelEditor);
 
 	/** Open the Place Actors Panel */
-	static void OpenPlaceActors();
+	static UE_API void OpenPlaceActors();
 
 	/** Open the Content Browser */
-	static void OpenContentBrowser();
+	static UE_API void OpenContentBrowser();
 
 	/** Import content into a chosen location*/
-	static void ImportContent();
+	static UE_API void ImportContent();
 
 	/** Checks out the Project Settings config */
-	static void CheckOutProjectSettingsConfig();
+	static UE_API void CheckOutProjectSettingsConfig();
 
 	/** Open the level's blueprint in Kismet2 */
-	static void OpenLevelBlueprint( TWeakPtr< SLevelEditor > LevelEditor );
+	static UE_API void OpenLevelBlueprint( TWeakPtr< SLevelEditor > LevelEditor );
 
 	/** Returns TRUE if the user can edit the game mode Blueprint, this requires the DefaultEngine config file to be writable */
-	static bool CanSelectGameModeBlueprint();
+	static UE_API bool CanSelectGameModeBlueprint();
 
 	/** Helps the user create a Blueprint class */
-	static void CreateBlankBlueprintClass();
+	static UE_API void CreateBlankBlueprintClass();
 
 	/** Can the selected actors be converted to a blueprint class in any of the supported ways? */
-	static bool CanConvertSelectedActorsIntoBlueprintClass();
+	static UE_API bool CanConvertSelectedActorsIntoBlueprintClass();
 
 	/** Bring up the convert actors to blueprint UI */
-	static void ConvertSelectedActorsIntoBlueprintClass();
+	static UE_API void ConvertSelectedActorsIntoBlueprintClass();
 
 	/** Shows only selected actors, hiding any unselected actors and unhiding any selected hidden actors. */
-	static void OnShowOnlySelectedActors();
+	static UE_API void OnShowOnlySelectedActors();
 
 	/**
 	 * View callbacks
 	 */ 
-	static void OnToggleTransformWidgetVisibility();
-	static bool OnGetTransformWidgetVisibility();
-	static void OnToggleShowSelectionSubcomponents();
-	static bool OnGetShowSelectionSubcomponents();
-	static void OnAllowTranslucentSelection();
-	static bool OnIsAllowTranslucentSelectionEnabled();	
-	static void OnAllowGroupSelection();
-	static bool OnIsAllowGroupSelectionEnabled(); 
-	static void OnToggleStrictBoxSelect();
-	static bool OnIsStrictBoxSelectEnabled(); 
-	static void OnToggleTransparentBoxSelect();
-	static bool OnIsTransparentBoxSelectEnabled();
-	static void OnDrawBrushMarkerPolys();
-	static bool OnIsDrawBrushMarkerPolysEnabled();
-	static void OnToggleOnlyLoadVisibleInPIE();
-	static bool OnIsOnlyLoadVisibleInPIEEnabled(); 
-	static void OnToggleSocketSnapping();
-	static bool OnIsSocketSnappingEnabled(); 
-	static void OnToggleParticleSystemLOD();
-	static bool OnIsParticleSystemLODEnabled(); 
-	static void OnToggleFreezeParticleSimulation();
-	static bool OnIsParticleSimulationFrozen();
-	static void OnToggleParticleSystemHelpers();
-	static bool OnIsParticleSystemHelpersEnabled();
-	static void OnToggleLODViewLocking();
-	static bool OnIsLODViewLockingEnabled(); 
-	static void OnToggleLevelStreamingVolumePrevis();
-	static bool OnIsLevelStreamingVolumePrevisEnabled(); 
+	static UE_API void OnToggleTransformWidgetVisibility();
+	static UE_API bool OnGetTransformWidgetVisibility();
+	static UE_API void OnToggleShowSelectionSubcomponents();
+	static UE_API bool OnGetShowSelectionSubcomponents();
+	static UE_API void OnAllowTranslucentSelection();
+	static UE_API bool OnIsAllowTranslucentSelectionEnabled();	
+	static UE_API void OnAllowGroupSelection();
+	static UE_API bool OnIsAllowGroupSelectionEnabled(); 
+	static UE_API void OnToggleStrictBoxSelect();
+	static UE_API bool OnIsStrictBoxSelectEnabled(); 
+	static UE_API void OnToggleTransparentBoxSelect();
+	static UE_API bool OnIsTransparentBoxSelectEnabled();
+	static UE_API void OnDrawBrushMarkerPolys();
+	static UE_API bool OnIsDrawBrushMarkerPolysEnabled();
+	static UE_API void OnToggleOnlyLoadVisibleInPIE();
+	static UE_API bool OnIsOnlyLoadVisibleInPIEEnabled(); 
+	static UE_API void OnToggleSocketSnapping();
+	static UE_API bool OnIsSocketSnappingEnabled(); 
+	static UE_API void OnToggleParticleSystemLOD();
+	static UE_API bool OnIsParticleSystemLODEnabled(); 
+	static UE_API void OnToggleFreezeParticleSimulation();
+	static UE_API bool OnIsParticleSimulationFrozen();
+	static UE_API void OnToggleParticleSystemHelpers();
+	static UE_API bool OnIsParticleSystemHelpersEnabled();
+	static UE_API void OnToggleLODViewLocking();
+	static UE_API bool OnIsLODViewLockingEnabled(); 
+	static UE_API void OnToggleLevelStreamingVolumePrevis();
+	static UE_API bool OnIsLevelStreamingVolumePrevisEnabled(); 
 	
-	static FText GetAudioVolumeToolTip();
-	static float GetAudioVolume();
-	static void OnAudioVolumeChanged(float Volume);
-	static bool GetAudioMuted();
-	static void OnAudioMutedChanged(bool bMuted); 
+	static UE_API FText GetAudioVolumeToolTip();
+	static UE_API float GetAudioVolume();
+	static UE_API void OnAudioVolumeChanged(float Volume);
+	static UE_API bool GetAudioMuted();
+	static UE_API void OnAudioMutedChanged(bool bMuted); 
 
-	static void OnEnableActorSnap();
-	static bool OnIsActorSnapEnabled();
-	static FText GetActorSnapTooltip();
-	static float GetActorSnapSetting();
-	static void SetActorSnapSetting(float Distance);
-	static void OnEnableVertexSnap();
-	static bool OnIsVertexSnapEnabled();
+	static UE_API void OnEnableActorSnap();
+	static UE_API bool OnIsActorSnapEnabled();
+	static UE_API FText GetActorSnapTooltip();
+	static UE_API float GetActorSnapSetting();
+	static UE_API void SetActorSnapSetting(float Distance);
+	static UE_API void OnEnableVertexSnap();
+	static UE_API bool OnIsVertexSnapEnabled();
 
-	static void OnToggleShowViewportToolbar();
-	static bool IsViewportToolbarVisible();
+	static UE_API void OnToggleShowViewportToolbar();
+	static UE_API bool IsViewportToolbarVisible();
 
-	static void OnToggleShowViewportUI();
-	static bool IsViewportUIVisible();
+	static UE_API void OnToggleShowViewportUI();
+	static UE_API bool IsViewportUIVisible();
 
-	static bool IsEditorModeActive( FEditorModeID EditorMode );
+	static UE_API bool IsEditorModeActive( FEditorModeID EditorMode );
 
-	static void MakeBuilderBrush( UClass* BrushBuilderClass );
+	static UE_API void MakeBuilderBrush( UClass* BrushBuilderClass );
 
-	static void OnAddVolume( UClass* VolumeClass );
+	static UE_API void OnAddVolume( UClass* VolumeClass );
 
-	static void SelectActorsInLayers();
+	static UE_API void SelectActorsInLayers();
 
-	static void SetWidgetMode( UE::Widget::EWidgetMode WidgetMode );
-	static bool IsWidgetModeActive( UE::Widget::EWidgetMode WidgetMode );
-	static bool CanSetWidgetMode( UE::Widget::EWidgetMode WidgetMode );
-	static bool IsTranslateRotateModeVisible();
-	static void SetCoordinateSystem( ECoordSystem CoordSystem );
-	static bool IsCoordinateSystemActive( ECoordSystem CoordSystem );
+	static UE_API void SetWidgetMode( UE::Widget::EWidgetMode WidgetMode );
+	static UE_API bool IsWidgetModeActive( UE::Widget::EWidgetMode WidgetMode );
+	static UE_API bool CanSetWidgetMode( UE::Widget::EWidgetMode WidgetMode );
+	static UE_API bool IsTranslateRotateModeVisible();
+	static UE_API void SetCoordinateSystem( ECoordSystem CoordSystem );
+	static UE_API bool IsCoordinateSystemActive( ECoordSystem CoordSystem );
 
 	/**
 	 * Return a world
 	 */
-	static class UWorld* GetWorld();
+	static UE_API class UWorld* GetWorld();
 public:
 	/** 
 	 * Moves the selected elements to the grid.
 	 */
-	static void MoveElementsToGrid_Clicked( bool InAlign, bool InPerElement );
+	static UE_API void MoveElementsToGrid_Clicked( bool InAlign, bool InPerElement );
 
 	/**
 	* Snaps a selected actor to the camera view.
 	*/
-	static void SnapObjectToView_Clicked();
+	static UE_API void SnapObjectToView_Clicked();
 
 	/**
 	* Copy the file path where the actor is saved.
 	*/
-	static void CopyActorFilePathtoClipboard_Clicked();
+	static UE_API void CopyActorFilePathtoClipboard_Clicked();
 
 	/**
 	* Opens the reference viewer for the selected actor.
 	*/
-	static void ViewActorReferences_Clicked();
+	static UE_API void ViewActorReferences_Clicked();
 
 	/**
 	 * Checks whether OpenActorInReferenceViewer can be executed on the selected actor.
 	 * @return 	True if OpenActorInReferenceViewer can be executed on the selected actor.
 	 */
-	static bool ViewActorReferences_CanExecute();
+	static UE_API bool ViewActorReferences_CanExecute();
 
 	/**
 	* Save the actor.
 	*/
-	static void SaveActor_Clicked();
+	static UE_API void SaveActor_Clicked();
 
 	/**
 	 * Checks whether SaveActor can be executed on the selected actors.
 	 * @return 	True if SaveActor can be executed on the selected actors.
 	 */
-	static bool SaveActor_CanExecute();
+	static UE_API bool SaveActor_CanExecute();
 
 	/**
 	* Shows the history of the file containing the actor.
 	*/
-	static void ShowActorHistory_Clicked();
+	static UE_API void ShowActorHistory_Clicked();
 
 	/**
 	 * Checks whether ShowActorHistory can be executed on the selected actors.
 	 * @return 	True if ShowActorHistory can be executed on the selected actors.
 	 */
-	static bool ShowActorHistory_CanExecute();
+	static UE_API bool ShowActorHistory_CanExecute();
 
 	/** 
 	 * Moves the selected elements to the last selected element.
 	 */
-	static void MoveElementsToElement_Clicked( bool InAlign );
+	static UE_API void MoveElementsToElement_Clicked( bool InAlign );
 
 	/** 
 	 * Snaps an actor to the currently selected 2D snap layer
 	 */
-	static void SnapTo2DLayer_Clicked();
+	static UE_API void SnapTo2DLayer_Clicked();
 
 	/**
 	 * Checks to see if at least a single actor is selected and the 2D editor mode is enabled
 	 *	@return true if it can execute.
 	 */
-	static bool CanSnapTo2DLayer();
+	static UE_API bool CanSnapTo2DLayer();
 
 	/** 
 	 * Snaps an actor to the currently selected 2D snap layer
 	 */
-	static void MoveSelectionToDifferent2DLayer_Clicked(bool bGoingUp, bool bForceToTopOrBottom);
+	static UE_API void MoveSelectionToDifferent2DLayer_Clicked(bool bGoingUp, bool bForceToTopOrBottom);
 
 	/**
 	 * Checks to see if at least a single actor is selected and the 2D editor mode is enabled and there is a layer above/below the current setting
 	 *	@return true if it can execute.
 	 */
-	static bool CanMoveSelectionToDifferent2DLayer(bool bGoingUp);
+	static UE_API bool CanMoveSelectionToDifferent2DLayer(bool bGoingUp);
 
 	/**
 	 * Changes the active 2D snap layer to one a delta above or below the current layer
 	 */
-	static void Select2DLayerDeltaAway_Clicked(int32 Delta);
+	static UE_API void Select2DLayerDeltaAway_Clicked(int32 Delta);
 
 	/** 
 	 * Snaps the selected elements to the floor.  Optionally will align with the trace normal.
@@ -1461,7 +1463,7 @@ public:
 	 * @param InUseBounds		Whether or not to base the line trace off of the bounds.
 	 * @param InUsePivot		Whether or not to use the pivot position.
 	 */
-	static void SnapToFloor_Clicked( bool InAlign, bool InUseLineTrace, bool InUseBounds, bool InUsePivot );
+	static UE_API void SnapToFloor_Clicked( bool InAlign, bool InUseLineTrace, bool InUseBounds, bool InUsePivot );
 
 	/**
 	 * Snaps the selected elements to another element.  Optionally will align with the trace normal.
@@ -1470,18 +1472,18 @@ public:
 	 * @param InUseBounds		Whether or not to base the line trace off of the bounds.
 	 * @param InUsePivot		Whether or not to use the pivot position.
 	 */
-	static void SnapElementsToElement_Clicked( bool InAlign, bool InUseLineTrace, bool InUseBounds, bool InUsePivot );
+	static UE_API void SnapElementsToElement_Clicked( bool InAlign, bool InUseLineTrace, bool InUseBounds, bool InUsePivot );
 
 	/**
 	 * Aligns brush vertices to the nearest grid point.
 	 */
-	static void AlignBrushVerticesToGrid_Execute();
+	static UE_API void AlignBrushVerticesToGrid_Execute();
 
 	/**
 	 * Checks to see if at least one actor is selected
 	 *	@return true if it can execute.
 	 */
-	static bool ActorSelected_CanExecute();
+	static UE_API bool ActorSelected_CanExecute();
 
 	enum EActorTypeFlags : uint8
 	{
@@ -1498,61 +1500,61 @@ public:
 	 * @param bSingleOnly	if true, then requires selection to be exactly one actor
 	 * @return				true if it can execute.
 	 */
-	static bool ActorTypesSelected_CanExecute(EActorTypeFlags TypeFlags, bool bSingleOnly);
+	static UE_API bool ActorTypesSelected_CanExecute(EActorTypeFlags TypeFlags, bool bSingleOnly);
 
 	/**
 	 * Checks to see if multiple actors are selected
 	 *	@return true if it can execute.
 	 */
-	static bool ActorsSelected_CanExecute();
+	static UE_API bool ActorsSelected_CanExecute();
 
 	/**
 	 * Checks to see if at least one element is selected
 	 *	@return true if it can execute.
 	 */
 	UE_DEPRECATED(5.4, "Use ElementSelected_CanExecuteMove. If you only needed to verify if there is at least one item in the selection set, use the selection set directly from the level editor.")
-	static bool ElementSelected_CanExecute();
+	static UE_API bool ElementSelected_CanExecute();
 
 	/**
 	 * Checks to see if multiple elements are selected
 	 *	@return true if it can execute.
 	 */
 	UE_DEPRECATED(5.4, "Use ElementsSelected_CanExecuteMove. If you only needed to verify if there is more than one item in the selection set, use the selection set directly from the level editor.")
-	static bool ElementsSelected_CanExecute();
+	static UE_API bool ElementsSelected_CanExecute();
 
 	/**
 	 * Checks to see if at least one element is selected that can be translated
 	 *	@return true if it can execute.
 	 */
-	static bool ElementSelected_CanExecuteMove();
+	static UE_API bool ElementSelected_CanExecuteMove();
 
 	/**
 	 * Checks to see if multiple elements are selected that can be translated
 	 *	@return true if it can execute.
 	 */
-	static bool ElementsSelected_CanExecuteMove();
+	static UE_API bool ElementsSelected_CanExecuteMove();
 
 	/**
 	 * Checks to see if at least one element is selected that can be scaled
 	 *	@return true if it can execute.
 	 */
-	static bool ElementSelected_CanExecuteScale();
+	static UE_API bool ElementSelected_CanExecuteScale();
 
 	/**
 	 * Checks to see if multiple elements are selected that can be scaled
 	 *	@return true if it can execute.
 	 */
-	static bool ElementsSelected_CanExecuteScale();
+	static UE_API bool ElementsSelected_CanExecuteScale();
 
 	/** Called when 'Open Merge Actor' is clicked */
-	static void OpenMergeActor_Clicked();
+	static UE_API void OpenMergeActor_Clicked();
 
 private:
 	/** 
 	 * Moves the selected elements.
 	 * @param InDestination		The destination element we want to move this element to, or invalid to move to the grid
 	 */
-	static void MoveTo_Clicked( const UTypedElementSelectionSet* InSelectionSet, const bool InAlign, bool InPerElement, const TTypedElement<ITypedElementWorldInterface>& InDestination = TTypedElement<ITypedElementWorldInterface>() );
+	static UE_API void MoveTo_Clicked( const UTypedElementSelectionSet* InSelectionSet, const bool InAlign, bool InPerElement, const TTypedElement<ITypedElementWorldInterface>& InDestination = TTypedElement<ITypedElementWorldInterface>() );
 
 	/** 
 	 * Snaps the selected elements. Optionally will align with the trace normal.
@@ -1562,7 +1564,7 @@ private:
 	 * @param InUsePivot		Whether or not to use the pivot position.
 	 * @param InDestination		The destination element we want to move this actor to, or invalid to go towards the floor
 	 */
-	static void SnapTo_Clicked(const UTypedElementSelectionSet* InSelectionSet, const bool InAlign, const bool InUseLineTrace, const bool InUseBounds, const bool InUsePivot, const TTypedElement<ITypedElementWorldInterface>& InDestination = TTypedElement<ITypedElementWorldInterface>() );
+	static UE_API void SnapTo_Clicked(const UTypedElementSelectionSet* InSelectionSet, const bool InAlign, const bool InUseLineTrace, const bool InUseBounds, const bool InUsePivot, const TTypedElement<ITypedElementWorldInterface>& InDestination = TTypedElement<ITypedElementWorldInterface>() );
 
 	/** 
 	 * Create and apply animation to the SkeletalMeshComponent if Simulating
@@ -1570,9 +1572,10 @@ private:
 	 * @param EditorActor	Editor Counterpart Actor
 	 * @param SimActor		Simulating Actor in PIE or SIE
 	 */
-	static bool SaveAnimationFromSkeletalMeshComponent(AActor * EditorActor, AActor * SimActor, TArray<class USkeletalMeshComponent*> & OutEditorComponents);
+	static UE_API bool SaveAnimationFromSkeletalMeshComponent(AActor * EditorActor, AActor * SimActor, TArray<class USkeletalMeshComponent*> & OutEditorComponents);
 
 public:
-	static void FixupGroupActor_Clicked();
+	static UE_API void FixupGroupActor_Clicked();
 };
 
+#undef UE_API
diff --git a/Engine/Source/Editor/LevelEditor/Public/LevelEditorContextMenu.h b/Engine/Source/Editor/LevelEditor/Public/LevelEditorContextMenu.h
--- a/Engine/Source/Editor/LevelEditor/Public/LevelEditorContextMenu.h
+++ b/Engine/Source/Editor/LevelEditor/Public/LevelEditorContextMenu.h
@@ -6,6 +6,8 @@
 #include "Editor/UnrealEdTypes.h"
 #include "LevelEditorMenuContext.h"
 
+#define UE_API LEVELEDITOR_API
+
 class FExtender;
 class UToolMenu;
 class UTypedElementSelectionSet;
@@ -17,7 +19,7 @@ struct FToolMenuSection;
 /**
  * Context menu construction class 
  */
-class LEVELEDITOR_API FLevelEditorContextMenu
+class FLevelEditorContextMenu
 {
 
 public:
@@ -28,13 +30,13 @@ public:
 	 * @param	ContextType		The context we should use to specialize this menu
 	 * @param	HitProxyActor	The hitproxy actor in the case the ContextType is Viewport
 	 */
-	static void SummonMenu( const TSharedRef< class SLevelEditor >& LevelEditor, ELevelEditorMenuContext ContextType, const FTypedElementHandle& HitProxyElement = FTypedElementHandle());
+	static UE_API void SummonMenu( const TSharedRef< class SLevelEditor >& LevelEditor, ELevelEditorMenuContext ContextType, const FTypedElementHandle& HitProxyElement = FTypedElementHandle());
 
 	/**
 	 * Summons the viewport view option menu
 	 * @param LevelEditor		The level editor using this menu.
 	 */
-	static void SummonViewOptionMenu( const TSharedRef< class SLevelEditor >& LevelEditor, const ELevelViewportType ViewOption );
+	static UE_API void SummonViewOptionMenu( const TSharedRef< class SLevelEditor >& LevelEditor, const ELevelViewportType ViewOption );
 
 	/**
 	 * Creates a widget for the context menu that can be inserted into a pop-up window
@@ -45,7 +47,7 @@ public:
 	 * @param	HitProxyActor	The hitproxy actor in the case the ContextType is Viewport
 	 * @return	Widget for this context menu
 	 */
-	static TSharedRef< SWidget > BuildMenuWidget(TWeakPtr< ILevelEditor > LevelEditor, ELevelEditorMenuContext ContextType, TSharedPtr<FExtender> Extender = TSharedPtr<FExtender>(), const FTypedElementHandle& HitProxyElement = FTypedElementHandle());
+	static UE_API TSharedRef< SWidget > BuildMenuWidget(TWeakPtr< ILevelEditor > LevelEditor, ELevelEditorMenuContext ContextType, TSharedPtr<FExtender> Extender = TSharedPtr<FExtender>(), const FTypedElementHandle& HitProxyElement = FTypedElementHandle());
 
 	/**
 	 * Populates the specified menu builder for the context menu that can be inserted into a pop-up window
@@ -56,28 +58,28 @@ public:
 	 * @param	Extender		Allows extension of this menu based on context.
 	 * @param	HitProxyActor	The hitproxy actor in the case the ContextType is Viewport
 	 */
-	static UToolMenu* GenerateMenu(TWeakPtr< ILevelEditor > LevelEditor, ELevelEditorMenuContext ContextType, TSharedPtr<FExtender> Extender = TSharedPtr<FExtender>(), const FTypedElementHandle& HitProxyElement = FTypedElementHandle());
+	static UE_API UToolMenu* GenerateMenu(TWeakPtr< ILevelEditor > LevelEditor, ELevelEditorMenuContext ContextType, TSharedPtr<FExtender> Extender = TSharedPtr<FExtender>(), const FTypedElementHandle& HitProxyElement = FTypedElementHandle());
 
 	/* Adds required information to Context for build menu based on current selection */
-	static FName InitMenuContext(FToolMenuContext& Context, TWeakPtr<ILevelEditor> LevelEditor, ELevelEditorMenuContext ContextType, const FTypedElementHandle& HitProxyElement = FTypedElementHandle());
+	static UE_API FName InitMenuContext(FToolMenuContext& Context, TWeakPtr<ILevelEditor> LevelEditor, ELevelEditorMenuContext ContextType, const FTypedElementHandle& HitProxyElement = FTypedElementHandle());
 
 	/* Returns name of menu to display based on current selection */
-	static FName GetContextMenuName(ELevelEditorMenuContext ContextType, const UTypedElementSelectionSet* InSelectionSet);
+	static UE_API FName GetContextMenuName(ELevelEditorMenuContext ContextType, const UTypedElementSelectionSet* InSelectionSet);
 
 	/* Returns a user-readable title for the menu to display; the title can be displayed in UI like the menu bar */
-	static FText GetContextMenuTitle(ELevelEditorMenuContext ContextType, const UTypedElementSelectionSet* InSelectionSet);
+	static UE_API FText GetContextMenuTitle(ELevelEditorMenuContext ContextType, const UTypedElementSelectionSet* InSelectionSet);
 
 	/* Returns a user-readable tooltip describing the menu to display */
-	static FText GetContextMenuToolTip(ELevelEditorMenuContext ContextType, const UTypedElementSelectionSet* InSelectionSet);
+	static UE_API FText GetContextMenuToolTip(ELevelEditorMenuContext ContextType, const UTypedElementSelectionSet* InSelectionSet);
 
-	static void RegisterComponentContextMenu();
-	static void RegisterActorContextMenu();
-	static void RegisterElementContextMenu();
-	static void RegisterSceneOutlinerContextMenu();
-	static void RegisterMenuBarEmptyContextMenu();
-	static void RegisterEmptySelectionContextMenu();
+	static UE_API void RegisterComponentContextMenu();
+	static UE_API void RegisterActorContextMenu();
+	static UE_API void RegisterElementContextMenu();
+	static UE_API void RegisterSceneOutlinerContextMenu();
+	static UE_API void RegisterMenuBarEmptyContextMenu();
+	static UE_API void RegisterEmptySelectionContextMenu();
 
-	static void AddPlayFromHereSubMenu(FToolMenuSection& Section);
+	static UE_API void AddPlayFromHereSubMenu(FToolMenuSection& Section);
 
 	/**
 	 * Builds the actor group menu
@@ -85,5 +87,7 @@ public:
 	 * @param Menu		The menu to add items to.
 	 * @param SelectedActorInfo	Information about the selected actors.
 	 */
-	static void BuildGroupMenu(FToolMenuSection& Section, const struct FSelectedActorInfo& SelectedActorInfo);
+	static UE_API void BuildGroupMenu(FToolMenuSection& Section, const struct FSelectedActorInfo& SelectedActorInfo);
 };
+
+#undef UE_API
diff --git a/Engine/Source/Editor/LevelEditor/Public/LevelEditorMenuContext.h b/Engine/Source/Editor/LevelEditor/Public/LevelEditorMenuContext.h
--- a/Engine/Source/Editor/LevelEditor/Public/LevelEditorMenuContext.h
+++ b/Engine/Source/Editor/LevelEditor/Public/LevelEditorMenuContext.h
@@ -17,6 +17,8 @@
 
 #include "LevelEditorMenuContext.generated.h"
 
+#define UE_API LEVELEDITOR_API
+
 class AActor;
 class FLevelEditorViewportClient;
 class ILevelEditor;
@@ -27,12 +29,12 @@ class UActorComponent;
 class UTypedElementSelectionSet;
 struct FFrame;
 
-UCLASS()
-class LEVELEDITOR_API ULevelEditorMenuContext : public UObject
+UCLASS(MinimalAPI)
+class ULevelEditorMenuContext : public UObject
 {
 	GENERATED_BODY()
 public:
-	TWeakPtr<ILevelEditor> GetLevelEditor() const;
+	UE_API TWeakPtr<ILevelEditor> GetLevelEditor() const;
 
 	TWeakPtr<SLevelEditor> LevelEditor;
 };
@@ -50,8 +52,8 @@ enum class ELevelEditorMenuContext : uint8
 	MainMenu,
 };
 
-UCLASS()
-class LEVELEDITOR_API ULevelEditorContextMenuContext : public UObject
+UCLASS(MinimalAPI)
+class ULevelEditorContextMenuContext : public UObject
 {
 	GENERATED_BODY()
 public:
@@ -77,25 +79,27 @@ public:
 	TWeakObjectPtr<AActor> HitProxyActor = nullptr;
 
 	UFUNCTION(BlueprintCallable, Category = "LevelEditor | Menu", DisplayName="Get Hit Proxy Element", meta=(ScriptName="GetHitProxyElement"))
-	FScriptTypedElementHandle GetScriptHitProxyElement();
+	UE_API FScriptTypedElementHandle GetScriptHitProxyElement();
 };
 
-UCLASS()
-class LEVELEDITOR_API ULevelViewportToolBarContext : public UObject
+UCLASS(MinimalAPI)
+class ULevelViewportToolBarContext : public UObject
 {
 	GENERATED_BODY()
 public:
 	TWeakPtr<SLevelViewportToolBar> LevelViewportToolBarWidget;
 
-	virtual FLevelEditorViewportClient* GetLevelViewportClient() const;
+	UE_API virtual FLevelEditorViewportClient* GetLevelViewportClient() const;
 };
 
 
-UCLASS()
-class LEVELEDITOR_API UQuickActionMenuContext : public UObject
+UCLASS(MinimalAPI)
+class UQuickActionMenuContext : public UObject
 {
 	GENERATED_BODY()
 public:
 	UPROPERTY(VisibleAnywhere, BlueprintReadOnly, Category="LevelEditor|Menu")
 	TObjectPtr<const UTypedElementSelectionSet> CurrentSelection;
 };
+
+#undef UE_API
