Updated headers using UnrealCodeFixup to make sure dllstorage is on methods/staticvars instead of types. Part 5/n

diff --git a/Engine/Source/Editor/LevelEditor/Public/LevelEditorOutlinerSettings.h b/Engine/Source/Editor/LevelEditor/Public/LevelEditorOutlinerSettings.h
--- a/Engine/Source/Editor/LevelEditor/Public/LevelEditorOutlinerSettings.h
+++ b/Engine/Source/Editor/LevelEditor/Public/LevelEditorOutlinerSettings.h
@@ -4,6 +4,8 @@
 
 #include "Filters/GenericFilter.h"
 
+#define UE_API LEVELEDITOR_API
+
 struct ISceneOutlinerTreeItem;
 
 namespace SceneOutliner
@@ -35,13 +37,13 @@ struct FLevelEditorOutlinerBuiltInCategories
  *  Use AddCustomFilter/AddCustomClass Filter to register new filters in the Outliner.
  *  Make sure that the filters you attach have a category, otherwise they will not show up
  */
-class LEVELEDITOR_API FLevelEditorOutlinerSettings :  public TSharedFromThis<class FLevelEditorOutlinerSettings>
+class FLevelEditorOutlinerSettings :  public TSharedFromThis<class FLevelEditorOutlinerSettings>
 {
 public:
 	FLevelEditorOutlinerSettings(){}
-	~FLevelEditorOutlinerSettings();
+	UE_API ~FLevelEditorOutlinerSettings();
 	
-	void Initialize();
+	UE_API void Initialize();
 	
 	/** Delegate to create a Filter for the Outliner */
 	DECLARE_DELEGATE_RetVal(TSharedPtr<FFilterBase<const ISceneOutlinerTreeItem&>>, FOutlinerFilterFactory);
@@ -53,46 +55,46 @@ public:
 	 */
 
 	UE_DEPRECATED(5.2, "Use the AddCustomFilter override that takes in a factory instead so each Outliner can have a unique instance of the filter")
-	void AddCustomFilter(TSharedRef<FFilterBase<const ISceneOutlinerTreeItem&>> InCustomFilter);
+	UE_API void AddCustomFilter(TSharedRef<FFilterBase<const ISceneOutlinerTreeItem&>> InCustomFilter);
 
 	/**  Add a custom filter to the outliner filter bar. These are all AND'd together
 	 *   This function takes in a delegate that will be called to create an instance of the custom filter for each Outliner that
 	 *   the level editor creates.
 	 *   @see FGenericFilter on how to create generic filters
 	 */
-	void AddCustomFilter(FOutlinerFilterFactory InCreateCustomFilter);
+	UE_API void AddCustomFilter(FOutlinerFilterFactory InCreateCustomFilter);
 
 	/**  Add a custom class filter to the outliner filter bar. These represent asset/actor type filters and are OR'd
 	 *   Can be created using IAssetTypeActions or UClass (@see constructor)
 	 */
-	void AddCustomClassFilter(TSharedRef<FCustomClassFilterData> InCustomClassFilterData);
+	UE_API void AddCustomClassFilter(TSharedRef<FCustomClassFilterData> InCustomClassFilterData);
 
 	// Creates the default filters that the level editor outliner has
-	void CreateDefaultFilters();
+	UE_API void CreateDefaultFilters();
 
 	// Setup the built in filter categories
-	void SetupBuiltInCategories();
+	UE_API void SetupBuiltInCategories();
 
 	// Append the init options stored in this class to the given Outliner init options
-	void GetOutlinerFilters(FSceneOutlinerFilterBarOptions& OutFilterBarOptions);
+	UE_API void GetOutlinerFilters(FSceneOutlinerFilterBarOptions& OutFilterBarOptions);
 
 	// Get the FFilterCategory attached to the given category name. Use this to add filters to the built in categories
-	TSharedPtr<FFilterCategory> GetFilterCategory(const FName& CategoryName);
+	UE_API TSharedPtr<FFilterCategory> GetFilterCategory(const FName& CategoryName);
 
 private:
 
-	void CreateSCCFilters();
+	UE_API void CreateSCCFilters();
 
-	bool DoesActorPassUnsavedFilter(const ISceneOutlinerTreeItem& InItem);
-	bool DoesActorPassUncontrolledFilter(const ISceneOutlinerTreeItem& InItem);
+	UE_API bool DoesActorPassUnsavedFilter(const ISceneOutlinerTreeItem& InItem);
+	UE_API bool DoesActorPassUncontrolledFilter(const ISceneOutlinerTreeItem& InItem);
 
-	void OnUnsavedAssetAdded(const FString& InAsset);
-	void OnUnsavedAssetRemoved(const FString& InAsset);
+	UE_API void OnUnsavedAssetAdded(const FString& InAsset);
+	UE_API void OnUnsavedAssetRemoved(const FString& InAsset);
 
-	void OnUncontrolledChangelistModuleChanged();
+	UE_API void OnUncontrolledChangelistModuleChanged();
 	
 	// Refresh any Outliners that have the given filter active
-	void RefreshOutlinersWithActiveFilter(bool bFullRefresh, const FString& InFilterName);
+	UE_API void RefreshOutlinersWithActiveFilter(bool bFullRefresh, const FString& InFilterName);
 private:
 
 	/** These are the custom filters that the Scene Outliner will have. All active filters will be AND'd together to test
@@ -125,6 +127,8 @@ private:
 	// List of currently uncontrolled packages. Using a set here because it can grow a lot and impact performances when looking up in the list.
 	TSet<FName> UncontrolledPackages;
 
-	static const FString UnsavedAssetsFilterName;
-	static const FString UncontrolledAssetsFilterName;
+	static UE_API const FString UnsavedAssetsFilterName;
+	static UE_API const FString UncontrolledAssetsFilterName;
 };
+
+#undef UE_API
diff --git a/Engine/Source/Editor/LevelEditor/Public/LevelEditorSubsystem.h b/Engine/Source/Editor/LevelEditor/Public/LevelEditorSubsystem.h
--- a/Engine/Source/Editor/LevelEditor/Public/LevelEditorSubsystem.h
+++ b/Engine/Source/Editor/LevelEditor/Public/LevelEditorSubsystem.h
@@ -11,6 +11,8 @@
 
 #include "LevelEditorSubsystem.generated.h"
 
+#define UE_API LEVELEDITOR_API
+
 class UTypedElementSelectionSet;
 struct FToolMenuContext;
 class FEditorModeTools;
@@ -32,66 +34,66 @@ DECLARE_DYNAMIC_MULTICAST_DELEGATE_TwoParams(FOnLevelEditorMapOpened, const FStr
 * ULevelEditorSubsystem
 * Subsystem for exposing Level Editor related functionality to scripts
 */
-UCLASS()
-class LEVELEDITOR_API ULevelEditorSubsystem : public UEditorSubsystem, public IActorEditorContextClient
+UCLASS(MinimalAPI)
+class ULevelEditorSubsystem : public UEditorSubsystem, public IActorEditorContextClient
 {
 	GENERATED_BODY()
 
 public:
-	virtual void Initialize(FSubsystemCollectionBase& Collection) override;
+	UE_API virtual void Initialize(FSubsystemCollectionBase& Collection) override;
 
-	virtual void Deinitialize() override;
-	void ExtendQuickActionMenu();
+	UE_API virtual void Deinitialize() override;
+	UE_API void ExtendQuickActionMenu();
 
 	UFUNCTION(BlueprintCallable, Category = "Editor Scripting | Level Utility", meta=(DevelopmentOnly))
-	void PilotLevelActor(AActor* ActorToPilot, FName ViewportConfigKey = NAME_None);
-	void PilotLevelActor(const FToolMenuContext& InContext);
+	UE_API void PilotLevelActor(AActor* ActorToPilot, FName ViewportConfigKey = NAME_None);
+	UE_API void PilotLevelActor(const FToolMenuContext& InContext);
 	UFUNCTION(BlueprintCallable, Category = "Editor Scripting | Level Utility", meta = (DevelopmentOnly))
-	void EjectPilotLevelActor(FName ViewportConfigKey = NAME_None);
+	UE_API void EjectPilotLevelActor(FName ViewportConfigKey = NAME_None);
 
 	UFUNCTION(BlueprintCallable, Category = "Editor Scripting | Level Utility", meta=(DevelopmentOnly))
-	AActor* GetPilotLevelActor(FName ViewportConfigKey = NAME_None);
+	UE_API AActor* GetPilotLevelActor(FName ViewportConfigKey = NAME_None);
 
 	UFUNCTION(BlueprintCallable, Category = "Editor Scripting | Level Utility", meta = (DevelopmentOnly))
-	bool GetExactCameraView(FName ViewportConfigKey = NAME_None);
+	UE_API bool GetExactCameraView(FName ViewportConfigKey = NAME_None);
 	UFUNCTION(BlueprintCallable, Category = "Editor Scripting | Level Utility", meta = (DevelopmentOnly))
-	void SetExactCameraView(bool bExactCameraView, FName ViewportConfigKey = NAME_None);
+	UE_API void SetExactCameraView(bool bExactCameraView, FName ViewportConfigKey = NAME_None);
 
 	UFUNCTION(BlueprintCallable, Category = "Editor Scripting | Level Utility", meta = (DevelopmentOnly))
-	void EditorPlaySimulate();
+	UE_API void EditorPlaySimulate();
 
 	UFUNCTION(BlueprintCallable, Category = "Editor Scripting | Level Utility", meta = (DevelopmentOnly))
-	void EditorInvalidateViewports();
+	UE_API void EditorInvalidateViewports();
 
 	UFUNCTION(BlueprintCallable, Category = "Editor Scripting | Level Utility", meta = (DevelopmentOnly))
-	void EditorSetViewportRealtime(bool bInRealtime, FName ViewportConfigKey = NAME_None);
+	UE_API void EditorSetViewportRealtime(bool bInRealtime, FName ViewportConfigKey = NAME_None);
 
 	UFUNCTION(BlueprintCallable, Category = "Editor Scripting | Level Utility", meta = (DevelopmentOnly))
-	void EditorSetGameView(bool bGameView, FName ViewportConfigKey = NAME_None);
+	UE_API void EditorSetGameView(bool bGameView, FName ViewportConfigKey = NAME_None);
 
 	UFUNCTION(BlueprintCallable, Category = "Editor Scripting | Level Utility", meta = (DevelopmentOnly))
-	bool EditorGetGameView(FName ViewportConfigKey = NAME_None);
+	UE_API bool EditorGetGameView(FName ViewportConfigKey = NAME_None);
 
 	UFUNCTION(BlueprintCallable, Category = "Editor Scripting | Level Utility", meta = (DevelopmentOnly))
-	void EditorRequestBeginPlay();
+	UE_API void EditorRequestBeginPlay();
 
 	UFUNCTION(BlueprintCallable, Category = "Editor Scripting | Level Utility", meta = (DevelopmentOnly))
-	void EditorRequestEndPlay();
+	UE_API void EditorRequestEndPlay();
 
 	UFUNCTION(BlueprintCallable, Category = "Editor Scripting | Level Utility", meta = (DevelopmentOnly))
-	bool IsInPlayInEditor() const;
+	UE_API bool IsInPlayInEditor() const;
 
 	UFUNCTION(BlueprintCallable, Category = "Editor Scripting | Level Utility", meta = (DevelopmentOnly))
-	TArray<FName> GetViewportConfigKeys();
+	UE_API TArray<FName> GetViewportConfigKeys();
 
 	UFUNCTION(BlueprintCallable, Category = "Editor Scripting | Level Utility", meta = (DevelopmentOnly))
-	FName GetActiveViewportConfigKey();
+	UE_API FName GetActiveViewportConfigKey();
 
 	UFUNCTION(BlueprintCallable, Category = "Editor Scripting | Level Utility", meta = (DevelopmentOnly))
-	void SetAllowsCinematicControl(bool bAllow, FName ViewportConfigKey = NAME_None);
+	UE_API void SetAllowsCinematicControl(bool bAllow, FName ViewportConfigKey = NAME_None);
 
 	UFUNCTION(BlueprintCallable, Category = "Editor Scripting | Level Utility", meta = (DevelopmentOnly))
-	bool GetAllowsCinematicControl(FName ViewportConfigKey = NAME_None);
+	UE_API bool GetAllowsCinematicControl(FName ViewportConfigKey = NAME_None);
 
 	/**
 	 * Close the current Persistent Level (without saving it). Create a new blank Level and save it. Load the new created level.
@@ -101,7 +103,7 @@ public:
 	 * @return	True if the operation succeeds.
 	 */
 	UFUNCTION(BlueprintCallable, Category = "Editor Scripting | Level Utility")
-	bool NewLevel(const FString& AssetPath, bool bIsPartitionedWorld = false);
+	UE_API bool NewLevel(const FString& AssetPath, bool bIsPartitionedWorld = false);
 
 	/**
 	 * Close the current Persistent Level (without saving it). Create a new Level base on another level and save it. Load the new created level.
@@ -112,7 +114,7 @@ public:
 	 * @return	True if the operation succeeds.
 	 */
 	UFUNCTION(BlueprintCallable, Category = "Editor Scripting | Level Utility")
-	bool NewLevelFromTemplate(const FString& AssetPath, const FString& TemplateAssetPath);
+	UE_API bool NewLevelFromTemplate(const FString& AssetPath, const FString& TemplateAssetPath);
 
 	/**
 	 * Close the current Persistent Level (without saving it). Loads the specified level.
@@ -121,21 +123,21 @@ public:
 	 * @return	True if the operation succeeds.
 	 */
 	UFUNCTION(BlueprintCallable, Category = "Editor Scripting | Level Utility")
-	bool LoadLevel(const FString& AssetPath);
+	UE_API bool LoadLevel(const FString& AssetPath);
 
 	/**
 	 * Saves the specified Level. Must already be saved at lease once to have a valid path.
 	 * @return	True if the operation succeeds.
 	 */
 	UFUNCTION(BlueprintCallable, Category = "Editor Scripting | Level Utility")
-	bool SaveCurrentLevel();
+	UE_API bool SaveCurrentLevel();
 
 	/**
 	 * Saves all Level currently loaded by the World Editor.
 	 * @return	True if the operation succeeds.
 	 */
 	UFUNCTION(BlueprintCallable, Category = "Editor Scripting | Level Utility")
-	bool SaveAllDirtyLevels();
+	UE_API bool SaveAllDirtyLevels();
 
 	/**
 	 * Set the current level used by the world editor.
@@ -144,21 +146,21 @@ public:
 	 * @return	True if the operation succeeds.
 	 */
 	UFUNCTION(BlueprintCallable, Category = "Editor Scripting | Level Utility")
-	bool SetCurrentLevelByName(FName LevelName);
+	UE_API bool SetCurrentLevelByName(FName LevelName);
 
 	/**
 	* Get the current level used by the world editor.
 	* @return	The current level
 	*/
 	UFUNCTION(BlueprintCallable, Category = "Editor Scripting | Level Utility")
-	ULevel* GetCurrentLevel();
+	UE_API ULevel* GetCurrentLevel();
 
 	/**
 	 * Get the selection set for the current world, you can use this to track
 	 * and create changes to the level editor's selection
 	 */
 	UFUNCTION(BlueprintCallable, Category = "Editor Scripting | Level Utility")
-	UTypedElementSelectionSet* GetSelectionSet();
+	UE_API UTypedElementSelectionSet* GetSelectionSet();
 
 	/**
 	 * Build Light Maps and optionally the reflection captures.
@@ -167,21 +169,21 @@ public:
 	 * @return	True if build was successful.
 	 */
 	UFUNCTION(BlueprintCallable, Category = "Editor Scripting | Level Utility")
-	bool BuildLightMaps(ELightingBuildQuality Quality = ELightingBuildQuality::Quality_Production, bool bWithReflectionCaptures = false);
+	UE_API bool BuildLightMaps(ELightingBuildQuality Quality = ELightingBuildQuality::Quality_Production, bool bWithReflectionCaptures = false);
 
 	/**
 	 * Gets the global level editor mode manager, if we have one.
 	 * The mode manager is not created in commandlet environments, because modes inherently imply user interactions.
 	 */
-	FEditorModeTools* GetLevelEditorModeManager();
+	UE_API FEditorModeTools* GetLevelEditorModeManager();
 
 	//~ Begin IActorEditorContextClient interface
 	virtual void OnExecuteActorEditorContextAction(UWorld* InWorld, const EActorEditorContextAction& InType, class AActor* InActor = nullptr) override {}
 	virtual void CaptureActorEditorContextState(UWorld* InWorld, UActorEditorContextStateCollection* InStateCollection) const override {}
 	virtual void RestoreActorEditorContextState(UWorld* InWorld, const UActorEditorContextStateCollection* InStateCollection) override {}
-	virtual bool GetActorEditorContextDisplayInfo(UWorld* InWorld, FActorEditorContextClientDisplayInfo& OutDiplayInfo) const override;
+	UE_API virtual bool GetActorEditorContextDisplayInfo(UWorld* InWorld, FActorEditorContextClientDisplayInfo& OutDiplayInfo) const override;
 	virtual bool CanResetContext(UWorld* InWorld) const override { return false; };
-	virtual TSharedRef<SWidget> GetActorEditorContextWidget(UWorld* InWorld) const override;
+	UE_API virtual TSharedRef<SWidget> GetActorEditorContextWidget(UWorld* InWorld) const override;
 	virtual FOnActorEditorContextClientChanged& GetOnActorEditorContextClientChanged() override { return ActorEditorContextClientChanged; }
 	//~ End IActorEditorContextClient interface
 
@@ -206,24 +208,26 @@ public:
 private:
 
 	/** Called when a Level is added to a world or remove from a world */
-	void OnLevelAddedOrRemoved(ULevel* InLevel, UWorld* InWorld);
+	UE_API void OnLevelAddedOrRemoved(ULevel* InLevel, UWorld* InWorld);
 
 	/** Called when the current level changes on a world */
-	void OnCurrentLevelChanged(ULevel* InNewLevel, ULevel* InOldLevel, UWorld* InWorld);
+	UE_API void OnCurrentLevelChanged(ULevel* InNewLevel, ULevel* InOldLevel, UWorld* InWorld);
 
 	/** Delegate used to notify changes to ActorEditorContextSubsystem */
 	FOnActorEditorContextClientChanged ActorEditorContextClientChanged;
 
 	/** Called Pre SaveWorld */
-	void HandleOnPreSaveWorldWithContext(class UWorld* World, FObjectPreSaveContext ObjectSaveContext);
+	UE_API void HandleOnPreSaveWorldWithContext(class UWorld* World, FObjectPreSaveContext ObjectSaveContext);
 	/** Called Post SaveWorld */
-	void HandleOnPostSaveWorldWithContext(class UWorld* World, FObjectPostSaveContext ObjectSaveContext);
+	UE_API void HandleOnPostSaveWorldWithContext(class UWorld* World, FObjectPostSaveContext ObjectSaveContext);
 
 	/** Called on Camera Movement*/
-	void HandleOnEditorCameraMoved(const FVector& Location, const FRotator& Rotation, ELevelViewportType ViewportType, int32 ViewIndex);
+	UE_API void HandleOnEditorCameraMoved(const FVector& Location, const FRotator& Rotation, ELevelViewportType ViewportType, int32 ViewIndex);
 
 	/** Called on Map Change */
-	void HandleOnMapChanged(uint32 MapChangeFlags);
+	UE_API void HandleOnMapChanged(uint32 MapChangeFlags);
 	/** Called on Map Open */
-	void HandleOnMapOpened(const FString& Filename, bool bAsTemplate);
-};
\ No newline at end of file
+	UE_API void HandleOnMapOpened(const FString& Filename, bool bAsTemplate);
+};
+
+#undef UE_API
diff --git a/Engine/Source/Editor/LevelEditor/Public/LevelViewportLayout.h b/Engine/Source/Editor/LevelEditor/Public/LevelViewportLayout.h
--- a/Engine/Source/Editor/LevelEditor/Public/LevelViewportLayout.h
+++ b/Engine/Source/Editor/LevelEditor/Public/LevelViewportLayout.h
@@ -14,6 +14,8 @@
 #include "AssetEditorViewportLayout.h"
 #include "TickableEditorObject.h"
 
+#define UE_API LEVELEDITOR_API
+
 class FLevelEditorViewportClient;
 class FLevelViewportLayout;
 class FLevelViewportTabContent;
@@ -46,18 +48,18 @@ public:
  * Base class for level viewport layout configurations
  * Handles maximizing and restoring well as visibility of specific viewports.
  */
-class LEVELEDITOR_API FLevelViewportLayout : public FAssetEditorViewportLayout
+class FLevelViewportLayout : public FAssetEditorViewportLayout
 {
 public:
 	/**
 	 * Constructor
 	 */
-	FLevelViewportLayout();
+	UE_API FLevelViewportLayout();
 
 	/**
 	 * Destructor
 	 */
-	virtual ~FLevelViewportLayout();
+	UE_API virtual ~FLevelViewportLayout();
 
 	/**
 	 * Builds a viewport layout and returns the widget containing the layout
@@ -67,9 +69,9 @@ public:
 	 * @param LayoutString			The layout string loaded from file to custom build the layout with
 	 * @param InParentLevelEditor	Optional level editor parent to use for new viewports
 	 */
-	virtual TSharedRef<SWidget> BuildViewportLayout(TSharedPtr<SDockTab> InParentDockTab, TSharedPtr<FEditorViewportTabContent> InParentTab, const FString& LayoutString) override;
+	UE_API virtual TSharedRef<SWidget> BuildViewportLayout(TSharedPtr<SDockTab> InParentDockTab, TSharedPtr<FEditorViewportTabContent> InParentTab, const FString& LayoutString) override;
 
-	virtual TSharedRef<SWidget> FactoryViewport(FName InTypeName, const FAssetEditorViewportConstructionArgs& ConstructionArgs) override;
+	UE_API virtual TSharedRef<SWidget> FactoryViewport(FName InTypeName, const FAssetEditorViewportConstructionArgs& ConstructionArgs) override;
 
 	/**
 	 * Makes a request to maximize a specific viewport and hide the others in this layout
@@ -79,12 +81,12 @@ public:
 	 * @param	bWantImmersive		True to perform an "immersive" maximize, which transitions the viewport to fill the entire application window
 	 * @param	bAllowAnimation		True if an animated transition should be used
 	 */
-	void RequestMaximizeViewport( FName ViewportToMaximize, const bool bWantMaximize, const bool bWantImmersive, const bool bAllowAnimation = true );
+	UE_API void RequestMaximizeViewport( FName ViewportToMaximize, const bool bWantMaximize, const bool bWantImmersive, const bool bAllowAnimation = true );
 
 	/**
 	 * @return true if this layout is visible.  It is not visible if its parent tab is not active
 	 */
-	virtual bool IsVisible() const override;
+	UE_API virtual bool IsVisible() const override;
 
 	/**
 	 * Checks to see the specified level viewport is visible in this layout
@@ -93,7 +95,7 @@ public:
 	 * @param InViewport	The viewport within this layout that should be checked
 	 * @return true if the viewport is visible.  
 	 */
-	bool IsLevelViewportVisible( FName InViewport ) const;
+	UE_API bool IsLevelViewportVisible( FName InViewport ) const;
 
 	/**
 	* Checks to see if the specified level viewport supports maximizing one pane
@@ -113,7 +115,7 @@ public:
 	 * @param InViewport	The viewport to check
 	 * @return true if the viewport is maximized, false otherwise
 	 */
-	bool IsViewportMaximized( FName InViewport ) const;
+	UE_API bool IsViewportMaximized( FName InViewport ) const;
 
 	/** 
 	 * Checks to see if the specified level viewport is in immersive mode
@@ -121,14 +123,14 @@ public:
 	 * @param InViewport	The viewport to check
 	 * @return true if the viewport is immersive, false otherwise
 	 */
-	bool IsViewportImmersive( FName InViewport ) const;
+	UE_API bool IsViewportImmersive( FName InViewport ) const;
 
 
 
 
 	/** FTickableEditorObject interface */
-	virtual void Tick( float DeltaTime ) override;
-	virtual bool IsTickable() const override;
+	UE_API virtual void Tick( float DeltaTime ) override;
+	UE_API virtual bool IsTickable() const override;
 	virtual TStatId GetStatId() const override { return TStatId(); }
 
 
@@ -138,8 +140,8 @@ public:
 	/** Returns whether a viewport animation is currently taking place */
 	bool IsTransitioning() const { return bIsTransitioning; }
 
-	virtual void LoadConfig(const FString& LayoutString) override;
-	virtual void SaveConfig(const FString& LayoutString) const override;
+	UE_API virtual void LoadConfig(const FString& LayoutString) override;
+	UE_API virtual void SaveConfig(const FString& LayoutString) const override;
 
 protected:
 
@@ -151,7 +153,7 @@ protected:
 	 * @param	bWantImmersive		True to perform an "immersive" maximize, which transitions the viewport to fill the entire application window
 	 * @param	bAllowAnimation		True if an animated transition should be used
 	 */
-	void MaximizeViewport( FName ViewportToMaximize, const bool bWantMaximize, const bool bWantImmersive, const bool bAllowAnimation );
+	UE_API void MaximizeViewport( FName ViewportToMaximize, const bool bWantMaximize, const bool bWantImmersive, const bool bAllowAnimation );
 
 	/**
 	 * Delegate called to get the visibility of the non-maximized viewports
@@ -159,41 +161,41 @@ protected:
 	 *
 	 * @param EVisibility::Visible when visible, EVisibility::Collapsed otherwise
 	 */
-	virtual EVisibility OnGetNonMaximizedVisibility() const override;
+	UE_API virtual EVisibility OnGetNonMaximizedVisibility() const override;
 
 	/**
 	 * Returns the widget position for viewport transition animations
 	 *
 	 * @return	Viewport position on canvas
 	 */
-	FVector2D GetMaximizedViewportPositionOnCanvas() const;
+	UE_API FVector2D GetMaximizedViewportPositionOnCanvas() const;
 
 	/**
 	 * Returns the widget size for viewport transition animations
 	 *
 	 * @return	Viewport size on canvas
 	 */
-	FVector2D GetMaximizedViewportSizeOnCanvas() const;
+	UE_API FVector2D GetMaximizedViewportSizeOnCanvas() const;
 
 	/** If a viewport animation is in progress, finishes that transition immediately */
-	void FinishMaximizeTransition();
+	UE_API void FinishMaximizeTransition();
 
 	/**
 	 * Begins a draw throttle for responsiveness when animating a viewports size
 	 */
-	void BeginThrottleForAnimatedResize();
+	UE_API void BeginThrottleForAnimatedResize();
 
 	/**
 	 * Ends a draw throttle for responsiveness when animating a viewports size
 	 */
-	void EndThrottleForAnimatedResize();
+	UE_API void EndThrottleForAnimatedResize();
 
 protected:
 	/** Called in after factory of a new layout functions for derived types of layout to init values common to all types */
-	void InitCommonLayoutFromString(const FString& LayoutString, const FName PerspectiveViewportKey);
+	UE_API void InitCommonLayoutFromString(const FString& LayoutString, const FName PerspectiveViewportKey);
 
 	/** Called in SaveConfig() functions for derived types of layout to save values common to all types */
-	void SaveCommonLayoutString(const FString& SpecificLayoutString) const;
+	UE_API void SaveCommonLayoutString(const FString& SpecificLayoutString) const;
 
 	/** True if we've started an animation and are waiting for it to finish */
 	bool bIsTransitioning;
@@ -205,7 +207,7 @@ protected:
 	TWeakPtr< ILevelEditor > ParentLevelEditor;
 
 	/** Method for taking high res screen shots of viewports */
-	void TakeHighResScreenShot();
+	UE_API void TakeHighResScreenShot();
 
 	/** The current maximized viewport if any */
 	FName MaximizedViewport;
@@ -272,3 +274,5 @@ protected:
 	};
 	TArray<FMaximizeViewportCommand> DeferredMaximizeCommands;
 };
+
+#undef UE_API
diff --git a/Engine/Source/Editor/LevelEditor/Public/LevelViewportTabContent.h b/Engine/Source/Editor/LevelEditor/Public/LevelViewportTabContent.h
--- a/Engine/Source/Editor/LevelEditor/Public/LevelViewportTabContent.h
+++ b/Engine/Source/Editor/LevelEditor/Public/LevelViewportTabContent.h
@@ -6,6 +6,8 @@
 #include "LevelViewportLayout.h"
 #include "EditorViewportTabContent.h"
 
+#define UE_API LEVELEDITOR_API
+
 class ILevelEditor;
 class FEditorViewportLayout;
 
@@ -13,28 +15,30 @@ class FEditorViewportLayout;
  * Represents the content in a viewport tab in the level editor.
  * Each SDockTab holding viewports in the level editor contains and owns one of these.
  */
-class LEVELEDITOR_API FLevelViewportTabContent : public FEditorViewportTabContent
+class FLevelViewportTabContent : public FEditorViewportTabContent
 {
 public:
-	~FLevelViewportTabContent();
+	UE_API ~FLevelViewportTabContent();
 
 	/** Starts the tab content object and creates the initial layout based on the layout string */
-	virtual void Initialize(AssetEditorViewportFactoryFunction Func, TSharedPtr<SDockTab> InParentTab, const FString& InLayoutString) override;
+	UE_API virtual void Initialize(AssetEditorViewportFactoryFunction Func, TSharedPtr<SDockTab> InParentTab, const FString& InLayoutString) override;
 
-	virtual void BindViewportLayoutCommands(FUICommandList& InOutCommandList, FName ViewportConfigKey) override;
+	UE_API virtual void BindViewportLayoutCommands(FUICommandList& InOutCommandList, FName ViewportConfigKey) override;
 
 protected:
-	virtual TSharedPtr<FEditorViewportLayout> FactoryViewportLayout(bool bIsSwitchingLayouts) override;
-	virtual FName GetLayoutTypeNameFromLayoutString() const override;
+	UE_API virtual TSharedPtr<FEditorViewportLayout> FactoryViewportLayout(bool bIsSwitchingLayouts) override;
+	UE_API virtual FName GetLayoutTypeNameFromLayoutString() const override;
 
-	void OnLayoutStartChange(bool bSwitchingLayouts);
-	void OnLayoutChanged();
+	UE_API void OnLayoutStartChange(bool bSwitchingLayouts);
+	UE_API void OnLayoutChanged();
 
 private:
-	void OnUIActionSetViewportConfiguration(FName InConfigurationName);
-	FName GetViewportTypeWithinLayout(FName InConfigKey) const;
-	void OnUIActionSetViewportTypeWithinLayout(FName InConfigKey, FName InLayoutType);
-	bool IsViewportTypeWithinLayoutEqual(FName InConfigName, FName InLayoutType) const;
-	bool IsViewportConfigurationChecked(FName InLayoutType) const;
-	bool IsLayoutMaximized() const;
+	UE_API void OnUIActionSetViewportConfiguration(FName InConfigurationName);
+	UE_API FName GetViewportTypeWithinLayout(FName InConfigKey) const;
+	UE_API void OnUIActionSetViewportTypeWithinLayout(FName InConfigKey, FName InLayoutType);
+	UE_API bool IsViewportTypeWithinLayoutEqual(FName InConfigName, FName InLayoutType) const;
+	UE_API bool IsViewportConfigurationChecked(FName InLayoutType) const;
+	UE_API bool IsLayoutMaximized() const;
 };
+
+#undef UE_API
diff --git a/Engine/Source/Editor/LevelEditor/Public/SLevelViewport.h b/Engine/Source/Editor/LevelEditor/Public/SLevelViewport.h
--- a/Engine/Source/Editor/LevelEditor/Public/SLevelViewport.h
+++ b/Engine/Source/Editor/LevelEditor/Public/SLevelViewport.h
@@ -21,6 +21,8 @@
 #include "IAssetViewport.h"
 #include "LevelEditorViewport.h"
 
+#define UE_API LEVELEDITOR_API
+
 class SLevelViewportToolBar;
 class FLevelEditorViewportClient;
 class FLevelViewportLayout;
@@ -38,7 +40,7 @@ enum ELabelAnchorMode : int;
 /**
  * Encapsulates an SViewport and an SLevelViewportToolBar
  */
-class LEVELEDITOR_API SLevelViewport : public SAssetEditorViewport, public IAssetViewport
+class SLevelViewport : public SAssetEditorViewport, public IAssetViewport
 {
 public:
 	SLATE_BEGIN_ARGS( SLevelViewport )
@@ -48,32 +50,32 @@ public:
 		SLATE_ARGUMENT( TSharedPtr<FLevelEditorViewportClient>, LevelEditorViewportClient )
 	SLATE_END_ARGS()
 
-	SLevelViewport();
-	~SLevelViewport();
+	UE_API SLevelViewport();
+	UE_API ~SLevelViewport();
 
 
 	/**
 	 * Constructs the viewport widget                   
 	 */
-	void Construct(const FArguments& InArgs, const FAssetEditorViewportConstructionArgs& InConstructionArgs);
+	UE_API void Construct(const FArguments& InArgs, const FAssetEditorViewportConstructionArgs& InConstructionArgs);
 
 	/**
 	 * Constructs the widgets for the viewport overlay
 	 */
-	void ConstructViewportOverlayContent();
+	UE_API void ConstructViewportOverlayContent();
 
 	/**
 	 * Constructs the level editor viewport client
 	 */
-	void ConstructLevelEditorViewportClient(FLevelEditorViewportInstanceSettings& ViewportInstanceSettings);
+	UE_API void ConstructLevelEditorViewportClient(FLevelEditorViewportInstanceSettings& ViewportInstanceSettings);
 
 	/**
 	 * @return true if the viewport is visible. false otherwise                  
 	 */
-	virtual bool IsVisible() const override;
+	UE_API virtual bool IsVisible() const override;
 
 	/** @return true if this viewport is in a foregrounded tab */
-	bool IsInForegroundTab() const;
+	UE_API bool IsInForegroundTab() const;
 
 	/**
 	 * @return The editor client for this viewport
@@ -96,7 +98,7 @@ public:
 	/**
 	 * Sets Slate keyboard focus to this viewport
 	 */
-	void SetKeyboardFocusToThisViewport();
+	UE_API void SetKeyboardFocusToThisViewport();
 
 	/**
 	 * @return The list of commands on the viewport that are bound to delegates                    
@@ -104,60 +106,60 @@ public:
 	const TSharedPtr<FUICommandList>& GetCommandList() const { return CommandList; }
 
 	/** Saves this viewport's config to ULevelEditorViewportSettings */
-	void SaveConfig(const FString& ConfigName) const;
+	UE_API void SaveConfig(const FString& ConfigName) const;
 
 	/** IAssetViewport Interface */
-	virtual void StartPlayInEditorSession( UGameViewportClient* PlayClient, const bool bInSimulateInEditor ) override;
-	virtual void EndPlayInEditorSession() override;
-	virtual void SwapViewportsForSimulateInEditor() override;
-	virtual void SwapViewportsForPlayInEditor() override;
-	virtual void OnSimulateSessionStarted() override;
-	virtual void OnSimulateSessionFinished() override;
-	virtual void RegisterGameViewportIfPIE() override;
-	virtual bool HasPlayInEditorViewport() const override; 
-	virtual FViewport* GetActiveViewport() override;
+	UE_API virtual void StartPlayInEditorSession( UGameViewportClient* PlayClient, const bool bInSimulateInEditor ) override;
+	UE_API virtual void EndPlayInEditorSession() override;
+	UE_API virtual void SwapViewportsForSimulateInEditor() override;
+	UE_API virtual void SwapViewportsForPlayInEditor() override;
+	UE_API virtual void OnSimulateSessionStarted() override;
+	UE_API virtual void OnSimulateSessionFinished() override;
+	UE_API virtual void RegisterGameViewportIfPIE() override;
+	UE_API virtual bool HasPlayInEditorViewport() const override; 
+	UE_API virtual FViewport* GetActiveViewport() override;
 	TSharedPtr<FSceneViewport> GetSharedActiveViewport() const override {return ActiveViewport;};
 	virtual TSharedRef< const SWidget> AsWidget() const override { return AsShared(); }
 	virtual TSharedRef< SWidget> AsWidget() override { return AsShared(); }
 	virtual TWeakPtr< SViewport > GetViewportWidget() override { return ViewportWidget; }
-	virtual void AddOverlayWidget( TSharedRef<SWidget> OverlaidWidget, int32 ZOrder=INDEX_NONE ) override;
-	virtual void RemoveOverlayWidget( TSharedRef<SWidget> OverlaidWidget ) override;
+	UE_API virtual void AddOverlayWidget( TSharedRef<SWidget> OverlaidWidget, int32 ZOrder=INDEX_NONE ) override;
+	UE_API virtual void RemoveOverlayWidget( TSharedRef<SWidget> OverlaidWidget ) override;
 
 
 
 	/** SEditorViewport Interface */
-	virtual void OnFocusViewportToSelection() override;
-	virtual EVisibility GetTransformToolbarVisibility() const override;
-	virtual UWorld* GetWorld() const override;
-	virtual FReply OnFocusReceived(const FGeometry& MyGeometry, const FFocusEvent& InFocusEvent) override;
-	virtual void ToggleInViewportContextMenu() override;
+	UE_API virtual void OnFocusViewportToSelection() override;
+	UE_API virtual EVisibility GetTransformToolbarVisibility() const override;
+	UE_API virtual UWorld* GetWorld() const override;
+	UE_API virtual FReply OnFocusReceived(const FGeometry& MyGeometry, const FFocusEvent& InFocusEvent) override;
+	UE_API virtual void ToggleInViewportContextMenu() override;
 
-	virtual void HideInViewportContextMenu() override;
-	virtual bool CanToggleInViewportContextMenu() override;
+	UE_API virtual void HideInViewportContextMenu() override;
+	UE_API virtual bool CanToggleInViewportContextMenu() override;
 
-	static void EnableInViewportMenu();
-	FMargin GetContextMenuPadding() const;
+	static UE_API void EnableInViewportMenu();
+	UE_API FMargin GetContextMenuPadding() const;
 	/**
 	 * Called when the maximize command is executed                   
 	 */
-	FReply OnToggleMaximize();
+	UE_API FReply OnToggleMaximize();
 
 	/**
 	 * @return true if this viewport is maximized, false otherwise
 	 */
-	bool IsMaximized() const;
+	UE_API bool IsMaximized() const;
 
 	/**
 	 * @return true if this viewport can be maximized, false otherwise
 	 */
-	bool CanMaximize() const;
+	UE_API bool CanMaximize() const;
 
 	/**
 	 * Attempt to switch this viewport into a maximized or normal mode
 	 * @param bWantMaximized Whether to witch to maximized or normal
 	 * @param bAllowAnimation True to allow animation when transitioning, otherwise false
 	 */
-	void MakeMaximized(bool bWantMaximized, bool bAllowAnimation = true);
+	UE_API void MakeMaximized(bool bWantMaximized, bool bAllowAnimation = true);
 
 	/**
 	 * Attempts to switch this viewport into immersive mode
@@ -165,39 +167,39 @@ public:
 	 * @param	bWantImmersive Whether to switch to immersive mode, or switch back to normal mode
 	 * @param	bAllowAnimation	True to allow animation when transitioning, otherwise false
 	 */
-	void MakeImmersive( const bool bWantImmersive, const bool bAllowAnimation ) override;
+	UE_API void MakeImmersive( const bool bWantImmersive, const bool bAllowAnimation ) override;
 
 	/**
 	 * @return true if this viewport is in immersive mode, false otherwise
 	 */
-	bool IsImmersive () const override;
+	UE_API bool IsImmersive () const override;
 	
 	/**
 	 * Called to get the visibility of the viewport's 'Restore from Immersive' button. Returns EVisibility::Collapsed when not in immersive mode
 	 */
-	EVisibility GetCloseImmersiveButtonVisibility() const;
+	UE_API EVisibility GetCloseImmersiveButtonVisibility() const;
 		
 	/**
 	 * Called to get the visibility of the viewport's maximize/minimize toggle button. Returns EVisibility::Collapsed when in immersive mode
 	 */
-	EVisibility GetMaximizeToggleVisibility() const;
+	UE_API EVisibility GetMaximizeToggleVisibility() const;
 
 	/**
 	 * @return true if the active viewport is currently being used for play in editor
 	 */
-	bool IsPlayInEditorViewportActive() const;
+	UE_API bool IsPlayInEditorViewportActive() const;
 
 	/**
 	 * @return The active play client, if any 
 	 */
-	UGameViewportClient* GetPlayClient() const;
+	UE_API UGameViewportClient* GetPlayClient() const;
 
 	/**
 	 * Called on all viewports, when actor selection changes.
 	 * 
 	 * @param NewSelection	List of objects that are now selected
 	 */
-	void OnActorSelectionChanged(const TArray<UObject*>& NewSelection, bool bForceRefresh=false);
+	UE_API void OnActorSelectionChanged(const TArray<UObject*>& NewSelection, bool bForceRefresh=false);
 
 	/**
 	 * Called on all viewports, when element selection changes.
@@ -205,140 +207,140 @@ public:
 	 * @param SelectionSet  New selection
 	 * @param bForceRefresh Force refresh
 	 */
-	void OnElementSelectionChanged(const UTypedElementSelectionSet* SelectionSet, bool bForceRefresh);
+	UE_API void OnElementSelectionChanged(const UTypedElementSelectionSet* SelectionSet, bool bForceRefresh);
 
 	/**
 	 * Called when game view should be toggled
 	 */
-	void ToggleGameView() override;
+	UE_API void ToggleGameView() override;
 
 	/**
 	 * @return true if we can toggle game view
 	 */
-	bool CanToggleGameView() const;
+	UE_API bool CanToggleGameView() const;
 
 	/**
 	 * @return true if we are in game view                   
 	 */
-	bool IsInGameView() const override;
+	UE_API bool IsInGameView() const override;
 
 	/**
 	 * Toggles layer visibility in this viewport
 	 *
 	 * @param LayerID					Index of the layer
 	 */
-	void ToggleShowLayer( FName LayerName );
+	UE_API void ToggleShowLayer( FName LayerName );
 
 	/**
 	 * Checks if a layer is visible in this viewport
 	 *
 	 * @param LayerID					Index of the layer
 	 */
-	bool IsLayerVisible( FName LayerName ) const;
+	UE_API bool IsLayerVisible( FName LayerName ) const;
 
 	/**
 	 * Toggles foliage type visibility in this viewport
 	 *
 	 * @param FoliageType	Target foliage type
 	 */
-	void ToggleShowFoliageType(TWeakObjectPtr<class UFoliageType> FoliageType);
+	UE_API void ToggleShowFoliageType(TWeakObjectPtr<class UFoliageType> FoliageType);
 
 	/**
 	 * Toggles all foliage types visibility
 	 *
 	 * @param Visible	true if foliage types should be visible, false otherwise
 	 */
-	void ToggleAllFoliageTypes(bool bVisible);
+	UE_API void ToggleAllFoliageTypes(bool bVisible);
 
 	/**
 	 * Checks if a foliage type is visible in this viewport
 	 *
 	 * @param FoliageType	Target foliage type
 	 */
-	bool IsFoliageTypeVisible(TWeakObjectPtr<class UFoliageType> FoliageType) const;
+	UE_API bool IsFoliageTypeVisible(TWeakObjectPtr<class UFoliageType> FoliageType) const;
 
 
 	/** Called to lock/unlock the actor from the viewport's context menu */
-	void OnActorLockToggleFromMenu(AActor* Actor);
+	UE_API void OnActorLockToggleFromMenu(AActor* Actor);
 
 	/** Called to unlock the actor from the viewport's context menu */
-	void OnActorLockToggleFromMenu();
+	UE_API void OnActorLockToggleFromMenu();
 
 	/**
 	 * @return true if the actor is locked to the viewport
 	 */
-	bool IsActorLocked(const TWeakObjectPtr<AActor> Actor) const;
+	UE_API bool IsActorLocked(const TWeakObjectPtr<AActor> Actor) const;
 
 	/**
 	 * @return true if an actor is locked to the viewport
 	 */
-	bool IsAnyActorLocked() const;
+	UE_API bool IsAnyActorLocked() const;
 
 	/**
 	 * @return true if the viewport is locked to selected actor
 	 */
-	bool IsSelectedActorLocked() const;
+	UE_API bool IsSelectedActorLocked() const;
 
 	/**
 	 * Toggles enabling the exact camera view when locking a viewport to a camera
 	 */
-	void ToggleActorPilotCameraView();
+	UE_API void ToggleActorPilotCameraView();
 
 	/**
 	 * Check whether locked camera view is enabled
 	 */
-	bool IsLockedCameraViewEnabled() const;
+	UE_API bool IsLockedCameraViewEnabled() const;
 
 	/**
 	 * Sets whether the viewport should allow cinematic control
 	 * 
 	 * @param Whether the viewport should allow cinematic control.
      */
-	void SetAllowsCinematicControl(bool bAllow);
+	UE_API void SetAllowsCinematicControl(bool bAllow);
 
 	/**
 	 * @return Whether the viewport allows cinematic control.
      */
-	bool GetAllowsCinematicControl() const;
+	UE_API bool GetAllowsCinematicControl() const;
 
 	/**
 	 * @return the fixed width that a column returned by CreateActorLockSceneOutlinerColumn expects to be
 	 */
-	static float GetActorLockSceneOutlinerColumnWidth();
+	static UE_API float GetActorLockSceneOutlinerColumnWidth();
 
 	/**
 	 * @return a new custom column for a scene outliner that indicates whether each actor is locked to this viewport
 	 */
-	TSharedRef< class ISceneOutlinerColumn > CreateActorLockSceneOutlinerColumn( class ISceneOutliner& SceneOutliner ) const;
+	UE_API TSharedRef< class ISceneOutlinerColumn > CreateActorLockSceneOutlinerColumn( class ISceneOutliner& SceneOutliner ) const;
 
 	/** Called when Preview Selected Cameras preference is changed.*/
-	void OnPreviewSelectedCamerasChange();
+	UE_API void OnPreviewSelectedCamerasChange();
 
 	/**
 	 * Set the device profile name
 	 *
 	 * @param ProfileName The profile name to set
 	 */
-	void SetDeviceProfileString( const FString& ProfileName );
+	UE_API void SetDeviceProfileString( const FString& ProfileName );
 
 	/**
 	 * @return true if the in profile name matches the set profile name
 	 */
-	bool IsDeviceProfileStringSet( FString ProfileName ) const;
+	UE_API bool IsDeviceProfileStringSet( FString ProfileName ) const;
 
 	/**
 	 * @return the name of the selected device profile
 	 */
-	FString GetDeviceProfileString( ) const;
+	UE_API FString GetDeviceProfileString( ) const;
 
 	/** Get the parent level editor for this viewport */
 	TWeakPtr<ILevelEditor> GetParentLevelEditor() const { return ParentLevelEditor; }
 
 	/** @return whether the the actor editor context should be displayed for this viewport */
-	virtual bool IsActorEditorContextVisible() const;
+	UE_API virtual bool IsActorEditorContextVisible() const;
 
 	/** Called to get the screen percentage preview text */
-	FText GetCurrentScreenPercentageText() const;
+	UE_API FText GetCurrentScreenPercentageText() const;
 
 	UE_DEPRECATED(5.1, "GetCurrentLevelTextVisibility not used anymore.")
 	virtual EVisibility GetCurrentLevelTextVisibility() const { return EVisibility::Collapsed; }
@@ -346,21 +348,21 @@ public:
 	virtual EVisibility GetCurrentLevelButtonVisibility() const { return EVisibility::Collapsed; }
 
 	/** @return The visibility of the current level text display */
-	virtual EVisibility GetSelectedActorsCurrentLevelTextVisibility() const;
+	UE_API virtual EVisibility GetSelectedActorsCurrentLevelTextVisibility() const;
 
 	/** Called to get the text for the level the currently selected actor or actors are in. */
-	FText GetSelectedActorsCurrentLevelText(bool bDrawOnlyLabel) const;
+	UE_API FText GetSelectedActorsCurrentLevelText(bool bDrawOnlyLabel) const;
 
 	/** @return The visibility of the current screen percentage text display */
-	EVisibility GetCurrentScreenPercentageVisibility() const;
+	UE_API EVisibility GetCurrentScreenPercentageVisibility() const;
 
 	/** @return The visibility of the viewport controls popup */
-	virtual EVisibility GetViewportControlsVisibility() const;
+	UE_API virtual EVisibility GetViewportControlsVisibility() const;
 
 	/**
 	 * Called to get the visibility of the level viewport toolbar
 	 */
-	virtual EVisibility GetToolBarVisibility() const;
+	UE_API virtual EVisibility GetToolBarVisibility() const;
 
 	/**
 	 * Sets the current layout on the parent tab that this viewport belongs to.
@@ -368,7 +370,7 @@ public:
 	 * @param ConfigurationName		The name of the layout (for the names in namespace LevelViewportConfigurationNames)
 	 */
 	UE_DEPRECATED(5.1, "Moved to internal handling by FViewportTabContent. See FViewportTabContent::BindCommonViewportCommands")
-	void OnSetViewportConfiguration(FName ConfigurationName);
+	UE_API void OnSetViewportConfiguration(FName ConfigurationName);
 
 	/**
 	 * Returns whether the named layout is currently selected on the parent tab that this viewport belongs to.
@@ -377,94 +379,94 @@ public:
 	 * @return						True, if the named layout is currently active
 	 */
 	UE_DEPRECATED(5.1, "Moved to internal handling by FViewportTabContent. See FViewportTabContent::BindCommonViewportCommands")
-	bool IsViewportConfigurationSet(FName ConfigurationName) const;
+	UE_API bool IsViewportConfigurationSet(FName ConfigurationName) const;
 
 	/** Get this level viewport widget's type within its parent layout */
 	UE_DEPRECATED(5.1, "Moved to internal handling by FViewportTabContent. See FViewportTabContent::BindCommonViewportCommands")
-	FName GetViewportTypeWithinLayout() const;
+	UE_API FName GetViewportTypeWithinLayout() const;
 
 	/** Set this level viewport widget's type within its parent layout */
 	UE_DEPRECATED(5.1, "Moved to internal handling by FViewportTabContent. See FViewportTabContent::BindCommonViewportCommands")
-	void SetViewportTypeWithinLayout(FName InLayoutType);
+	UE_API void SetViewportTypeWithinLayout(FName InLayoutType);
 
 	/** Activates the specified viewport type in the layout, if it's not already, or reverts to default if it is. */
 	UE_DEPRECATED(5.1, "Moved to internal handling by FViewportTabContent. See FViewportTabContent::BindCommonViewportCommands")
-	void ToggleViewportTypeActivationWithinLayout(FName InLayoutType);
+	UE_API void ToggleViewportTypeActivationWithinLayout(FName InLayoutType);
 
 	/** Checks if the specified layout type matches our current viewport type. */
 	UE_DEPRECATED(5.1, "Moved to internal handling by FViewportTabContent. See FViewportTabContent::BindCommonViewportCommands")
-	bool IsViewportTypeWithinLayoutEqual(FName InLayoutType);
+	UE_API bool IsViewportTypeWithinLayoutEqual(FName InLayoutType);
 
 	/** For the specified actor, See if we're forcing a preview */
-	bool IsActorAlwaysPreview(TWeakObjectPtr<AActor> Actor) const;
+	UE_API bool IsActorAlwaysPreview(TWeakObjectPtr<AActor> Actor) const;
 
 	/** For the specified actor, toggle Pinned/Unpinned of it's ActorPreview */
-	void SetActorAlwaysPreview(TWeakObjectPtr<AActor> PreviewActor, bool bAlwaysPreview = true);
+	UE_API void SetActorAlwaysPreview(TWeakObjectPtr<AActor> PreviewActor, bool bAlwaysPreview = true);
 
 	/** For the specified actor, toggle Pinned/Unpinned of it's ActorPreview */
-	void ToggleActorPreviewIsPinned(TWeakObjectPtr<AActor> PreviewActor);
+	UE_API void ToggleActorPreviewIsPinned(TWeakObjectPtr<AActor> PreviewActor);
 
 	/** For the specified actor, toggle whether the panel is detached from it*/
-	void ToggleActorPreviewIsPanelDetached(TWeakObjectPtr<AActor> PreviewActor);
+	UE_API void ToggleActorPreviewIsPanelDetached(TWeakObjectPtr<AActor> PreviewActor);
 
 	/** See if the specified actor's ActorPreview is pinned or not */
-	bool IsActorPreviewPinned(TWeakObjectPtr<AActor> PreviewActor);
+	UE_API bool IsActorPreviewPinned(TWeakObjectPtr<AActor> PreviewActor);
 
 	/** See if the specified actor's ActorPreview is detached from actor */
-	bool IsActorPreviewDetached(TWeakObjectPtr<AActor> PreviewActor);
+	UE_API bool IsActorPreviewDetached(TWeakObjectPtr<AActor> PreviewActor);
 
 	/** Actions to perform whenever the viewports floating buttons are pressed */
-	void OnFloatingButtonClicked();
+	UE_API void OnFloatingButtonClicked();
 
 	UE_DEPRECATED(5.6, "Call GetOptionsMenuVisibility() instead.")
-	EVisibility GetToolbarVisibility() const;
+	UE_API EVisibility GetToolbarVisibility() const;
 
 	/** Get the visibility for viewport toolbar's options menu. */
-	EVisibility GetOptionsMenuVisibility() const;
+	UE_API EVisibility GetOptionsMenuVisibility() const;
 
 	/** Get the visibility for items considered to be part of the 'full' viewport toolbar */
 	EVisibility GetFullToolbarVisibility() const { return (bShowToolbarAndControls && bShowEditorToolbar) ? EVisibility::SelfHitTestInvisible : EVisibility::Collapsed; }
 
 	/** Unpin and close all actor preview windows */
-	void RemoveAllPreviews(const bool bRemoveFromDesktopViewport = true);
+	UE_API void RemoveAllPreviews(const bool bRemoveFromDesktopViewport = true);
 
 	/**
 	 * Called to set a bookmark
 	 *
 	 * @param BookmarkIndex	The index of the bookmark to set
 	 */
-	void OnSetBookmark( int32 BookmarkIndex );
+	UE_API void OnSetBookmark( int32 BookmarkIndex );
 
 	/**
 	 * Called to check if a bookmark is set
 	 *
 	 * @param BookmarkIndex	The index of the bookmark to check
 	 */
-	bool OnHasBookmarkSet(int32 BookmarkIndex);
+	UE_API bool OnHasBookmarkSet(int32 BookmarkIndex);
 
 	/**
 	 * Called to jump to a bookmark
 	 *
 	 * @param BookmarkIndex	The index of the bookmark to jump to
 	 */
-	void OnJumpToBookmark( int32 BookmarkIndex );
+	UE_API void OnJumpToBookmark( int32 BookmarkIndex );
 
 	/**
 	 * Called to clear a bookmark
 	 *
 	 * @param BookmarkIndex The index of the bookmark to clear
 	 */
-	void OnClearBookmark( int32 BookmarkIndex );
+	UE_API void OnClearBookmark( int32 BookmarkIndex );
 
 	/**
 	 * Called to clear all bookmarks
 	 */
-	void OnClearAllBookmarks();
+	UE_API void OnClearAllBookmarks();
 
 	/**
 	 * Called to Compact Bookmarks.
 	 */
-	void OnCompactBookmarks();
+	UE_API void OnCompactBookmarks();
 
 	/** 
 	 * Returns the config key associated with this viewport. 
@@ -476,211 +478,211 @@ public:
 	/**
 	 * Toggles the current viewport toolbar visibility. 
 	 */
-	void ToggleViewportToolbarVisibility();
+	UE_API void ToggleViewportToolbarVisibility();
 
 	/**
 	 * Whether the viewport toolbar is currently visible.
 	 */
-	bool IsViewportToolbarVisible() const; 
+	UE_API bool IsViewportToolbarVisible() const; 
 
 protected:
 	/** SEditorViewport interface */
-	virtual TSharedRef<FEditorViewportClient> MakeEditorViewportClient() override;
-	virtual TSharedPtr<SWidget> MakeViewportToolbar() override;
-	virtual TSharedPtr<SWidget> BuildViewportToolbar() override;
-	TSharedRef<SWidget> BuildPIEViewportToolbar();
-
-	virtual void OnIncrementPositionGridSize() override;
-	virtual void OnDecrementPositionGridSize() override;
-	virtual void OnIncrementRotationGridSize() override;
-	virtual void OnDecrementRotationGridSize() override;
-	virtual const FSlateBrush* OnGetViewportBorderBrush() const override;
-	virtual FSlateColor OnGetViewportBorderColorAndOpacity() const override;
-	virtual EVisibility OnGetViewportContentVisibility() const override;
-	virtual EVisibility OnGetFocusedViewportIndicatorVisibility() const override;
-	virtual void BindCommands() override;
+	UE_API virtual TSharedRef<FEditorViewportClient> MakeEditorViewportClient() override;
+	UE_API virtual TSharedPtr<SWidget> MakeViewportToolbar() override;
+	UE_API virtual TSharedPtr<SWidget> BuildViewportToolbar() override;
+	UE_API TSharedRef<SWidget> BuildPIEViewportToolbar();
+
+	UE_API virtual void OnIncrementPositionGridSize() override;
+	UE_API virtual void OnDecrementPositionGridSize() override;
+	UE_API virtual void OnIncrementRotationGridSize() override;
+	UE_API virtual void OnDecrementRotationGridSize() override;
+	UE_API virtual const FSlateBrush* OnGetViewportBorderBrush() const override;
+	UE_API virtual FSlateColor OnGetViewportBorderColorAndOpacity() const override;
+	UE_API virtual EVisibility OnGetViewportContentVisibility() const override;
+	UE_API virtual EVisibility OnGetFocusedViewportIndicatorVisibility() const override;
+	UE_API virtual void BindCommands() override;
 private:
 	/** Flag to know if we need to update the previews which is handled in the tick. */
 	bool bNeedToUpdatePreviews;
 
 	/** Loads this viewport's config from the ini file */
-	FLevelEditorViewportInstanceSettings LoadLegacyConfigFromIni(const FString& ConfigKey, const FLevelEditorViewportInstanceSettings& InDefaultSettings);
+	UE_API FLevelEditorViewportInstanceSettings LoadLegacyConfigFromIni(const FString& ConfigKey, const FLevelEditorViewportInstanceSettings& InDefaultSettings);
 
 	/** Called when a property is changed */
-	void HandleViewportSettingChanged(FName PropertyName);
+	UE_API void HandleViewportSettingChanged(FName PropertyName);
 
 	/**
 	 * Handles which viewport toolbar is active.
 	 */
-	int32 GetViewportToolbarIndex() const;
+	UE_API int32 GetViewportToolbarIndex() const;
 
 	/**
 	 * Called when the advanced settings should be opened.
 	 */
-	void OnAdvancedSettings();
+	UE_API void OnAdvancedSettings();
 
 	/**
 	 * Called when the play settings should be opened.
 	 */
-	void OnPlaySettings();
+	UE_API void OnPlaySettings();
 
 	/**
 	 * Called when immersive mode is toggled by the user
 	 */
-	void OnToggleImmersive();
+	UE_API void OnToggleImmersive();
 
 	/** Called when moving tabs in and out of a sidebar is activated by the user */
-	void OnToggleSidebarTabs();
+	UE_API void OnToggleSidebarTabs();
 
 	/**
 	* Called to determine whether the maximize mode of current viewport can be toggled
 	*/
-	bool CanToggleMaximizeMode() const;
+	UE_API bool CanToggleMaximizeMode() const;
 
 	/**
 	* Called to toggle maximize mode of current viewport
 	*/
-	void OnToggleMaximizeMode();
+	UE_API void OnToggleMaximizeMode();
 
 	/** Starts previewing any selected camera actors using live "PIP" sub-views */
-	void PreviewSelectedCameraActors(const bool bPreviewInDesktopViewport = true);
+	UE_API void PreviewSelectedCameraActors(const bool bPreviewInDesktopViewport = true);
 
 	/**
 	 * Called to create a cameraActor in the currently selected perspective viewport
 	 */
-	void OnCreateCameraActor(UClass *InClass);
+	UE_API void OnCreateCameraActor(UClass *InClass);
 
 	/**
 	 * Called to bring up the screenshot UI
 	 */
-	void OnTakeHighResScreenshot();
+	UE_API void OnTakeHighResScreenshot();
 
 	/**
 	 * Returns whether this level viewport is the active level viewport 
 	 */
-	bool IsActiveLevelViewport() const;
+	UE_API bool IsActiveLevelViewport() const;
 
 	/**
 	 * Called to check currently selected editor viewport is a perspective one
 	 */
-	bool IsPerspectiveViewport() const;
+	UE_API bool IsPerspectiveViewport() const;
 
 	/**
 	 * Toggles all volume classes visibility
 	 *
 	 * @param Visible					true if volumes should be visible, false otherwise
 	 */
-	void OnToggleAllVolumeActors( bool bVisible );
+	UE_API void OnToggleAllVolumeActors( bool bVisible );
 
 	/**
 	 * Toggles volume classes visibility
 	 *
 	 * @param VolumeID					Index of the volume class
 	 */
-	void ToggleShowVolumeClass( int32 VolumeID );
+	UE_API void ToggleShowVolumeClass( int32 VolumeID );
 
 	/**
 	 * Checks if volume class is visible in this viewport
 	 *
 	 * @param VolumeID					Index of the volume class
 	 */
-	bool IsVolumeVisible( int32 VolumeID ) const;
+	UE_API bool IsVolumeVisible( int32 VolumeID ) const;
 
 	/**
 	 * Toggles all layers visibility
 	 *
 	 * @param Visible					true if layers should be visible, false otherwise
 	 */
-	void OnToggleAllLayers( bool bVisible );
+	UE_API void OnToggleAllLayers( bool bVisible );
 
 	/**
 	 * Toggles all sprite categories visibility
 	 *
 	 * @param Visible					true if sprites should be visible, false otherwise
 	 */
-	void OnToggleAllSpriteCategories( bool bVisible );
+	UE_API void OnToggleAllSpriteCategories( bool bVisible );
 
 	/**
 	 * Toggles sprite category visibility in this viewport
 	 *
 	 * @param CategoryID				Index of the category
 	 */
-	void ToggleSpriteCategory( int32 CategoryID );
+	UE_API void ToggleSpriteCategory( int32 CategoryID );
 
 	/**
 	 * Checks if sprite category is visible in this viewport
 	 *
 	 * @param CategoryID				Index of the category
 	 */
-	bool IsSpriteCategoryVisible( int32 CategoryID ) const;
+	UE_API bool IsSpriteCategoryVisible( int32 CategoryID ) const;
 
 	/**
 	 * Toggles all Stat commands visibility
 	 *
 	 * @param Visible					true if Stats should be visible, false otherwise
 	 */
-	void OnToggleAllStatCommands(bool bVisible);
+	UE_API void OnToggleAllStatCommands(bool bVisible);
 
 	/**
 	 * Called when show flags for this viewport should be reset to default, or the saved settings
 	 */
-	void OnUseDefaultShowFlags(bool bUseSavedDefaults = false);
+	UE_API void OnUseDefaultShowFlags(bool bUseSavedDefaults = false);
 
 	/**
 	 * Called to toggle allowing sequencer to use this viewport to preview in
 	 */
-	void OnToggleAllowCinematicPreview();
+	UE_API void OnToggleAllowCinematicPreview();
 
 	/**
 	 * @return true if this viewport allows cinematics to be previewed in it                   
 	 */
-	bool AllowsCinematicPreview() const;
+	UE_API bool AllowsCinematicPreview() const;
 
 	/** Find currently selected actor in the level script.  */
-	void FindSelectedInLevelScript();
+	UE_API void FindSelectedInLevelScript();
 	
 	/** Can we find the currently selected actor in the level script. */
-	bool CanFindSelectedInLevelScript() const;
+	UE_API bool CanFindSelectedInLevelScript() const;
 
 	/** Called to select the currently locked actor */
-	void OnSelectLockedActor();
+	UE_API void OnSelectLockedActor();
 
 	/**
 	 * @return true if the currently locked actor is selectable
 	 */
-	bool CanExecuteSelectLockedActor() const;
+	UE_API bool CanExecuteSelectLockedActor() const;
 
 	/** Called to clear the current actor lock */
-	void OnActorUnlock();
+	UE_API void OnActorUnlock();
 
 	/**
 	 * @return true if clearing the current actor lock is a valid input
 	 */
-	bool CanExecuteActorUnlock() const;
+	UE_API bool CanExecuteActorUnlock() const;
 
 	/** Called to lock the viewport to the currently selected actor */
-	void OnActorLockSelected();
+	UE_API void OnActorLockSelected();
 
 	/**
 	 * @return true if clearing the setting the actor lock to the selected actor is a valid input
 	 */
-	bool CanExecuteActorLockSelected() const;
+	UE_API bool CanExecuteActorLockSelected() const;
 
 	/**
 	 * Called when the viewport should be redrawn
 	 *
 	 * @param bInvalidateHitProxies	Whether or not to invalidate hit proxies
 	 */
-	void RedrawViewport( bool bInvalidateHitProxies );
+	UE_API void RedrawViewport( bool bInvalidateHitProxies );
 
 	/** An internal handler for dagging dropable objects into the viewport. */
-	bool HandleDragObjects(const FGeometry& MyGeometry, const FDragDropEvent& DragDropEvent);
+	UE_API bool HandleDragObjects(const FGeometry& MyGeometry, const FDragDropEvent& DragDropEvent);
 
 	/** An internal handler for dropping objects into the viewport. 
 	 *	@param DragDropEvent		The drag event.
 	 *	@param bCreateDropPreview	If true, a drop preview actor will be spawned instead of a normal actor.
 	 */
-	bool HandlePlaceDraggedObjects(const FGeometry& MyGeometry, const FDragDropEvent& DragDropEvent, bool bCreateDropPreview);
+	UE_API bool HandlePlaceDraggedObjects(const FGeometry& MyGeometry, const FDragDropEvent& DragDropEvent, bool bCreateDropPreview);
 
 	/**
 	 * Tries to get assets from a drag and drop event.
@@ -688,50 +690,50 @@ private:
 	 * @param DragDropEvent		Event to get assets from.
 	 * @param AssetDataArray	Asets will be added here.
 	 */
-	void GetAssetsFromDrag(const FDragDropEvent& DragDropEvent, TArray<FAssetData>& AssetDataArray);
+	UE_API void GetAssetsFromDrag(const FDragDropEvent& DragDropEvent, TArray<FAssetData>& AssetDataArray);
 
 	/** SWidget Interface */
-	virtual FReply OnKeyDown( const FGeometry& MyGeometry, const FKeyEvent& InKeyEvent ) override;
-	virtual void OnDragEnter( const FGeometry& MyGeometry, const FDragDropEvent& DragDropEvent ) override;
-	virtual void OnDragLeave( const FDragDropEvent& DragDropEvent ) override;
-	virtual FReply OnDragOver( const FGeometry& MyGeometry, const FDragDropEvent& DragDropEvent ) override;
-	virtual FReply OnDrop( const FGeometry& MyGeometry, const FDragDropEvent& DragDropEvent ) override;
-	virtual void Tick( const FGeometry& AllottedGeometry, const double InCurrentTime, const float InDeltaTime ) override;
+	UE_API virtual FReply OnKeyDown( const FGeometry& MyGeometry, const FKeyEvent& InKeyEvent ) override;
+	UE_API virtual void OnDragEnter( const FGeometry& MyGeometry, const FDragDropEvent& DragDropEvent ) override;
+	UE_API virtual void OnDragLeave( const FDragDropEvent& DragDropEvent ) override;
+	UE_API virtual FReply OnDragOver( const FGeometry& MyGeometry, const FDragDropEvent& DragDropEvent ) override;
+	UE_API virtual FReply OnDrop( const FGeometry& MyGeometry, const FDragDropEvent& DragDropEvent ) override;
+	UE_API virtual void Tick( const FGeometry& AllottedGeometry, const double InCurrentTime, const float InDeltaTime ) override;
 	/** End of SWidget interface */
 
 	/**
 	 * Bound event Triggered via FLevelViewportCommands::ApplyMaterialToActor, attempts to apply a material selected in the content browser
 	 * to an actor being hovered over in the Editor viewport.
 	 */ 
-	void OnApplyMaterialToViewportTarget();
+	UE_API void OnApplyMaterialToViewportTarget();
 
 	/**
 	 * Binds commands for our toolbar options menu
 	 *
 	 * @param CommandList	The list to bind commands to
 	 */
-	void BindOptionCommands( FUICommandList& CommandList );
+	UE_API void BindOptionCommands( FUICommandList& CommandList );
 
 	/**
 	 * Binds commands for our toolbar view menu
 	 *
 	 * @param CommandList	The list to bind commands to
 	 */
-	void BindViewCommands( FUICommandList& CommandList );
+	UE_API void BindViewCommands( FUICommandList& CommandList );
 
 	/**
 	 * Binds commands for our toolbar show menu
 	 *
 	 * @param CommandList	The list to bind commands to
 	 */
-	void BindShowCommands( FUICommandList& CommandList ) override;
+	UE_API void BindShowCommands( FUICommandList& CommandList ) override;
 
 	/**
 	 * Binds commands for our drag-drop context menu
 	 *
 	 * @param CommandList	The list to bind commands to
 	 */
-	void BindDropCommands( FUICommandList& CommandList );
+	UE_API void BindDropCommands( FUICommandList& CommandList );
 
 	/**
 	* Binds commands for our stat menu, also used as a delegate listener
@@ -739,140 +741,140 @@ private:
 	* @param InMenuItem		The menu item we need to bind
 	* @param InCommandName	The command used by the functions
 	*/
-	void BindStatCommand(const TSharedPtr<FUICommandInfo> InMenuItem, const FString& InCommandName);
+	UE_API void BindStatCommand(const TSharedPtr<FUICommandInfo> InMenuItem, const FString& InCommandName);
 
 	/**
 	 * Binds commands for PIE-specific behavior
 	 * @param OutCommandList	The list to bind commands to
 	 */
-	void BindPIECommands( FUICommandList& OutCommandList );
+	UE_API void BindPIECommands( FUICommandList& OutCommandList );
 
 	/**
 	 * Resets the PIE show flags to default.
 	 */
-	void ResetPIEShowFlags();
+	UE_API void ResetPIEShowFlags();
 	
 	/**
 	 * Toggles a PIE show flag on or off.
 	 */
-	void TogglePIEShowFlag(FEngineShowFlags::EShowFlag Flag);
+	UE_API void TogglePIEShowFlag(FEngineShowFlags::EShowFlag Flag);
 
 	/**
 	 * Returns whether a show flag is on in PIE.
 	 */
-	bool IsPIEShowFlagEnabled(FEngineShowFlags::EShowFlag Flag) const;
+	UE_API bool IsPIEShowFlagEnabled(FEngineShowFlags::EShowFlag Flag) const;
 
 	/**
 	 * Sets the view mode on the PIE viewport.
 	 */
-	void SetPIEViewMode(EViewModeIndex ViewMode);
+	UE_API void SetPIEViewMode(EViewModeIndex ViewMode);
 
 	/**
 	 * Gets whether the view mode is enabled in the PIE Viewport
 	 */
-	bool IsPIEViewModeEnabled(EViewModeIndex ViewMode) const;
+	UE_API bool IsPIEViewModeEnabled(EViewModeIndex ViewMode) const;
 
 	/**
 	 * Refreshes the PIE viewport.
 	 */
-	void RefreshPIEViewport();
+	UE_API void RefreshPIEViewport();
 
 	/**
 	 * A handler for processing input events before gameplay does.
 	 */
-	bool OnPIEViewportInputOverride(FInputKeyEventArgs& Input);
+	UE_API bool OnPIEViewportInputOverride(FInputKeyEventArgs& Input);
 	
 	/**
 	 * Called to build the viewport context menu when the user is Drag Dropping from the Content Browser
 	 */
-	TSharedRef< SWidget > BuildViewportDragDropContextMenu();
+	UE_API TSharedRef< SWidget > BuildViewportDragDropContextMenu();
 
 	/**
 	 * Called when a map is changed (loaded,saved,new map, etc)
 	 */
-	void OnMapChanged( UWorld* World, EMapChangeType MapChangeType );
+	UE_API void OnMapChanged( UWorld* World, EMapChangeType MapChangeType );
 
 	/** Called in response to an actor being deleted in the level */
-	void OnLevelActorsRemoved(AActor* InActor);
+	UE_API void OnLevelActorsRemoved(AActor* InActor);
 	
-	void OnEditorClose();
+	UE_API void OnEditorClose();
 
 	/** Gets whether the locked icon should be shown in the viewport because it is locked to an actor */
-	EVisibility GetLockedIconVisibility() const;
+	UE_API EVisibility GetLockedIconVisibility() const;
 
 	/** Gets the locked icon tooltip text showing the meaning of the icon and the name of the locked actor */
-	FText GetLockedIconToolTip() const;
+	UE_API FText GetLockedIconToolTip() const;
 
 	/**
 	 * Starts previewing the specified actors.  If the supplied list of actors is empty, turns off the preview.
 	 *
 	 * @param	ActorsToPreview		List of actors to draw previews for
 	 */
-	void PreviewActors( const TArray< AActor* >& ActorsToPreview, const bool bPreviewInDesktopViewport = true);
+	UE_API void PreviewActors( const TArray< AActor* >& ActorsToPreview, const bool bPreviewInDesktopViewport = true);
 
 	/** Called every frame to update any actor preview viewports we may have */
-	void UpdateActorPreviewViewports();
+	UE_API void UpdateActorPreviewViewports();
 
 	/**
 	 * Removes a specified actor preview from the list
 	 *
 	 * @param PreviewIndex Array index of the preview to remove.
 	 */
-	void RemoveActorPreview( int32 PreviewIndex, AActor* Actor = nullptr, const bool bRemoveFromDesktopViewport = true);
+	UE_API void RemoveActorPreview( int32 PreviewIndex, AActor* Actor = nullptr, const bool bRemoveFromDesktopViewport = true);
 	
 	/** Returns true if this viewport is the active viewport and can process UI commands */
-	bool CanProduceActionForCommand(const TSharedRef<const FUICommandInfo>& Command) const;
+	UE_API bool CanProduceActionForCommand(const TSharedRef<const FUICommandInfo>& Command) const;
 
 	/** Called when undo is executed */
-	void OnUndo();
+	UE_API void OnUndo();
 
 	/** Called when undo is executed */
-	void OnRedo();
+	UE_API void OnRedo();
 
 	/** @return Whether or not undo can be executed */
-	bool CanExecuteUndo() const;
+	UE_API bool CanExecuteUndo() const;
 
 	/** @return Whether or not redo can be executed */
-	bool CanExecuteRedo() const;
+	UE_API bool CanExecuteRedo() const;
 
 	/** @return Whether the mouse capture label is visible */
-	EVisibility GetMouseCaptureLabelVisibility() const;
+	UE_API EVisibility GetMouseCaptureLabelVisibility() const;
 
 	/** @return The current color & opacity for the mouse capture label */
-	FLinearColor GetMouseCaptureLabelColorAndOpacity() const;
+	UE_API FLinearColor GetMouseCaptureLabelColorAndOpacity() const;
 
 	/** @return The current text for the mouse capture label */
-	FText GetMouseCaptureLabelText() const;
+	UE_API FText GetMouseCaptureLabelText() const;
 
 	/** Show the mouse capture label with the specified vertical and horizontal alignment */
-	void ShowMouseCaptureLabel(ELabelAnchorMode AnchorMode);
+	UE_API void ShowMouseCaptureLabel(ELabelAnchorMode AnchorMode);
 	
 	/** Build the mouse capture label widget */
-	TSharedRef<SWidget> BuildMouseCaptureWidget();
+	UE_API TSharedRef<SWidget> BuildMouseCaptureWidget();
 
 	/** Hide the mouse capture label */
-	void HideMouseCaptureLabel();
+	UE_API void HideMouseCaptureLabel();
 
 	/** Resets view flags when a new level is created or opened */
-	void ResetNewLevelViewFlags();
+	UE_API void ResetNewLevelViewFlags();
 
 	/** Gets the active scene viewport for the game */
-	FSceneViewport* GetGameSceneViewport() const;
+	UE_API FSceneViewport* GetGameSceneViewport() const;
 	
 	/** Handle any level viewport changes on entering PIE or simulate */
-	void TransitionToPIE(bool bIsSimulating);
+	UE_API void TransitionToPIE(bool bIsSimulating);
 
 	/** Handle any level viewport changes on leaving PIE or simulate */
-	void TransitionFromPIE(bool bIsSimulating);
+	UE_API void TransitionFromPIE(bool bIsSimulating);
 
 	/** Get the stretch type of the viewport */
-	EStretch::Type OnGetScaleBoxStretch() const;
+	UE_API EStretch::Type OnGetScaleBoxStretch() const;
 
 	/** Get the SViewport size */
-	FVector2D GetSViewportSize() const;
+	UE_API FVector2D GetSViewportSize() const;
 
 	/** Updates the real-time overrride applied to the viewport */
-	void OnPerformanceSettingsChanged(UObject* Obj, struct FPropertyChangedEvent& ChangeEvent);
+	UE_API void OnPerformanceSettingsChanged(UObject* Obj, struct FPropertyChangedEvent& ChangeEvent);
 
 private:
 	/** Tab which this viewport is located in */
@@ -1060,7 +1062,7 @@ private:
 	bool bIsInViewportMenuShowing;
 	bool bIsInViewportMenuInitialized;
 	TSharedPtr<class SInViewportDetails> InViewportMenu;
-	static bool bInViewportMenuEnabled;
+	static UE_API bool bInViewportMenuEnabled;
 
 	TSharedPtr<SWorldPartitionViewportWidget> WorldPartitionViewportWidget;
 
@@ -1074,13 +1076,15 @@ private:
 	FViewportCameraTransform CachedPerspectiveCameraTransform;
 
 protected:
-	void LockActorInternal(AActor* NewActorToLock);
+	UE_API void LockActorInternal(AActor* NewActorToLock);
 
 	/** Can be overriden by derived classes to add new context objects to the new toolbar */
 	virtual void ExtendToolbarContext(FToolMenuContext& InToolMenuContext) {};
 
 public:
-	static bool GetCameraInformationFromActor(AActor* Actor, FMinimalViewInfo& out_CameraInfo);
+	static UE_API bool GetCameraInformationFromActor(AActor* Actor, FMinimalViewInfo& out_CameraInfo);
 
-	static bool CanGetCameraInformationFromActor(AActor* Actor);
+	static UE_API bool CanGetCameraInformationFromActor(AActor* Actor);
 };
+
+#undef UE_API
diff --git a/Engine/Source/Editor/LocalizationCommandletExecution/Public/LocalizationCommandletExecution.h b/Engine/Source/Editor/LocalizationCommandletExecution/Public/LocalizationCommandletExecution.h
--- a/Engine/Source/Editor/LocalizationCommandletExecution/Public/LocalizationCommandletExecution.h
+++ b/Engine/Source/Editor/LocalizationCommandletExecution/Public/LocalizationCommandletExecution.h
@@ -4,11 +4,13 @@
 #include "CoreMinimal.h"
 #include "HAL/PlatformProcess.h"
 
+#define UE_API LOCALIZATIONCOMMANDLETEXECUTION_API
+
 class SWindow;
 
 namespace LocalizationCommandletExecution
 {
-	struct LOCALIZATIONCOMMANDLETEXECUTION_API FTask
+	struct FTask
 	{
 		FTask() {}
 
@@ -26,10 +28,10 @@ namespace LocalizationCommandletExecution
 	bool Execute(const TSharedRef<SWindow>& ParentWindow, const FText& Title, const TArray<FTask>& Tasks);
 };
 
-class LOCALIZATIONCOMMANDLETEXECUTION_API FLocalizationCommandletProcess : public TSharedFromThis<FLocalizationCommandletProcess>
+class FLocalizationCommandletProcess : public TSharedFromThis<FLocalizationCommandletProcess>
 {
 public:
-	static TSharedPtr<FLocalizationCommandletProcess> Execute(const FString& ConfigFilePath, const bool UseProjectFile = true);
+	static UE_API TSharedPtr<FLocalizationCommandletProcess> Execute(const FString& ConfigFilePath, const bool UseProjectFile = true);
 
 private:
 	FLocalizationCommandletProcess(void* const InReadPipe, void* const InWritePipe, const FProcHandle InProcessHandle, const FString& InProcessArguments)
@@ -41,7 +43,7 @@ private:
 	}
 
 public:
-	~FLocalizationCommandletProcess();
+	UE_API ~FLocalizationCommandletProcess();
 
 	void* GetReadPipe() const
 	{
@@ -64,3 +66,5 @@ private:
 	FProcHandle ProcessHandle;
 	FString ProcessArguments;
 };
+
+#undef UE_API
diff --git a/Engine/Source/Editor/MassEntityDebugger/Private/MassDebuggerStyle.h b/Engine/Source/Editor/MassEntityDebugger/Private/MassDebuggerStyle.h
--- a/Engine/Source/Editor/MassEntityDebugger/Private/MassDebuggerStyle.h
+++ b/Engine/Source/Editor/MassEntityDebugger/Private/MassDebuggerStyle.h
@@ -4,14 +4,16 @@
 
 #include "Styling/SlateStyle.h"
 
+#define UE_API MASSENTITYDEBUGGER_API
 
-class MASSENTITYDEBUGGER_API FMassDebuggerStyle
+
+class FMassDebuggerStyle
 {
 public:
-	static void Initialize();
-	static void Shutdown();
+	static UE_API void Initialize();
+	static UE_API void Shutdown();
 	static ISlateStyle& Get() { return *StyleSet.Get(); }
-	static FName GetStyleSetName();
+	static UE_API FName GetStyleSetName();
 
 	static const FSlateBrush* GetBrush(FName PropertyName, const ANSICHAR* Specifier = NULL)
 	{
@@ -19,5 +21,7 @@ public:
 	}
 
 private:
-	static TSharedPtr<FSlateStyleSet> StyleSet;
+	static UE_API TSharedPtr<FSlateStyleSet> StyleSet;
 };
+
+#undef UE_API
diff --git a/Engine/Source/Editor/MassEntityDebugger/Public/MassDebuggerSettings.h b/Engine/Source/Editor/MassEntityDebugger/Public/MassDebuggerSettings.h
--- a/Engine/Source/Editor/MassEntityDebugger/Public/MassDebuggerSettings.h
+++ b/Engine/Source/Editor/MassEntityDebugger/Public/MassDebuggerSettings.h
@@ -8,8 +8,8 @@
 
 #define GET_MASSDEBUGGER_CONFIG_VALUE(a) (GetMutableDefault<UMassDebuggerSettings>()->a)
 
-UCLASS(config = EditorPerProjectUserSettings, DisplayName = "Mass Debugger")
-class MASSENTITYDEBUGGER_API UMassDebuggerSettings : public UMassModuleSettings
+UCLASS(MinimalAPI, config = EditorPerProjectUserSettings, DisplayName = "Mass Debugger")
+class UMassDebuggerSettings : public UMassModuleSettings
 {
 	GENERATED_BODY()
 
diff --git a/Engine/Source/Editor/MassEntityEditor/Public/MassEditorStyle.h b/Engine/Source/Editor/MassEntityEditor/Public/MassEditorStyle.h
--- a/Engine/Source/Editor/MassEntityEditor/Public/MassEditorStyle.h
+++ b/Engine/Source/Editor/MassEntityEditor/Public/MassEditorStyle.h
@@ -4,23 +4,27 @@
 
 #include "Styling/SlateStyle.h"
 
-class MASSENTITYEDITOR_API FMassEntityEditorStyle
+#define UE_API MASSENTITYEDITOR_API
+
+class FMassEntityEditorStyle
 {
 public:
-	static void Initialize();
-	static void Shutdown();
+	static UE_API void Initialize();
+	static UE_API void Shutdown();
 	static TSharedPtr<ISlateStyle> Get() { return StyleSet; }
-	static FName GetStyleSetName();
+	static UE_API FName GetStyleSetName();
 
 	static const FSlateBrush* GetBrush(FName PropertyName, const ANSICHAR* Specifier = NULL)
 	{
 		return StyleSet->GetBrush(PropertyName, Specifier);
 	}
 
-	static FString InContent(const FString& RelativePath, const ANSICHAR* Extension);
+	static UE_API FString InContent(const FString& RelativePath, const ANSICHAR* Extension);
 
 
 private:
 
-	static TSharedPtr<FSlateStyleSet> StyleSet;
+	static UE_API TSharedPtr<FSlateStyleSet> StyleSet;
 };
+
+#undef UE_API
diff --git a/Engine/Source/Editor/MassEntityEditor/Public/MassEntityEditor.h b/Engine/Source/Editor/MassEntityEditor/Public/MassEntityEditor.h
--- a/Engine/Source/Editor/MassEntityEditor/Public/MassEntityEditor.h
+++ b/Engine/Source/Editor/MassEntityEditor/Public/MassEntityEditor.h
@@ -5,9 +5,11 @@
 #include "Logging/TokenizedMessage.h"
 #include "MassEntityEditor.generated.h"
 
+#define UE_API MASSENTITYEDITOR_API
+
 
 USTRUCT()
-struct MASSENTITYEDITOR_API FMassEditorNotification
+struct FMassEditorNotification
 {
 	GENERATED_BODY()
 
@@ -20,10 +22,12 @@ struct MASSENTITYEDITOR_API FMassEditorNotification
 	 */
 	bool bIncludeSeeOutputLogForDetails = false;
 
-	void Show();
+	UE_API void Show();
 };
 
 namespace UE::Mass::Editor
 {
 	const FName MessageLogPageName(TEXT("MassEntity"));
-}
\ No newline at end of file
+}
+
+#undef UE_API
diff --git a/Engine/Source/Editor/MassEntityEditor/Public/MassEntityEditorModule.h b/Engine/Source/Editor/MassEntityEditor/Public/MassEntityEditorModule.h
--- a/Engine/Source/Editor/MassEntityEditor/Public/MassEntityEditorModule.h
+++ b/Engine/Source/Editor/MassEntityEditor/Public/MassEntityEditorModule.h
@@ -6,6 +6,8 @@
 #include "Toolkits/AssetEditorToolkit.h"
 #include "Delegates/IDelegateInstance.h"
 
+#define UE_API MASSENTITYEDITOR_API
+
 class IMassEntityEditor;
 struct FGraphPanelNodeFactory;
 struct FGraphNodeClassHelper;
@@ -14,11 +16,11 @@ class UWorld;
 /**
 * The public interface to this module
 */
-class MASSENTITYEDITOR_API FMassEntityEditorModule : public IModuleInterface, public IHasMenuExtensibility, public IHasToolBarExtensibility
+class FMassEntityEditorModule : public IModuleInterface, public IHasMenuExtensibility, public IHasToolBarExtensibility
 {
 public:
-	virtual void StartupModule() override;
-	virtual void ShutdownModule() override;
+	UE_API virtual void StartupModule() override;
+	UE_API virtual void ShutdownModule() override;
 
 	virtual TSharedPtr<FExtensibilityManager> GetMenuExtensibilityManager() override
 	{
@@ -37,7 +39,7 @@ public:
 
 protected:
 #if WITH_UNREAL_DEVELOPER_TOOLS
-	static void OnWorldCleanup(UWorld* /*World*/, bool /*bSessionEnded*/, bool /*bCleanupResources*/);
+	static UE_API void OnWorldCleanup(UWorld* /*World*/, bool /*bSessionEnded*/, bool /*bCleanupResources*/);
 	FDelegateHandle OnWorldCleanupHandle;
 #endif // WITH_UNREAL_DEVELOPER_TOOLS
 
@@ -46,3 +48,5 @@ protected:
 	TSharedPtr<FExtensibilityManager> MenuExtensibilityManager;
 	TSharedPtr<FExtensibilityManager> ToolBarExtensibilityManager;
 };
+
+#undef UE_API
diff --git a/Engine/Source/Editor/MassEntityEditor/Public/MassEntityEditorSubsystem.h b/Engine/Source/Editor/MassEntityEditor/Public/MassEntityEditorSubsystem.h
--- a/Engine/Source/Editor/MassEntityEditor/Public/MassEntityEditorSubsystem.h
+++ b/Engine/Source/Editor/MassEntityEditor/Public/MassEntityEditorSubsystem.h
@@ -6,6 +6,8 @@
 #include "MassProcessingPhaseManager.h"
 #include "MassEntityEditorSubsystem.generated.h"
 
+#define UE_API MASSENTITYEDITOR_API
+
 
 struct FMassEntityManager;
 namespace UE::Mass
@@ -13,8 +15,8 @@ namespace UE::Mass
 	struct FMassEditorTickFunction;
 }
 
-UCLASS()
-class MASSENTITYEDITOR_API UMassEntityEditorSubsystem : public UEditorSubsystem
+UCLASS(MinimalAPI)
+class UMassEntityEditorSubsystem : public UEditorSubsystem
 {
 	GENERATED_BODY()
 
@@ -22,8 +24,8 @@ public:
 	DECLARE_MULTICAST_DELEGATE_OneParam(FOnPreTick, float /*DeltaTime*/);
 	DECLARE_MULTICAST_DELEGATE_OneParam(FOnPostTick, float /*DeltaTime*/);
 
-	UMassEntityEditorSubsystem();
-	~UMassEntityEditorSubsystem();
+	UE_API UMassEntityEditorSubsystem();
+	UE_API ~UMassEntityEditorSubsystem();
 
 	TSharedRef<FMassEntityManager> GetMutableEntityManager() { return EntityManager; }
 	TSharedRef<FMassProcessingPhaseManager> GetMutablePhaseManager() { return PhaseManager; }
@@ -32,21 +34,21 @@ public:
 	FOnPostTick& GetOnPostTickDelegate() { return OnPostTickDelegate; }
 
 	/** Registers a dynamic processor. This needs to be a fully formed processor and will be slotted in during the next tick. */
-	void RegisterDynamicProcessor(UMassProcessor& Processor);
+	UE_API void RegisterDynamicProcessor(UMassProcessor& Processor);
 	/** Removes a previously registered dynamic processor of throws an assert if not found. */
-	void UnregisterDynamicProcessor(UMassProcessor& Processor);
+	UE_API void UnregisterDynamicProcessor(UMassProcessor& Processor);
 
 protected:
 	// USubsystem interface begin
-	virtual void Initialize(FSubsystemCollectionBase& Collection) override;
-	virtual void Deinitialize() override;
+	UE_API virtual void Initialize(FSubsystemCollectionBase& Collection) override;
+	UE_API virtual void Deinitialize() override;
 	// USubsystem interface end
 
-	void InitializeMassInternals();
-	void StopAndCleanUp();
+	UE_API void InitializeMassInternals();
+	UE_API void StopAndCleanUp();
 
 	friend UE::Mass::FMassEditorTickFunction;
-	void Tick(float DeltaTime);
+	UE_API void Tick(float DeltaTime);
 
 	UE::Mass::FMassEditorTickFunction* EditorTickFunction = nullptr;
 
@@ -64,3 +66,5 @@ protected:
 	bool IsProcessing{ false };
 	FDelegateHandle OnSettingsInitializedHandle;
 };
+
+#undef UE_API
diff --git a/Engine/Source/Editor/MaterialEditor/Public/MaterialEditor/MaterialNodes/SGraphNodeMaterialBase.h b/Engine/Source/Editor/MaterialEditor/Public/MaterialEditor/MaterialNodes/SGraphNodeMaterialBase.h
--- a/Engine/Source/Editor/MaterialEditor/Public/MaterialEditor/MaterialNodes/SGraphNodeMaterialBase.h
+++ b/Engine/Source/Editor/MaterialEditor/Public/MaterialEditor/MaterialNodes/SGraphNodeMaterialBase.h
@@ -14,6 +14,8 @@
 #include "Textures/SlateShaderResource.h"
 #include "Widgets/DeclarativeSyntaxSupport.h"
 
+#define UE_API MATERIALEDITOR_API
+
 class FMaterialRenderProxy;
 class FRHICommandListImmediate;
 class FSlateRect;
@@ -88,56 +90,56 @@ private:
 	bool bIsRealtime;
 };
 
-class MATERIALEDITOR_API SGraphNodeMaterialBase : public SGraphNode
+class SGraphNodeMaterialBase : public SGraphNode
 {
 public:
 	SLATE_BEGIN_ARGS(SGraphNodeMaterialBase){}
 	SLATE_END_ARGS()
 
-	void Construct(const FArguments& InArgs, class UMaterialGraphNode* InNode);
+	UE_API void Construct(const FArguments& InArgs, class UMaterialGraphNode* InNode);
 
 	// SGraphNode interface
-	virtual void CreatePinWidgets() override;
+	UE_API virtual void CreatePinWidgets() override;
 	// End of SGraphNode interface
 
 	// SNodePanel::SNode interface
-	virtual void MoveTo(const FVector2f& NewPosition, FNodeSet& NodeFilter, bool bMarkDirty = true) override;
+	UE_API virtual void MoveTo(const FVector2f& NewPosition, FNodeSet& NodeFilter, bool bMarkDirty = true) override;
 	// End of SNodePanel::SNode interface
 
 	UMaterialGraphNode* GetMaterialGraphNode() const {return MaterialNode;}
 
 	/* Populate a meta data tag with information about this graph node */
-	virtual void PopulateMetaTag(class FGraphNodeMetaData* TagMeta) const override;
+	UE_API virtual void PopulateMetaTag(class FGraphNodeMetaData* TagMeta) const override;
 
 protected:
 	// SGraphNode interface
-	virtual void AddPin( const TSharedRef<SGraphPin>& PinToAdd ) override;
-	virtual void CreateBelowPinControls(TSharedPtr<SVerticalBox> MainBox) override;
-	virtual void SetDefaultTitleAreaWidget(TSharedRef<SOverlay> DefaultTitleAreaWidget) override;
-	virtual TSharedRef<SWidget> CreateNodeContentArea() override;
-	virtual void OnAdvancedViewChanged(const ECheckBoxState NewCheckedState) override;
+	UE_API virtual void AddPin( const TSharedRef<SGraphPin>& PinToAdd ) override;
+	UE_API virtual void CreateBelowPinControls(TSharedPtr<SVerticalBox> MainBox) override;
+	UE_API virtual void SetDefaultTitleAreaWidget(TSharedRef<SOverlay> DefaultTitleAreaWidget) override;
+	UE_API virtual TSharedRef<SWidget> CreateNodeContentArea() override;
+	UE_API virtual void OnAdvancedViewChanged(const ECheckBoxState NewCheckedState) override;
 	// End of SGraphNode interface
 
 	/** Creates a preview viewport if necessary */
-	TSharedRef<SWidget> CreatePreviewWidget();
+	UE_API TSharedRef<SWidget> CreatePreviewWidget();
 
 	/** Returns visibility of Expression Preview viewport */
-	EVisibility ExpressionPreviewVisibility() const;
+	UE_API EVisibility ExpressionPreviewVisibility() const;
 
 	/** Returns text to over lay over the expression preview viewport */
-	FText ExpressionPreviewOverlayText() const;
+	UE_API FText ExpressionPreviewOverlayText() const;
 
 	/** Called when we click to preview node button */
-	FReply OnTogglePreviewClicked();
+	UE_API FReply OnTogglePreviewClicked();
 
 	/** Show/hide Expression Preview */
-	void OnExpressionPreviewChanged(const ECheckBoxState NewCheckedState);
+	UE_API void OnExpressionPreviewChanged(const ECheckBoxState NewCheckedState);
 
 	/** hidden == unchecked, shown == checked */
-	ECheckBoxState IsExpressionPreviewChecked() const;
+	UE_API ECheckBoxState IsExpressionPreviewChecked() const;
 
 	/** Up when shown, down when hidden */
-	const FSlateBrush* GetExpressionPreviewArrow() const;
+	UE_API const FSlateBrush* GetExpressionPreviewArrow() const;
 
 protected:
 	/** Slate viewport for rendering preview via custom slate element */
@@ -146,3 +148,5 @@ protected:
 	/** Cached material graph node pointer to avoid casting */
 	UMaterialGraphNode* MaterialNode;
 };
+
+#undef UE_API
diff --git a/Engine/Source/Editor/MaterialEditor/Public/MaterialEditor/MaterialNodes/SGraphNodeMaterialComposite.h b/Engine/Source/Editor/MaterialEditor/Public/MaterialEditor/MaterialNodes/SGraphNodeMaterialComposite.h
--- a/Engine/Source/Editor/MaterialEditor/Public/MaterialEditor/MaterialNodes/SGraphNodeMaterialComposite.h
+++ b/Engine/Source/Editor/MaterialEditor/Public/MaterialEditor/MaterialNodes/SGraphNodeMaterialComposite.h
@@ -9,37 +9,41 @@
 #include "Templates/SharedPointer.h"
 #include "Widgets/DeclarativeSyntaxSupport.h"
 
+#define UE_API MATERIALEDITOR_API
+
 class SToolTip;
 class SWidget;
 class UEdGraph;
 class UMaterialGraphNode_Composite;
 
 //@TODO: This class is mostly C&P from UK2Node_Composite, consolidate composites to not be BP specific.
-class MATERIALEDITOR_API SGraphNodeMaterialComposite : public SGraphNodeMaterialBase
+class SGraphNodeMaterialComposite : public SGraphNodeMaterialBase
 {
 public:
 	SLATE_BEGIN_ARGS(SGraphNodeMaterialComposite){}
 	SLATE_END_ARGS()
 
-	void Construct(const FArguments& InArgs, UMaterialGraphNode_Composite* InNode);
+	UE_API void Construct(const FArguments& InArgs, UMaterialGraphNode_Composite* InNode);
 
 	// SGraphNode interface
-	virtual void UpdateGraphNode() override;
-	virtual TSharedPtr<SToolTip> GetComplexTooltip() override;
+	UE_API virtual void UpdateGraphNode() override;
+	UE_API virtual TSharedPtr<SToolTip> GetComplexTooltip() override;
 	// End of SGraphNode interface
 
 	// SNodePanel::SNode interface
-	virtual void MoveTo(const FVector2f& NewPosition, FNodeSet& NodeFilter, bool bMarkDirty = true) override;
+	UE_API virtual void MoveTo(const FVector2f& NewPosition, FNodeSet& NodeFilter, bool bMarkDirty = true) override;
 	// End of SNodePanel::SNode interface
 protected:
-	virtual UEdGraph* GetInnerGraph() const;
+	UE_API virtual UEdGraph* GetInnerGraph() const;
 
 private:
-	FText GetPreviewCornerText() const;
-	FText GetTooltipTextForNode() const;
+	UE_API FText GetPreviewCornerText() const;
+	UE_API FText GetTooltipTextForNode() const;
 
-	TSharedRef<SWidget> CreateNodeBody();
+	UE_API TSharedRef<SWidget> CreateNodeBody();
 
 	/** Cached material graph node pointer to avoid casting */
 	UMaterialGraphNode_Composite* CompositeNode;
 };
+
+#undef UE_API
diff --git a/Engine/Source/Editor/MaterialEditor/Public/MaterialEditor/MaterialNodes/SGraphNodeMaterialConvert.h b/Engine/Source/Editor/MaterialEditor/Public/MaterialEditor/MaterialNodes/SGraphNodeMaterialConvert.h
--- a/Engine/Source/Editor/MaterialEditor/Public/MaterialEditor/MaterialNodes/SGraphNodeMaterialConvert.h
+++ b/Engine/Source/Editor/MaterialEditor/Public/MaterialEditor/MaterialNodes/SGraphNodeMaterialConvert.h
@@ -8,6 +8,8 @@
 #include "Types/SlateVector2.h"
 #include "Widgets/SCompoundWidget.h"
 
+#define UE_API MATERIALEDITOR_API
+
 class FConvertDragDropOp;
 class SConvertInnerPin;
 class SGraphNodeMaterialConvert;
@@ -141,32 +143,32 @@ protected:
 // SGraphNodeMaterialConvert: Custom Slate widget for UMaterialExpressionConvert
 ////////////////////////////////////////////////////////////////////////////////////////////////////////////////////////////////////////////
 
-class MATERIALEDITOR_API SGraphNodeMaterialConvert : public SGraphNodeMaterialBase
+class SGraphNodeMaterialConvert : public SGraphNodeMaterialBase
 {
 public:
 	SLATE_BEGIN_ARGS(SGraphNodeMaterialConvert){}
 	SLATE_END_ARGS()
 
-	void Construct(const FArguments& InArgs, class UMaterialGraphNode* InNode);
+	UE_API void Construct(const FArguments& InArgs, class UMaterialGraphNode* InNode);
 
 	// Begin SGraphNode interface
-	virtual void CreatePinWidgets() override;
-	virtual TSharedPtr<SGraphPin> CreatePinWidget(UEdGraphPin* Pin) const override;
-	virtual void AddPin(const TSharedRef<SGraphPin>& PinToAdd) override;
+	UE_API virtual void CreatePinWidgets() override;
+	UE_API virtual TSharedPtr<SGraphPin> CreatePinWidget(UEdGraphPin* Pin) const override;
+	UE_API virtual void AddPin(const TSharedRef<SGraphPin>& PinToAdd) override;
 
-	virtual void CreateInputSideAddButton(TSharedPtr<SVerticalBox> InputBox) override;
-	virtual void CreateOutputSideAddButton(TSharedPtr<SVerticalBox> OutputBox) override;
-	virtual EVisibility IsAddPinButtonVisible() const override;
+	UE_API virtual void CreateInputSideAddButton(TSharedPtr<SVerticalBox> InputBox) override;
+	UE_API virtual void CreateOutputSideAddButton(TSharedPtr<SVerticalBox> OutputBox) override;
+	UE_API virtual EVisibility IsAddPinButtonVisible() const override;
 	// End SGraphNode interface
 
-	FReply OnAddInputPinClicked();
-	FReply OnAddOutputPinClicked();
+	UE_API FReply OnAddInputPinClicked();
+	UE_API FReply OnAddOutputPinClicked();
 
-	TSharedRef<SWidget> CreateAddPinContextMenu(bool bInputPin);
-	void AddNewPin(bool bInputPin, EMaterialExpressionConvertType ConvertType);
+	UE_API TSharedRef<SWidget> CreateAddPinContextMenu(bool bInputPin);
+	UE_API void AddNewPin(bool bInputPin, EMaterialExpressionConvertType ConvertType);
 
 	// Begin SWidget Interface
-	virtual int32 OnPaint(
+	UE_API virtual int32 OnPaint(
 		const FPaintArgs& Args, 
 		const FGeometry& AllottedGeometry, 
 		const FSlateRect& MyCullingRect, 
@@ -177,18 +179,18 @@ public:
 	) const override;
 	// End SWidget Interface
 
-	void FormConnection(TSharedPtr<SConvertInnerPin> InnerPinA, TSharedPtr<SConvertInnerPin> InnerPinB);
-	void BreakConnections(TSharedPtr<SConvertInnerPin> InnerPin);
+	UE_API void FormConnection(TSharedPtr<SConvertInnerPin> InnerPinA, TSharedPtr<SConvertInnerPin> InnerPinB);
+	UE_API void BreakConnections(TSharedPtr<SConvertInnerPin> InnerPin);
 
-	TOptional<float> GetDefaultValue(TSharedPtr<const SConvertInnerPin> InnerPin) const;
-	void SetDefaultValue(TSharedPtr<SConvertInnerPin> InnerPin, const float InDefaultValue);
+	UE_API TOptional<float> GetDefaultValue(TSharedPtr<const SConvertInnerPin> InnerPin) const;
+	UE_API void SetDefaultValue(TSharedPtr<SConvertInnerPin> InnerPin, const float InDefaultValue);
 
 	void SetCurrentDragDropOp(TSharedPtr<FConvertDragDropOp> InDragDropOp) { CurrentDragDropOp = InDragDropOp; }
 	TSharedPtr<FConvertDragDropOp> GetCurrentDragDropOp() const { return CurrentDragDropOp; }
 
 protected:
-	TSharedPtr<SConvertInnerPin> GetInnerPin(bool bInputPin, const int32 InPinIndex, const int32 InComponentIndex) const;
-	void MakeConnectionCurve(
+	UE_API TSharedPtr<SConvertInnerPin> GetInnerPin(bool bInputPin, const int32 InPinIndex, const int32 InComponentIndex) const;
+	UE_API void MakeConnectionCurve(
 		const FGeometry& InAllottedGeometry,
 		FSlateWindowElementList& OutDrawElements,
 		int32& InOutLayerId, 
@@ -198,4 +200,6 @@ protected:
 
 private:
 	TSharedPtr<FConvertDragDropOp> CurrentDragDropOp;
-};
\ No newline at end of file
+};
+
+#undef UE_API
diff --git a/Engine/Source/Editor/MaterialEditor/Public/MaterialEditor/MaterialNodes/SGraphNodeMaterialCustom.h b/Engine/Source/Editor/MaterialEditor/Public/MaterialEditor/MaterialNodes/SGraphNodeMaterialCustom.h
--- a/Engine/Source/Editor/MaterialEditor/Public/MaterialEditor/MaterialNodes/SGraphNodeMaterialCustom.h
+++ b/Engine/Source/Editor/MaterialEditor/Public/MaterialEditor/MaterialNodes/SGraphNodeMaterialCustom.h
@@ -6,34 +6,38 @@
 #include "Text/HLSLSyntaxHighlighterMarshaller.h"
 #include "Widgets/Input/SMultiLineEditableTextBox.h"
 
-class MATERIALEDITOR_API SGraphNodeMaterialCustom : public SGraphNodeMaterialBase
+#define UE_API MATERIALEDITOR_API
+
+class SGraphNodeMaterialCustom : public SGraphNodeMaterialBase
 {
 public:
 	SLATE_BEGIN_ARGS(SGraphNodeMaterialCustom){}
 	SLATE_END_ARGS()
 
-	void Construct(const FArguments& InArgs, class UMaterialGraphNode* InNode);
+	UE_API void Construct(const FArguments& InArgs, class UMaterialGraphNode* InNode);
 
 	UMaterialGraphNode* GetMaterialGraphNode() const {return MaterialNode;}
 
-	FText GetHlslText() const;
-	void OnCustomHlslTextCommited(const FText& InText, ETextCommit::Type InType);
+	UE_API FText GetHlslText() const;
+	UE_API void OnCustomHlslTextCommited(const FText& InText, ETextCommit::Type InType);
 
 protected:
 	// SGraphNode interface
-	virtual void CreateBelowPinControls(TSharedPtr<SVerticalBox> MainBox) override;
-	virtual void CreateAdvancedViewArrow(TSharedPtr<SVerticalBox> MainBox) override;
-	virtual EVisibility AdvancedViewArrowVisibility() const override;
-	virtual void OnAdvancedViewChanged( const ECheckBoxState NewCheckedState ) override;
-	virtual ECheckBoxState IsAdvancedViewChecked() const override;
-	virtual const FSlateBrush* GetAdvancedViewArrow() const override;
+	UE_API virtual void CreateBelowPinControls(TSharedPtr<SVerticalBox> MainBox) override;
+	UE_API virtual void CreateAdvancedViewArrow(TSharedPtr<SVerticalBox> MainBox) override;
+	UE_API virtual EVisibility AdvancedViewArrowVisibility() const override;
+	UE_API virtual void OnAdvancedViewChanged( const ECheckBoxState NewCheckedState ) override;
+	UE_API virtual ECheckBoxState IsAdvancedViewChecked() const override;
+	UE_API virtual const FSlateBrush* GetAdvancedViewArrow() const override;
 	// End of SGraphNode interface
 
 private:
-	EVisibility CodeVisibility() const;
+	UE_API EVisibility CodeVisibility() const;
 
 protected:
 	/** Shader text box */
 	TSharedPtr<FHLSLSyntaxHighlighterMarshaller> SyntaxHighlighter;
 	TSharedPtr<SMultiLineEditableTextBox> ShaderTextBox;
 };
+
+#undef UE_API
diff --git a/Engine/Source/Editor/MaterialEditor/Public/MaterialEditor/MaterialNodes/SGraphNodeMaterialResult.h b/Engine/Source/Editor/MaterialEditor/Public/MaterialEditor/MaterialNodes/SGraphNodeMaterialResult.h
--- a/Engine/Source/Editor/MaterialEditor/Public/MaterialEditor/MaterialNodes/SGraphNodeMaterialResult.h
+++ b/Engine/Source/Editor/MaterialEditor/Public/MaterialEditor/MaterialNodes/SGraphNodeMaterialResult.h
@@ -7,26 +7,30 @@
 #include "SNodePanel.h"
 #include "Widgets/DeclarativeSyntaxSupport.h"
 
+#define UE_API MATERIALEDITOR_API
+
 class FGraphNodeMetaData;
 class UMaterialGraphNode_Root;
 
-class MATERIALEDITOR_API SGraphNodeMaterialResult : public SGraphNode
+class SGraphNodeMaterialResult : public SGraphNode
 {
 public:
 	SLATE_BEGIN_ARGS(SGraphNodeMaterialResult){}
 	SLATE_END_ARGS()
 
-	void Construct(const FArguments& InArgs, class UMaterialGraphNode_Root* InNode);
+	UE_API void Construct(const FArguments& InArgs, class UMaterialGraphNode_Root* InNode);
 
 	// SGraphNode interface
-	virtual void CreatePinWidgets() override;
-	virtual void PopulateMetaTag(FGraphNodeMetaData* TagMeta) const override;
+	UE_API virtual void CreatePinWidgets() override;
+	UE_API virtual void PopulateMetaTag(FGraphNodeMetaData* TagMeta) const override;
 	// End of SGraphNode interface
 
 	// SNodePanel::SNode interface
-	virtual void MoveTo(const FVector2f& NewPosition, FNodeSet& NodeFilter, bool bMarkDirty = true) override;
+	UE_API virtual void MoveTo(const FVector2f& NewPosition, FNodeSet& NodeFilter, bool bMarkDirty = true) override;
 	// End of SNodePanel::SNode interface
 
 private:
 	UMaterialGraphNode_Root* RootNode;
 };
+
+#undef UE_API
diff --git a/Engine/Source/Editor/MaterialEditor/Public/MaterialEditor/MaterialPins/SGraphPinMaterialInput.h b/Engine/Source/Editor/MaterialEditor/Public/MaterialEditor/MaterialPins/SGraphPinMaterialInput.h
--- a/Engine/Source/Editor/MaterialEditor/Public/MaterialEditor/MaterialPins/SGraphPinMaterialInput.h
+++ b/Engine/Source/Editor/MaterialEditor/Public/MaterialEditor/MaterialPins/SGraphPinMaterialInput.h
@@ -7,19 +7,23 @@
 #include "Styling/SlateColor.h"
 #include "Widgets/DeclarativeSyntaxSupport.h"
 
+#define UE_API MATERIALEDITOR_API
+
 class UEdGraphPin;
 
-class MATERIALEDITOR_API SGraphPinMaterialInput : public SGraphPin
+class SGraphPinMaterialInput : public SGraphPin
 {
 public:
 	SLATE_BEGIN_ARGS(SGraphPinMaterialInput) {}
 	SLATE_END_ARGS()
 
-	void Construct(const FArguments& InArgs, UEdGraphPin* InGraphPinObj);
+	UE_API void Construct(const FArguments& InArgs, UEdGraphPin* InGraphPinObj);
 
 protected:
 	//~ Begin SGraphPin Interface
-	virtual FSlateColor GetPinColor() const override;
+	UE_API virtual FSlateColor GetPinColor() const override;
 	//~ End SGraphPin Interface
 
 };
+
+#undef UE_API
diff --git a/Engine/Source/Editor/MaterialEditor/Public/MaterialEditor/SGraphSubstrateMaterial.h b/Engine/Source/Editor/MaterialEditor/Public/MaterialEditor/SGraphSubstrateMaterial.h
--- a/Engine/Source/Editor/MaterialEditor/Public/MaterialEditor/SGraphSubstrateMaterial.h
+++ b/Engine/Source/Editor/MaterialEditor/Public/MaterialEditor/SGraphSubstrateMaterial.h
@@ -6,15 +6,19 @@
 #include "Rendering/RenderingCommon.h"
 #include "SGraphPin.h"
 
+#define UE_API MATERIALEDITOR_API
+
 class UMaterialGraphNode;
 struct FSubstrateMaterialCompilationOutput;
 
-struct MATERIALEDITOR_API FSubstrateWidget
+struct FSubstrateWidget
 {
-	static const TSharedRef<SWidget> ProcessOperator(const FSubstrateMaterialCompilationOutput& CompilationOutput);
-	static const TSharedRef<SWidget> ProcessOperator(const FSubstrateMaterialCompilationOutput& CompilationOutput, const TArray<FGuid>& InGuid);
-	static void GetPinColor(TSharedPtr<SGraphPin>& Out, const UMaterialGraphNode* InNode);
-	static FLinearColor GetConnectionColor();
-	static bool HasInputSubstrateType(const UEdGraphPin* InPin);
-	static bool HasOutputSubstrateType(const UEdGraphPin* InPin);
-};
\ No newline at end of file
+	static UE_API const TSharedRef<SWidget> ProcessOperator(const FSubstrateMaterialCompilationOutput& CompilationOutput);
+	static UE_API const TSharedRef<SWidget> ProcessOperator(const FSubstrateMaterialCompilationOutput& CompilationOutput, const TArray<FGuid>& InGuid);
+	static UE_API void GetPinColor(TSharedPtr<SGraphPin>& Out, const UMaterialGraphNode* InNode);
+	static UE_API FLinearColor GetConnectionColor();
+	static UE_API bool HasInputSubstrateType(const UEdGraphPin* InPin);
+	static UE_API bool HasOutputSubstrateType(const UEdGraphPin* InPin);
+};
+
+#undef UE_API
diff --git a/Engine/Source/Editor/MaterialEditor/Public/MaterialEditorContext.h b/Engine/Source/Editor/MaterialEditor/Public/MaterialEditorContext.h
--- a/Engine/Source/Editor/MaterialEditor/Public/MaterialEditorContext.h
+++ b/Engine/Source/Editor/MaterialEditor/Public/MaterialEditorContext.h
@@ -8,8 +8,8 @@
 
 class IMaterialEditor;
 
-UCLASS()
-class MATERIALEDITOR_API UMaterialEditorMenuContext : public UObject
+UCLASS(MinimalAPI)
+class UMaterialEditorMenuContext : public UObject
 {
 	GENERATED_BODY()
 
diff --git a/Engine/Source/Editor/MaterialEditor/Public/MaterialEditorModes.h b/Engine/Source/Editor/MaterialEditor/Public/MaterialEditorModes.h
--- a/Engine/Source/Editor/MaterialEditor/Public/MaterialEditorModes.h
+++ b/Engine/Source/Editor/MaterialEditor/Public/MaterialEditorModes.h
@@ -7,10 +7,12 @@
 #include "MaterialEditor.h"
 #include "WorkflowOrientedApp/ApplicationMode.h"
 
-struct MATERIALEDITOR_API FMaterialEditorApplicationModes
+#define UE_API MATERIALEDITOR_API
+
+struct FMaterialEditorApplicationModes
 {
 	// Mode identifiers
-	static const FName StandardMaterialEditorMode;
+	static UE_API const FName StandardMaterialEditorMode;
 	static FText GetLocalizedMode(const FName InMode)
 	{
 		static TMap< FName, FText > LocModes;
@@ -25,18 +27,18 @@ struct MATERIALEDITOR_API FMaterialEditorApplicationModes
 		check(OutDesc);
 		return *OutDesc;
 	}
-	static TSharedPtr<FTabManager::FLayout> GetDefaultEditorLayout(TSharedPtr<class FMaterialEditor> InMaterialEditor);
+	static UE_API TSharedPtr<FTabManager::FLayout> GetDefaultEditorLayout(TSharedPtr<class FMaterialEditor> InMaterialEditor);
 private:
 	FMaterialEditorApplicationModes() {}
 };
 
 // Even though we currently only have one mode of operation, we still need a application mode to handle factory tab spawning
-class MATERIALEDITOR_API FMaterialEditorApplicationMode : public FApplicationMode
+class FMaterialEditorApplicationMode : public FApplicationMode
 {
 public:
-	FMaterialEditorApplicationMode(TSharedPtr<class FMaterialEditor> InMaterialEditor);
+	UE_API FMaterialEditorApplicationMode(TSharedPtr<class FMaterialEditor> InMaterialEditor);
 
-	virtual void RegisterTabFactories(TSharedPtr<FTabManager> InTabManager) override;
+	UE_API virtual void RegisterTabFactories(TSharedPtr<FTabManager> InTabManager) override;
 public:
 
 protected:
@@ -44,4 +46,6 @@ protected:
 
 	// Set of spawnable tabs handled by workflow factories
 	FWorkflowAllowedTabSet MaterialEditorTabFactories;
-};
\ No newline at end of file
+};
+
+#undef UE_API
