SEditorViewport::MakeViewportToolbar() is deprecated.

# All existing usages migrate to BuildViewportToolbar() or to placing the old toolbar directly in PopulateViewportOverlays().

diff --git a/Engine/Plugins/2D/Paper2D/Source/Paper2DEditor/Private/FlipbookEditor/FlipbookEditor.cpp b/Engine/Plugins/2D/Paper2D/Source/Paper2DEditor/Private/FlipbookEditor/FlipbookEditor.cpp
--- a/Engine/Plugins/2D/Paper2D/Source/Paper2DEditor/Private/FlipbookEditor/FlipbookEditor.cpp
+++ b/Engine/Plugins/2D/Paper2D/Source/Paper2DEditor/Private/FlipbookEditor/FlipbookEditor.cpp
@@ -57,7 +57,7 @@ public:
 	// SEditorViewport interface
 	virtual void BindCommands() override;
 	virtual TSharedRef<FEditorViewportClient> MakeEditorViewportClient() override;
-	virtual TSharedPtr<SWidget> MakeViewportToolbar() override;
+	virtual TSharedPtr<SWidget> BuildViewportToolbar() override;
 	virtual EVisibility GetTransformToolbarVisibility() const override;
 	virtual void OnFocusViewportToSelection() override;
 	// End of SEditorViewport interface
@@ -133,7 +133,7 @@ TSharedRef<FEditorViewportClient> SFlipbookEditorViewport::MakeEditorViewportCli
 	return EditorViewportClient.ToSharedRef();
 }
 
-TSharedPtr<SWidget> SFlipbookEditorViewport::MakeViewportToolbar()
+TSharedPtr<SWidget> SFlipbookEditorViewport::BuildViewportToolbar()
 {
 	return SNew(SFlipbookEditorViewportToolbar, SharedThis(this));
 }
diff --git a/Engine/Plugins/2D/Paper2D/Source/Paper2DEditor/Private/SpriteEditor/SpriteEditor.cpp b/Engine/Plugins/2D/Paper2D/Source/Paper2DEditor/Private/SpriteEditor/SpriteEditor.cpp
--- a/Engine/Plugins/2D/Paper2D/Source/Paper2DEditor/Private/SpriteEditor/SpriteEditor.cpp
+++ b/Engine/Plugins/2D/Paper2D/Source/Paper2DEditor/Private/SpriteEditor/SpriteEditor.cpp
@@ -55,7 +55,7 @@ public:
 	// SEditorViewport interface
 	virtual void BindCommands() override;
 	virtual TSharedRef<FEditorViewportClient> MakeEditorViewportClient() override;
-	virtual TSharedPtr<SWidget> MakeViewportToolbar() override;
+	virtual TSharedPtr<SWidget> BuildViewportToolbar() override;
 	virtual EVisibility GetTransformToolbarVisibility() const override;
 	virtual void OnFocusViewportToSelection() override;
 	// End of SEditorViewport interface
@@ -201,7 +201,7 @@ TSharedRef<FEditorViewportClient> SSpriteEditorViewport::MakeEditorViewportClien
 	return EditorViewportClient.ToSharedRef();
 }
 
-TSharedPtr<SWidget> SSpriteEditorViewport::MakeViewportToolbar()
+TSharedPtr<SWidget> SSpriteEditorViewport::BuildViewportToolbar()
 {
 	return SNew(SSpriteEditorViewportToolbar, SharedThis(this));
 }
diff --git a/Engine/Plugins/2D/Paper2D/Source/Paper2DEditor/Private/TileMapEditing/TileMapEditor.cpp b/Engine/Plugins/2D/Paper2D/Source/Paper2DEditor/Private/TileMapEditing/TileMapEditor.cpp
--- a/Engine/Plugins/2D/Paper2D/Source/Paper2DEditor/Private/TileMapEditing/TileMapEditor.cpp
+++ b/Engine/Plugins/2D/Paper2D/Source/Paper2DEditor/Private/TileMapEditing/TileMapEditor.cpp
@@ -49,7 +49,7 @@ public:
 	// SEditorViewport interface
 	virtual void BindCommands() override;
 	virtual TSharedRef<FEditorViewportClient> MakeEditorViewportClient() override;
-	virtual TSharedPtr<SWidget> MakeViewportToolbar() override;
+	virtual TSharedPtr<SWidget> BuildViewportToolbar() override;
 	virtual EVisibility GetTransformToolbarVisibility() const override;
 	virtual void OnFocusViewportToSelection() override;
 	// End of SEditorViewport interface
@@ -132,7 +132,7 @@ TSharedRef<FEditorViewportClient> STileMapEditorViewport::MakeEditorViewportClie
 	return EditorViewportClient.ToSharedRef();
 }
 
-TSharedPtr<SWidget> STileMapEditorViewport::MakeViewportToolbar()
+TSharedPtr<SWidget> STileMapEditorViewport::BuildViewportToolbar()
 {
 	return SNew(STileMapEditorViewportToolbar, SharedThis(this));
 }
diff --git a/Engine/Plugins/2D/Paper2D/Source/Paper2DEditor/Private/TileSetEditor/SingleTileEditorViewport.cpp b/Engine/Plugins/2D/Paper2D/Source/Paper2DEditor/Private/TileSetEditor/SingleTileEditorViewport.cpp
--- a/Engine/Plugins/2D/Paper2D/Source/Paper2DEditor/Private/TileSetEditor/SingleTileEditorViewport.cpp
+++ b/Engine/Plugins/2D/Paper2D/Source/Paper2DEditor/Private/TileSetEditor/SingleTileEditorViewport.cpp
@@ -102,7 +102,7 @@ void SSingleTileEditorViewport::Construct(const FArguments& InArgs, TSharedPtr<c
 	];
 }
 
-TSharedPtr<SWidget> SSingleTileEditorViewport::MakeViewportToolbar()
+TSharedPtr<SWidget> SSingleTileEditorViewport::BuildViewportToolbar()
 {
 	return SNew(STileSetEditorViewportToolbar, SharedThis(this));
 }
diff --git a/Engine/Plugins/2D/Paper2D/Source/Paper2DEditor/Private/TileSetEditor/SingleTileEditorViewport.h b/Engine/Plugins/2D/Paper2D/Source/Paper2DEditor/Private/TileSetEditor/SingleTileEditorViewport.h
--- a/Engine/Plugins/2D/Paper2D/Source/Paper2DEditor/Private/TileSetEditor/SingleTileEditorViewport.h
+++ b/Engine/Plugins/2D/Paper2D/Source/Paper2DEditor/Private/TileSetEditor/SingleTileEditorViewport.h
@@ -22,7 +22,7 @@ public:
 	void Construct(const FArguments& InArgs, TSharedPtr<FSingleTileEditorViewportClient> InViewportClient);
 
 	// SEditorViewport interface
-	virtual TSharedPtr<SWidget> MakeViewportToolbar() override;
+	virtual TSharedPtr<SWidget> BuildViewportToolbar() override;
 	virtual TSharedRef<FEditorViewportClient> MakeEditorViewportClient() override;
 	virtual void BindCommands() override;
 	virtual void OnFocusViewportToSelection() override;
diff --git a/Engine/Plugins/Animation/PoseSearch/Source/Editor/Private/PoseSearchInteractionAssetEditor.cpp b/Engine/Plugins/Animation/PoseSearch/Source/Editor/Private/PoseSearchInteractionAssetEditor.cpp
--- a/Engine/Plugins/Animation/PoseSearch/Source/Editor/Private/PoseSearchInteractionAssetEditor.cpp
+++ b/Engine/Plugins/Animation/PoseSearch/Source/Editor/Private/PoseSearchInteractionAssetEditor.cpp
@@ -577,7 +577,7 @@ TSharedRef<FEditorViewportClient> SInteractionAssetViewport::MakeEditorViewportC
 	return ViewportClient.ToSharedRef();
 }
 
-TSharedPtr<SWidget> SInteractionAssetViewport::MakeViewportToolbar()
+TSharedPtr<SWidget> SInteractionAssetViewport::BuildViewportToolbar()
 {
 	return SAssignNew(ViewportToolbar, SCommonEditorViewportToolbarBase, SharedThis(this));
 }
diff --git a/Engine/Plugins/Animation/PoseSearch/Source/Editor/Private/SPoseSearchDatabaseViewport.cpp b/Engine/Plugins/Animation/PoseSearch/Source/Editor/Private/SPoseSearchDatabaseViewport.cpp
--- a/Engine/Plugins/Animation/PoseSearch/Source/Editor/Private/SPoseSearchDatabaseViewport.cpp
+++ b/Engine/Plugins/Animation/PoseSearch/Source/Editor/Private/SPoseSearchDatabaseViewport.cpp
@@ -135,7 +135,7 @@ namespace UE::PoseSearch
 		return ViewportClient.ToSharedRef();
 	}
 
-	TSharedPtr<SWidget> SDatabaseViewport::MakeViewportToolbar()
+	TSharedPtr<SWidget> SDatabaseViewport::BuildViewportToolbar()
 	{
 		return SAssignNew(ViewportToolbar, SPoseSearchDatabaseViewportToolBar, SharedThis(this));
 	}
diff --git a/Engine/Plugins/Animation/PoseSearch/Source/Editor/Public/PoseSearchInteractionAssetEditor.h b/Engine/Plugins/Animation/PoseSearch/Source/Editor/Public/PoseSearchInteractionAssetEditor.h
--- a/Engine/Plugins/Animation/PoseSearch/Source/Editor/Public/PoseSearchInteractionAssetEditor.h
+++ b/Engine/Plugins/Animation/PoseSearch/Source/Editor/Public/PoseSearchInteractionAssetEditor.h
@@ -165,7 +165,7 @@ public:
 protected:
 	virtual void BindCommands() override;
 	virtual TSharedRef<FEditorViewportClient> MakeEditorViewportClient() override;
-	virtual TSharedPtr<SWidget> MakeViewportToolbar() override;
+	virtual TSharedPtr<SWidget> BuildViewportToolbar() override;
 
 	/** The viewport toolbar */
 	TSharedPtr<SCommonEditorViewportToolbarBase> ViewportToolbar;
diff --git a/Engine/Plugins/Animation/PoseSearch/Source/Editor/Public/SPoseSearchDatabaseViewport.h b/Engine/Plugins/Animation/PoseSearch/Source/Editor/Public/SPoseSearchDatabaseViewport.h
--- a/Engine/Plugins/Animation/PoseSearch/Source/Editor/Public/SPoseSearchDatabaseViewport.h
+++ b/Engine/Plugins/Animation/PoseSearch/Source/Editor/Public/SPoseSearchDatabaseViewport.h
@@ -49,7 +49,7 @@ namespace UE::PoseSearch
 		// ~SEditorViewport interface
 		virtual void BindCommands() override;
 		virtual TSharedRef<FEditorViewportClient> MakeEditorViewportClient() override;
-		virtual TSharedPtr<SWidget> MakeViewportToolbar() override;
+		virtual TSharedPtr<SWidget> BuildViewportToolbar() override;
 		// ~End of SEditorViewport interface
 
 		void OnToggleDisplayRootMotionSpeed();
diff --git a/Engine/Plugins/ChaosVD/Source/ChaosVD/Private/Widgets/SChaosVDPlaybackViewport.cpp b/Engine/Plugins/ChaosVD/Source/ChaosVD/Private/Widgets/SChaosVDPlaybackViewport.cpp
--- a/Engine/Plugins/ChaosVD/Source/ChaosVD/Private/Widgets/SChaosVDPlaybackViewport.cpp
+++ b/Engine/Plugins/ChaosVD/Source/ChaosVD/Private/Widgets/SChaosVDPlaybackViewport.cpp
@@ -469,7 +469,7 @@ TSharedRef<FEditorViewportClient> SChaosVDPlaybackViewport::MakeEditorViewportCl
 	return StaticCastSharedRef<FEditorViewportClient>(NewViewport.ToSharedRef());
 }
 
-TSharedPtr<SWidget> SChaosVDPlaybackViewport::MakeViewportToolbar()
+TSharedPtr<SWidget> SChaosVDPlaybackViewport::BuildViewportToolbar()
 {
 	// Build our toolbar level toolbar
 	TSharedRef< SChaosVDViewportToolbar > ToolBar = SNew(SChaosVDViewportToolbar, SharedThis(this));
diff --git a/Engine/Plugins/ChaosVD/Source/ChaosVD/Private/Widgets/SChaosVDPlaybackViewport.h b/Engine/Plugins/ChaosVD/Source/ChaosVD/Private/Widgets/SChaosVDPlaybackViewport.h
--- a/Engine/Plugins/ChaosVD/Source/ChaosVD/Private/Widgets/SChaosVDPlaybackViewport.h
+++ b/Engine/Plugins/ChaosVD/Source/ChaosVD/Private/Widgets/SChaosVDPlaybackViewport.h
@@ -74,7 +74,7 @@ protected:
 
 	virtual TSharedRef<FEditorViewportClient> MakeEditorViewportClient() override;
 	
-	virtual TSharedPtr<SWidget> MakeViewportToolbar() override;
+	virtual TSharedPtr<SWidget> BuildViewportToolbar() override;
 
 	virtual void RegisterNewController(TWeakPtr<FChaosVDPlaybackController> NewController )override;
 	virtual void HandlePlaybackControllerDataUpdated(TWeakPtr<FChaosVDPlaybackController> InController) override;
diff --git a/Engine/Plugins/Enterprise/DataprepEditor/Source/DataprepEditor/Private/Widgets/SDataprepEditorViewport.cpp b/Engine/Plugins/Enterprise/DataprepEditor/Source/DataprepEditor/Private/Widgets/SDataprepEditorViewport.cpp
--- a/Engine/Plugins/Enterprise/DataprepEditor/Source/DataprepEditor/Private/Widgets/SDataprepEditorViewport.cpp
+++ b/Engine/Plugins/Enterprise/DataprepEditor/Source/DataprepEditor/Private/Widgets/SDataprepEditorViewport.cpp
@@ -559,7 +559,7 @@ TSharedRef<FEditorViewportClient> SDataprepEditorViewport::MakeEditorViewportCli
 	return EditorViewportClient.ToSharedRef();
 }
 
-TSharedPtr<SWidget> SDataprepEditorViewport::MakeViewportToolbar()
+TSharedPtr<SWidget> SDataprepEditorViewport::BuildViewportToolbar()
 {
 	return SNew(SDataprepEditorViewportToolbar, SharedThis(this));
 }
diff --git a/Engine/Plugins/Enterprise/DataprepEditor/Source/DataprepEditor/Private/Widgets/SDataprepEditorViewport.h b/Engine/Plugins/Enterprise/DataprepEditor/Source/DataprepEditor/Private/Widgets/SDataprepEditorViewport.h
--- a/Engine/Plugins/Enterprise/DataprepEditor/Source/DataprepEditor/Private/Widgets/SDataprepEditorViewport.h
+++ b/Engine/Plugins/Enterprise/DataprepEditor/Source/DataprepEditor/Private/Widgets/SDataprepEditorViewport.h
@@ -193,7 +193,7 @@ protected:
 	virtual EVisibility OnGetViewportContentVisibility() const override;
 	virtual void BindCommands() override;
 	virtual void OnFocusViewportToSelection() override;
-	virtual TSharedPtr<SWidget> MakeViewportToolbar() override;
+	virtual TSharedPtr<SWidget> BuildViewportToolbar() override;
 	virtual void PopulateViewportOverlays(TSharedRef<SOverlay> Overlay) override;
 
 private:
diff --git a/Engine/Plugins/Enterprise/LidarPointCloud/Source/LidarPointCloudEditor/Private/LidarPointCloudEditorViewport.cpp b/Engine/Plugins/Enterprise/LidarPointCloud/Source/LidarPointCloudEditor/Private/LidarPointCloudEditorViewport.cpp
--- a/Engine/Plugins/Enterprise/LidarPointCloud/Source/LidarPointCloudEditor/Private/LidarPointCloudEditorViewport.cpp
+++ b/Engine/Plugins/Enterprise/LidarPointCloud/Source/LidarPointCloudEditor/Private/LidarPointCloudEditorViewport.cpp
@@ -169,7 +169,7 @@ TSharedRef<FEditorViewportClient> SLidarPointCloudEditorViewport::MakeEditorView
 	return EditorViewportClient.ToSharedRef();
 }
 
-TSharedPtr<SWidget> SLidarPointCloudEditorViewport::MakeViewportToolbar()
+TSharedPtr<SWidget> SLidarPointCloudEditorViewport::BuildViewportToolbar()
 {
 	return SNew(SPointCloudEditorViewportToolbar, SharedThis(this));
 }
diff --git a/Engine/Plugins/Enterprise/LidarPointCloud/Source/LidarPointCloudEditor/Private/LidarPointCloudEditorViewport.h b/Engine/Plugins/Enterprise/LidarPointCloud/Source/LidarPointCloudEditor/Private/LidarPointCloudEditorViewport.h
--- a/Engine/Plugins/Enterprise/LidarPointCloud/Source/LidarPointCloudEditor/Private/LidarPointCloudEditorViewport.h
+++ b/Engine/Plugins/Enterprise/LidarPointCloud/Source/LidarPointCloudEditor/Private/LidarPointCloudEditorViewport.h
@@ -84,7 +84,7 @@ protected:
 	virtual EVisibility OnGetViewportContentVisibility() const override { return IsVisible() ? EVisibility::Visible : EVisibility::Collapsed; }
 	virtual void BindCommands() override;
 	virtual void OnFocusViewportToSelection() override;
-	virtual TSharedPtr<SWidget> MakeViewportToolbar() override;
+	virtual TSharedPtr<SWidget> BuildViewportToolbar() override;
 
 private:
 	/** Determines the visibility of the viewport. */
diff --git a/Engine/Plugins/Experimental/ActorPalette/Source/ActorPalette/Private/ActorPaletteViewport.cpp b/Engine/Plugins/Experimental/ActorPalette/Source/ActorPalette/Private/ActorPaletteViewport.cpp
--- a/Engine/Plugins/Experimental/ActorPalette/Source/ActorPalette/Private/ActorPaletteViewport.cpp
+++ b/Engine/Plugins/Experimental/ActorPalette/Source/ActorPalette/Private/ActorPaletteViewport.cpp
@@ -260,7 +260,7 @@ void SActorPaletteViewport::Construct(const FArguments& InArgs, TSharedPtr<class
 	];
 }
 
-TSharedPtr<SWidget> SActorPaletteViewport::MakeViewportToolbar()
+TSharedPtr<SWidget> SActorPaletteViewport::BuildViewportToolbar()
 {
 	return SNew(SActorPaletteViewportToolbar, SharedThis(this), FOnGetContent::CreateSP(this, &SActorPaletteViewport::GenerateMapMenu));
 }
diff --git a/Engine/Plugins/Experimental/ActorPalette/Source/ActorPalette/Private/ActorPaletteViewport.h b/Engine/Plugins/Experimental/ActorPalette/Source/ActorPalette/Private/ActorPaletteViewport.h
--- a/Engine/Plugins/Experimental/ActorPalette/Source/ActorPalette/Private/ActorPaletteViewport.h
+++ b/Engine/Plugins/Experimental/ActorPalette/Source/ActorPalette/Private/ActorPaletteViewport.h
@@ -23,7 +23,7 @@ public:
 //	virtual FReply OnMouseButtonDown(const FGeometry& MyGeometry, const FPointerEvent& MouseEvent) override;
 
 	// SEditorViewport interface
-	virtual TSharedPtr<SWidget> MakeViewportToolbar() override;
+	virtual TSharedPtr<SWidget> BuildViewportToolbar() override;
 	virtual TSharedRef<FEditorViewportClient> MakeEditorViewportClient() override;
 	virtual void BindCommands() override;
 	virtual void OnFocusViewportToSelection() override;
diff --git a/Engine/Plugins/Experimental/Animation/ContextualAnimation/Source/ContextualAnimationEditor/Private/SContextualAnimViewport.cpp b/Engine/Plugins/Experimental/Animation/ContextualAnimation/Source/ContextualAnimationEditor/Private/SContextualAnimViewport.cpp
--- a/Engine/Plugins/Experimental/Animation/ContextualAnimation/Source/ContextualAnimationEditor/Private/SContextualAnimViewport.cpp
+++ b/Engine/Plugins/Experimental/Animation/ContextualAnimation/Source/ContextualAnimationEditor/Private/SContextualAnimViewport.cpp
@@ -259,7 +259,7 @@ TSharedRef<FEditorViewportClient> SContextualAnimViewport::MakeEditorViewportCli
 	return ViewportClient.ToSharedRef();
 }
 
-TSharedPtr<SWidget> SContextualAnimViewport::MakeViewportToolbar()
+TSharedPtr<SWidget> SContextualAnimViewport::BuildViewportToolbar()
 {
 	return SAssignNew(ViewportToolbar, SContextualAnimViewportToolBar, SharedThis(this));
 }
diff --git a/Engine/Plugins/Experimental/Animation/ContextualAnimation/Source/ContextualAnimationEditor/Public/SContextualAnimViewport.h b/Engine/Plugins/Experimental/Animation/ContextualAnimation/Source/ContextualAnimationEditor/Public/SContextualAnimViewport.h
--- a/Engine/Plugins/Experimental/Animation/ContextualAnimation/Source/ContextualAnimationEditor/Public/SContextualAnimViewport.h
+++ b/Engine/Plugins/Experimental/Animation/ContextualAnimation/Source/ContextualAnimationEditor/Public/SContextualAnimViewport.h
@@ -43,7 +43,7 @@ protected:
 	// ~SEditorViewport interface
 	virtual void BindCommands() override;
 	virtual TSharedRef<FEditorViewportClient> MakeEditorViewportClient() override;
-	virtual TSharedPtr<SWidget> MakeViewportToolbar() override;
+	virtual TSharedPtr<SWidget> BuildViewportToolbar() override;
 	virtual const FSlateBrush* OnGetViewportBorderBrush() const override;
 	virtual FSlateColor OnGetViewportBorderColorAndOpacity() const;
 	// ~End of SEditorViewport interface
diff --git a/Engine/Plugins/Experimental/Dataflow/Source/DataflowEditor/Private/Dataflow/DataflowEditorViewportToolBarBase.cpp b/Engine/Plugins/Experimental/Dataflow/Source/DataflowEditor/Private/Dataflow/DataflowEditorViewportToolBarBase.cpp
deleted file mode 100644
diff --git a/Engine/Plugins/Experimental/Dataflow/Source/DataflowEditor/Private/Dataflow/DataflowEditorViewportToolBarBase.h b/Engine/Plugins/Experimental/Dataflow/Source/DataflowEditor/Private/Dataflow/DataflowEditorViewportToolBarBase.h
deleted file mode 100644
diff --git a/Engine/Plugins/Experimental/ImageWidgets/Source/ImageWidgets/Private/SImageViewport.cpp b/Engine/Plugins/Experimental/ImageWidgets/Source/ImageWidgets/Private/SImageViewport.cpp
--- a/Engine/Plugins/Experimental/ImageWidgets/Source/ImageWidgets/Private/SImageViewport.cpp
+++ b/Engine/Plugins/Experimental/ImageWidgets/Source/ImageWidgets/Private/SImageViewport.cpp
@@ -248,7 +248,7 @@ namespace UE::ImageWidgets
 		return Client.ToSharedRef();
 	}
 
-	TSharedPtr<SWidget> SImageViewport::MakeViewportToolbar()
+	TSharedRef<SWidget> SImageViewport::MakeViewportToolbarOverlay()
 	{
 		auto HasImage = [this]
 		{
@@ -297,11 +297,17 @@ namespace UE::ImageWidgets
 
 		ToolbarExtender.Reset();
 
-		return ImageViewportToolbar;
+		return ImageViewportToolbar.ToSharedRef();
 	}
 
 	void SImageViewport::PopulateViewportOverlays(TSharedRef<SOverlay> Overlay)
 	{
+		Overlay->AddSlot()
+			.VAlign(VAlign_Top)
+			[
+				MakeViewportToolbarOverlay()
+			];
+	
 		Overlay->AddSlot()
 		       .VAlign(VAlign_Bottom)
 		[
diff --git a/Engine/Plugins/Experimental/ImageWidgets/Source/ImageWidgets/Public/SImageViewport.h b/Engine/Plugins/Experimental/ImageWidgets/Source/ImageWidgets/Public/SImageViewport.h
--- a/Engine/Plugins/Experimental/ImageWidgets/Source/ImageWidgets/Public/SImageViewport.h
+++ b/Engine/Plugins/Experimental/ImageWidgets/Source/ImageWidgets/Public/SImageViewport.h
@@ -234,11 +234,12 @@ namespace UE::ImageWidgets
 		virtual bool IsVisible() const override { return true; } // Overriding this method avoids the 0.25 seconds update delay in SEditorViewport.
 	protected:
 		virtual TSharedRef<FEditorViewportClient> MakeEditorViewportClient() override;
-		virtual TSharedPtr<SWidget> MakeViewportToolbar() override;
 		virtual void PopulateViewportOverlays(TSharedRef<SOverlay> Overlay) override;
 		// SEditorViewport overrides - end
 
 	private:
+		TSharedRef<SWidget> MakeViewportToolbarOverlay();
+	
 		/** Create the status bar widgets. */
 		TSharedRef<SWidget> MakeStatusBar(TSharedRef<SOverlay> Overlay);
 
diff --git a/Engine/Plugins/Experimental/MutablePopulation/Source/CustomizableObjectPopulationEditor/Private/MuCOPE/SCustomizableObjectPopulationEditorViewport.cpp b/Engine/Plugins/Experimental/MutablePopulation/Source/CustomizableObjectPopulationEditor/Private/MuCOPE/SCustomizableObjectPopulationEditorViewport.cpp
--- a/Engine/Plugins/Experimental/MutablePopulation/Source/CustomizableObjectPopulationEditor/Private/MuCOPE/SCustomizableObjectPopulationEditorViewport.cpp
+++ b/Engine/Plugins/Experimental/MutablePopulation/Source/CustomizableObjectPopulationEditor/Private/MuCOPE/SCustomizableObjectPopulationEditorViewport.cpp
@@ -112,13 +112,16 @@ void SCustomizableObjectPopulationEditorViewport::AddReferencedObjects(FReferenc
 	}
 }
 
-TSharedPtr<SWidget> SCustomizableObjectPopulationEditorViewport::MakeViewportToolbar()
+void SCustomizableObjectPopulationEditorViewport::PopulateViewportOverlays(TSharedRef<SOverlay> Overlay)
 {
-	return SNew(SCustomizableObjectPopulationEditorViewportToolBar,/* TabBodyPtr.Pin(),*/ SharedThis(this))
-		.Cursor(EMouseCursor::Default);
+	Overlay->AddSlot()
+		.VAlign(VAlign_Top)
+		[
+			SNew(SCustomizableObjectPopulationEditorViewportToolBar,/* TabBodyPtr.Pin(),*/ SharedThis(this))
+			.Cursor(EMouseCursor::Default)
+		];
 }
 
-
 int32 SCustomizableObjectPopulationEditorViewport::GetSelectedInstance() 
 { 
 	if (EditorViewportClient.IsValid())
diff --git a/Engine/Plugins/Experimental/MutablePopulation/Source/CustomizableObjectPopulationEditor/Private/MuCOPE/SCustomizableObjectPopulationEditorViewport.h b/Engine/Plugins/Experimental/MutablePopulation/Source/CustomizableObjectPopulationEditor/Private/MuCOPE/SCustomizableObjectPopulationEditorViewport.h
--- a/Engine/Plugins/Experimental/MutablePopulation/Source/CustomizableObjectPopulationEditor/Private/MuCOPE/SCustomizableObjectPopulationEditorViewport.h
+++ b/Engine/Plugins/Experimental/MutablePopulation/Source/CustomizableObjectPopulationEditor/Private/MuCOPE/SCustomizableObjectPopulationEditorViewport.h
@@ -59,5 +59,5 @@ private:
 	/** Current Collider Components */
 	TArray<UCapsuleComponent*> PreviewCollisionComponents;
 
-	virtual TSharedPtr<SWidget> MakeViewportToolbar() override;
+	virtual void PopulateViewportOverlays(TSharedRef<SOverlay> Overlay) override;
 };
diff --git a/Engine/Plugins/FX/Cascade/Source/Cascade/Private/SCascadePreviewViewport.cpp b/Engine/Plugins/FX/Cascade/Source/Cascade/Private/SCascadePreviewViewport.cpp
--- a/Engine/Plugins/FX/Cascade/Source/Cascade/Private/SCascadePreviewViewport.cpp
+++ b/Engine/Plugins/FX/Cascade/Source/Cascade/Private/SCascadePreviewViewport.cpp
@@ -116,18 +116,17 @@ TSharedRef<FEditorViewportClient> SCascadePreviewViewport::MakeEditorViewportCli
 	return ViewportClient.ToSharedRef();
 }
 
-TSharedPtr<SWidget> SCascadePreviewViewport::MakeViewportToolbar()
-{
-	return
-	SNew(SCascadePreviewViewportToolBar)
-		.CascadePtr(CascadePtr)
-		.IsEnabled(FSlateApplication::Get().GetNormalExecutionAttribute());
-}
-
-
 void SCascadePreviewViewport::PopulateViewportOverlays(TSharedRef<SOverlay> Overlay)
 {
 	SEditorViewport::PopulateViewportOverlays(Overlay);
+	
+	Overlay->AddSlot()
+		.VAlign(VAlign_Top)
+		[
+			SNew(SCascadePreviewViewportToolBar)
+			.CascadePtr(CascadePtr)
+			.IsEnabled(FSlateApplication::Get().GetNormalExecutionAttribute())
+		];
 
 	// this widget will display the current viewed feature level
 	Overlay->AddSlot()
diff --git a/Engine/Plugins/FX/Cascade/Source/Cascade/Private/SCascadePreviewViewport.h b/Engine/Plugins/FX/Cascade/Source/Cascade/Private/SCascadePreviewViewport.h
--- a/Engine/Plugins/FX/Cascade/Source/Cascade/Private/SCascadePreviewViewport.h
+++ b/Engine/Plugins/FX/Cascade/Source/Cascade/Private/SCascadePreviewViewport.h
@@ -53,7 +53,6 @@ public:
 	void ClearTickFlag()	{ JustTicked = false; }
 protected:
 	virtual TSharedRef<FEditorViewportClient> MakeEditorViewportClient() override;
-	virtual TSharedPtr<SWidget> MakeViewportToolbar() override;
 	virtual void PopulateViewportOverlays(TSharedRef<SOverlay> Overlay) override;
 	virtual void OnFocusViewportToSelection() override;
 public:
diff --git a/Engine/Plugins/FX/Niagara/Source/NiagaraEditor/Private/Widgets/SNiagaraSimCacheViewport.cpp b/Engine/Plugins/FX/Niagara/Source/NiagaraEditor/Private/Widgets/SNiagaraSimCacheViewport.cpp
--- a/Engine/Plugins/FX/Niagara/Source/NiagaraEditor/Private/Widgets/SNiagaraSimCacheViewport.cpp
+++ b/Engine/Plugins/FX/Niagara/Source/NiagaraEditor/Private/Widgets/SNiagaraSimCacheViewport.cpp
@@ -172,7 +172,7 @@ TSharedRef<FEditorViewportClient> SNiagaraSimCacheViewport::MakeEditorViewportCl
 	return SimCacheViewportClient.ToSharedRef();
 }
 
-TSharedPtr<SWidget> SNiagaraSimCacheViewport::MakeViewportToolbar()
+TSharedPtr<SWidget> SNiagaraSimCacheViewport::BuildViewportToolbar()
 {
 	return SNew(SNiagaraSimCacheViewportToolbar, SharedThis(this));
 }
diff --git a/Engine/Plugins/FX/Niagara/Source/NiagaraEditor/Private/Widgets/SNiagaraSimCacheViewport.h b/Engine/Plugins/FX/Niagara/Source/NiagaraEditor/Private/Widgets/SNiagaraSimCacheViewport.h
--- a/Engine/Plugins/FX/Niagara/Source/NiagaraEditor/Private/Widgets/SNiagaraSimCacheViewport.h
+++ b/Engine/Plugins/FX/Niagara/Source/NiagaraEditor/Private/Widgets/SNiagaraSimCacheViewport.h
@@ -41,7 +41,7 @@ public:
 protected:
 	/** SEditorViewport interface */
 	virtual TSharedRef<FEditorViewportClient> MakeEditorViewportClient() override;
-	virtual TSharedPtr<SWidget> MakeViewportToolbar() override;
+	virtual TSharedPtr<SWidget> BuildViewportToolbar() override;
 	virtual EVisibility OnGetViewportContentVisibility() const override;
 	void ToggleOrbit();
 	bool IsToggleOrbitChecked();
diff --git a/Engine/Plugins/Runtime/nDisplay/Source/DisplayClusterConfigurator/Private/Views/Viewport/DisplayClusterConfiguratorSCSEditorViewport.cpp b/Engine/Plugins/Runtime/nDisplay/Source/DisplayClusterConfigurator/Private/Views/Viewport/DisplayClusterConfiguratorSCSEditorViewport.cpp
--- a/Engine/Plugins/Runtime/nDisplay/Source/DisplayClusterConfigurator/Private/Views/Viewport/DisplayClusterConfiguratorSCSEditorViewport.cpp
+++ b/Engine/Plugins/Runtime/nDisplay/Source/DisplayClusterConfigurator/Private/Views/Viewport/DisplayClusterConfiguratorSCSEditorViewport.cpp
@@ -579,18 +579,18 @@ TSharedRef<FEditorViewportClient> SDisplayClusterConfiguratorSCSEditorViewport::
 	return ViewportClient.ToSharedRef();
 }
 
-TSharedPtr<SWidget> SDisplayClusterConfiguratorSCSEditorViewport::MakeViewportToolbar()
-{
-	return
-		SNew(SDisplayClusterConfiguratorSCSEditorViewportToolBar, ViewportClient)
-		.EditorViewport(SharedThis(this))
-		.IsEnabled(FSlateApplication::Get().GetNormalExecutionAttribute());
-}
-
 void SDisplayClusterConfiguratorSCSEditorViewport::PopulateViewportOverlays(TSharedRef<SOverlay> Overlay)
 {
 	SEditorViewport::PopulateViewportOverlays(Overlay);
 
+	Overlay->AddSlot()
+		.VAlign(VAlign_Top)
+		[
+			SNew(SDisplayClusterConfiguratorSCSEditorViewportToolBar, ViewportClient)
+			.EditorViewport(SharedThis(this))
+			.IsEnabled(FSlateApplication::Get().GetNormalExecutionAttribute())
+		];
+
 	// add the feature level display widget
 	Overlay->AddSlot()
 		.VAlign(VAlign_Bottom)
diff --git a/Engine/Plugins/Runtime/nDisplay/Source/DisplayClusterConfigurator/Private/Views/Viewport/DisplayClusterConfiguratorSCSEditorViewport.h b/Engine/Plugins/Runtime/nDisplay/Source/DisplayClusterConfigurator/Private/Views/Viewport/DisplayClusterConfiguratorSCSEditorViewport.h
--- a/Engine/Plugins/Runtime/nDisplay/Source/DisplayClusterConfigurator/Private/Views/Viewport/DisplayClusterConfiguratorSCSEditorViewport.h
+++ b/Engine/Plugins/Runtime/nDisplay/Source/DisplayClusterConfigurator/Private/Views/Viewport/DisplayClusterConfiguratorSCSEditorViewport.h
@@ -44,7 +44,6 @@ class SDisplayClusterConfiguratorSCSEditorViewport : public SAssetEditorViewport
 protected:
 	/** SEditorViewport interface */
 	virtual TSharedRef<FEditorViewportClient> MakeEditorViewportClient() override;
-	virtual TSharedPtr<SWidget> MakeViewportToolbar() override;
 	virtual void PopulateViewportOverlays(TSharedRef<SOverlay> Overlay) override;
 	virtual void BindCommands() override;
 	virtual void OnDragLeave(const FDragDropEvent& DragDropEvent) override;
diff --git a/Engine/Plugins/VirtualProduction/Avalanche/Source/AvalancheLevelViewport/Private/AvaLevelViewportCommands.cpp b/Engine/Plugins/VirtualProduction/Avalanche/Source/AvalancheLevelViewport/Private/AvaLevelViewportCommands.cpp
--- a/Engine/Plugins/VirtualProduction/Avalanche/Source/AvalancheLevelViewport/Private/AvaLevelViewportCommands.cpp
+++ b/Engine/Plugins/VirtualProduction/Avalanche/Source/AvalancheLevelViewport/Private/AvaLevelViewportCommands.cpp
@@ -69,7 +69,7 @@ void FAvaLevelViewportCommands::RegisterViewportCommands()
 		, FInputChord(EKeys::F, EModifierKey::Alt))
 
 	UI_COMMAND(ToggleChildActorLock
-		, "Toggle Child Actor Lock"
+		, "Child Actor Lock"
 		, "Keeps child actors in place when moving parent actors."
 		, EUserInterfaceActionType::ToggleButton
 		, FInputChord(EKeys::P, EModifierKey::Alt | EModifierKey::Shift))
diff --git a/Engine/Plugins/VirtualProduction/Avalanche/Source/AvalancheLevelViewport/Private/SAvaLevelViewport.cpp b/Engine/Plugins/VirtualProduction/Avalanche/Source/AvalancheLevelViewport/Private/SAvaLevelViewport.cpp
--- a/Engine/Plugins/VirtualProduction/Avalanche/Source/AvalancheLevelViewport/Private/SAvaLevelViewport.cpp
+++ b/Engine/Plugins/VirtualProduction/Avalanche/Source/AvalancheLevelViewport/Private/SAvaLevelViewport.cpp
@@ -46,14 +46,10 @@ void SAvaLevelViewport::PrivateRegisterAttributes(struct FSlateAttributeDescript
 SAvaLevelViewport::~SAvaLevelViewport()
 {
 	GetMutableDefault<UAvaViewportSettings>()->OnChange.RemoveAll(this);
-
-	UnregisterPanelExtension();
 }
 
 void SAvaLevelViewport::Construct(const FArguments& InArgs, const FAssetEditorViewportConstructionArgs& InViewportArgs)
 {
-	RegisterPanelExtension();
-
 	ViewportFrameWeak = InArgs._ViewportFrame;
 	VirtualSizeAspectRatio = 0.f;
 	VirtualSizeAspectRatioState = EAvaViewportVirtualSizeAspectRatioState::LockedToCamera;
@@ -116,7 +112,9 @@ void SAvaLevelViewport::BindCommands()
 
 	CommandListRef.MapAction(
 		AvaLevelViewportCommands.ToggleChildActorLock,
-		FExecuteAction::CreateSP(this, &SAvaLevelViewport::ExecuteToggleChildActorLock)
+		FExecuteAction::CreateSP(this, &SAvaLevelViewport::ExecuteToggleChildActorLock),
+		FCanExecuteAction(),
+		FIsActionChecked::CreateSP(this, &SAvaLevelViewport::IsChildActorLockEnabled)
 	);
 
 	CommandListRef.MapAction(
@@ -912,49 +910,6 @@ void SAvaLevelViewport::SetVirtualSizeAspectRatioState(EAvaViewportVirtualSizeAs
 	Data->VirtualSizeAspectRatioState = InAspectRatioState;
 }
 
-void SAvaLevelViewport::RegisterPanelExtension()
-{
-	UPanelExtensionSubsystem* PanelExtensionSubsystem = GEditor->GetEditorSubsystem<UPanelExtensionSubsystem>();
-
-	if (!PanelExtensionSubsystem)
-	{
-		return;
-	}
-
-	// Use pointer for unique extension factory because there's no way to access this instance from the factory...
-	PanelExtensionFactory.Identifier = FName(*FString::Printf(TEXT("AvaLevelViewportChildActorLockToggle_%p)"), this));
-
-	if (PanelExtensionSubsystem->IsPanelFactoryRegistered(PanelExtensionFactory.Identifier))
-	{
-		return;
-	}
-
-	PanelExtensionFactory.CreateExtensionWidget = FPanelExtensionFactory::FCreateExtensionWidget::CreateSP(this,
-		&SAvaLevelViewport::OnExtendLevelEditorViewportToolbarForChildActorLock);
-
-	PanelExtensionSubsystem->RegisterPanelFactory(
-		TEXT("LevelViewportToolBar.RightExtension"),
-		PanelExtensionFactory
-	);
-}
-
-void SAvaLevelViewport::UnregisterPanelExtension()
-{
-	UPanelExtensionSubsystem* PanelExtensionSubsystem = GEditor->GetEditorSubsystem<UPanelExtensionSubsystem>();
-
-	if (!PanelExtensionSubsystem)
-	{
-		return;
-	}
-
-	if (!PanelExtensionSubsystem->IsPanelFactoryRegistered(PanelExtensionFactory.Identifier))
-	{
-		return;
-	}
-
-	PanelExtensionSubsystem->UnregisterPanelFactory(PanelExtensionFactory.Identifier, TEXT("LevelViewportToolBar.RightExtension"));
-}
-
 ECheckBoxState SAvaLevelViewport::GetTextureOverlayStretchEnabledCheckBoxState() const
 {
 	if (const UAvaViewportSettings* AvaViewportSettings = GetDefault<UAvaViewportSettings>())
diff --git a/Engine/Plugins/VirtualProduction/Avalanche/Source/AvalancheLevelViewport/Private/SAvaLevelViewport.h b/Engine/Plugins/VirtualProduction/Avalanche/Source/AvalancheLevelViewport/Private/SAvaLevelViewport.h
--- a/Engine/Plugins/VirtualProduction/Avalanche/Source/AvalancheLevelViewport/Private/SAvaLevelViewport.h
+++ b/Engine/Plugins/VirtualProduction/Avalanche/Source/AvalancheLevelViewport/Private/SAvaLevelViewport.h
@@ -60,7 +60,7 @@ public:
 
 	//~ Begin SEditorViewport
 	virtual void BindCommands() override;
-	virtual TSharedPtr<SWidget> MakeViewportToolbar() override;
+	virtual FName RegisterViewportToolbar() const override;
 	virtual void PopulateViewportOverlays(TSharedRef<SOverlay> Overlay) override;
 	//~ End SEditorViewport
 
@@ -170,7 +170,6 @@ private:
 	TSharedPtr<SCanvas> GuideCanvas;
 	TArray<TSharedPtr<SAvaLevelViewportGuide>> Guides;
 	TSharedPtr<SAvaLevelViewportSnapIndicators> SnapIndicators;
-	FPanelExtensionFactory PanelExtensionFactory;
 	TSharedPtr<SWidget> BackgroundTextureSelector;
 	TSharedPtr<SWidget> PostProcessOpacitySlider;
 	TSharedPtr<FScopedTransaction> PostProcessInfoTransaction;
@@ -200,19 +199,14 @@ private:
 
 	void AddCameraZoomMenuEntries(UToolMenu* InMenu);
 
-	void AddVisualizerEntries(UToolMenu* InMenu);
+	static void AddVisualizerEntries(UToolMenu* InMenu);
 
 	void OnSettingsChanged(const UAvaViewportSettings* InSettings, FName InPropertyChanged);
 	void ApplySettings(const UAvaViewportSettings* InSettings);
 
 	void LoadVirtualSizeSettings();
 
-	void RegisterPanelExtension();
-	void UnregisterPanelExtension();
-	TSharedRef<SWidget> OnExtendLevelEditorViewportToolbarForChildActorLock(FWeakObjectPtr InExtensionContext);
-
 	bool IsChildActorLockEnabled() const;
-	FReply OnChildActorLockButtonClicked();
 
 	void BeginPostProcessInfoTransaction();
 	void EndPostProcessInfoTransaction();
diff --git a/Engine/Plugins/VirtualProduction/Avalanche/Source/AvalancheLevelViewport/Private/SAvaLevelViewport_Menus.cpp b/Engine/Plugins/VirtualProduction/Avalanche/Source/AvalancheLevelViewport/Private/SAvaLevelViewport_Menus.cpp
--- a/Engine/Plugins/VirtualProduction/Avalanche/Source/AvalancheLevelViewport/Private/SAvaLevelViewport_Menus.cpp
+++ b/Engine/Plugins/VirtualProduction/Avalanche/Source/AvalancheLevelViewport/Private/SAvaLevelViewport_Menus.cpp
@@ -21,6 +21,7 @@
 #include "SortHelper.h"
 #include "SViewportToolBar.h"
 #include "ToolMenus.h"
+#include "Editor/LevelEditor/Private/ViewportToolbar/LevelViewportContext.h"
 #include "ViewportClient/AvaLevelViewportClient.h"
 #include "Widgets/DataTypes/AvaUserInputDialogDataTypeText.h"
 #include "Widgets/Images/SImage.h"
@@ -64,61 +65,65 @@ namespace UE::AvaLevelViewport::Private
 	}
 }
 
-TSharedPtr<SWidget> SAvaLevelViewport::MakeViewportToolbar()
+FName SAvaLevelViewport::RegisterViewportToolbar() const
 {
-	const TSharedPtr<SWidget> ToolbarFrame = Super::MakeViewportToolbar();
-
-	if (ToolbarFrame.IsValid())
 	{
-		// Replace the Camera Menu with a custom one stripping off the View Entries (Perspective, Front, Left, etc...)
-		static const FName CameraMenuName = TEXT("EditorViewportToolBar.CameraMenu");
-		if (FSlotBase* const FoundSlot = UE::AvaLevelViewport::Private::FindChildSlotWithTag(CameraMenuName, ToolbarFrame.ToSharedRef()))
+		// Extend the settings menu with custom visualizer entries
+		const FName ParentSettingsName = "LevelEditor.ViewportToolbar.Settings";
+		const FName SettingsName = "Ava.ViewportToolbar.Settings";
+		
+		if (!UToolMenus::Get()->FindMenu(ParentSettingsName))
 		{
-			TSharedPtr<SViewportToolBar> ParentToolBar;
-			{
-				TSharedPtr<SEditorViewportToolbarMenu> CameraToolbarMenu = StaticCastSharedRef<SEditorViewportToolbarMenu>(FoundSlot->GetWidget());
-				ParentToolBar = CameraToolbarMenu->GetParentToolBar().Pin();
-			}
-
-			FoundSlot->AttachWidget(
-				SNew(SEditorViewportToolbarMenu)
-				.ParentToolBar(ParentToolBar)
-				.Label(LOCTEXT("CameraMenuLabel", "Camera"))
-				.LabelIcon(GetCameraIcon().GetIcon())
-				.AddMetaData<FTagMetaData>(FTagMetaData(TEXT("EditorViewportToolBar.CameraMenu")))
-				.OnGetMenuContent(this, &SAvaLevelViewport::GenerateCameraMenu)
-			);
+			UToolMenus::Get()->ExtendMenu(ParentSettingsName);
 		}
+		
+		if (!UToolMenus::Get()->IsMenuRegistered(SettingsName))
+		{
+			AddVisualizerEntries(UToolMenus::Get()->RegisterMenu(SettingsName, ParentSettingsName));
+		}	
 	}
-
-	static const FName OptionsMenuName("LevelEditor.LevelViewportToolBar.Options");
-
-	if (UToolMenus* const ToolMenus = UToolMenus::Get())
+	
+	const FName ViewportToolbarName = "Ava.ViewportToolbar";
+	
+	if (!UToolMenus::Get()->IsMenuRegistered(ViewportToolbarName))
 	{
-		if (UToolMenu* const OptionsMenu = ToolMenus->ExtendMenu(OptionsMenuName))
+		UToolMenu* Menu = UToolMenus::Get()->RegisterMenu(ViewportToolbarName, Super::RegisterViewportToolbar(), EMultiBoxType::SlimHorizontalToolBar);
+		Menu->StyleName = "ViewportToolbar";
+		Menu->bSeparateSections = false;
+		
 		{
-			AddVisualizerEntries(OptionsMenu);
+			FToolMenuSection& LeftSection = Menu->AddSection("Left");
+			FToolMenuEntry& ChildActorLock = LeftSection.AddMenuEntry(FAvaLevelViewportCommands::GetInternal().ToggleChildActorLock);
+			ChildActorLock.InsertPosition.Position = EToolMenuInsertType::Last;
+			ChildActorLock.Icon = FSlateIcon(FAvaLevelViewportStyle::Get().GetStyleSetName(), "Button.PivotMode");
+			ChildActorLock.ToolBarData.LabelOverride = FText::GetEmpty();
 		}
+		
+		{
+			FToolMenuSection& RightSection = Menu->AddSection("Right");
+			RightSection.Alignment = EToolMenuSectionAlign::Last;
+		
+			// Override the "Camera" submenu
+			RightSection.AddDynamicEntry("DynamicCamera", FNewToolMenuSectionDelegate::CreateLambda([](FToolMenuSection& Section)
+			{
+				TSharedPtr<SAvaLevelViewport> Viewport = StaticCastSharedPtr<SAvaLevelViewport>(ULevelViewportContext::GetLevelViewport(Section));
+				if (!Viewport)
+				{
+					return;
+				}
+			
+				FToolMenuEntry& CameraMenu = Section.AddSubMenu(
+					"Camera",
+					LOCTEXT("CameraMenuLabel", "Camera"),
+					FText::GetEmpty(),
+					FOnGetContent::CreateSP(Viewport.Get(), &SAvaLevelViewport::GenerateCameraMenu)
+				);
+				CameraMenu.Icon.Bind(Viewport.Get(), &SAvaLevelViewport::GetCameraIcon);
+			}));
+		}	
 	}
-
-	return ToolbarFrame;
-}
-
-TSharedRef<SWidget> SAvaLevelViewport::OnExtendLevelEditorViewportToolbarForChildActorLock(FWeakObjectPtr InExtensionContext)
-{
-	const FAvaLevelViewportCommands& Commands = FAvaLevelViewportCommands::GetInternal();
-
-	return SNew(SEditorViewportToolBarButton)
-		.ButtonType(EUserInterfaceActionType::Check)
-		.IsChecked(this, &SAvaLevelViewport::IsChildActorLockEnabled)
-		.OnClicked(this, &SAvaLevelViewport::OnChildActorLockButtonClicked)
-		.ToolTipText(Commands.ToggleChildActorLock->GetDescription())
-		.Content()
-		[
-			SNew(SImage)
-			.DesiredSizeOverride(FVector2D(16, 16))
-			.Image(FAvaLevelViewportStyle::Get().GetBrush(TEXT("Button.PivotMode")))
-		];
+	
+	return ViewportToolbarName;
 }
 
 FSlateIcon SAvaLevelViewport::GetCameraIcon() const
@@ -770,12 +775,7 @@ void SAvaLevelViewport::AddCameraZoomMenuEntries(UToolMenu* InMenu)
 
 void SAvaLevelViewport::AddVisualizerEntries(UToolMenu* InMenu)
 {
-	FToolMenuSection* Section = InMenu->FindSection("LevelViewportViewportOptions2");
-
-	if (!Section)
-	{
-		return;
-	}
+	FToolMenuSection& Section = InMenu->FindOrAddSection("Settings");
 
 	TSharedRef<SWidget> IconSizeWidget = SNew(SBox)
 		.HAlign(HAlign_Right)
@@ -847,7 +847,7 @@ void SAvaLevelViewport::AddVisualizerEntries(UToolMenu* InMenu)
 			]
 		];
 
-	Section->AddEntry(FToolMenuEntry::InitWidget(
+	Section.AddEntry(FToolMenuEntry::InitWidget(
 		"IconSizeWidget",
 		IconSizeWidget,
 		LOCTEXT("VisualizerIconSize", "Visualizer Icon Size")
@@ -866,13 +866,6 @@ bool SAvaLevelViewport::IsChildActorLockEnabled() const
 	return ViewportClient->AreChildActorsLocked();
 }
 
-FReply SAvaLevelViewport::OnChildActorLockButtonClicked()
-{
-	ExecuteToggleChildActorLock();
-
-	return FReply::Handled();
-}
-
 void SAvaLevelViewport::OnVirtualSizeComponentSliderCommitted(int32 InNewDimension, ETextCommit::Type InCommitType, EAxis::Type InAxis)
 {
 	using namespace UE::AvaLevelViewport::Private;
diff --git a/Engine/Plugins/VirtualProduction/DynamicMaterial/Source/DynamicMaterialEditor/Private/UI/Widgets/Editor/SDMMaterialPreview.cpp b/Engine/Plugins/VirtualProduction/DynamicMaterial/Source/DynamicMaterialEditor/Private/UI/Widgets/Editor/SDMMaterialPreview.cpp
--- a/Engine/Plugins/VirtualProduction/DynamicMaterial/Source/DynamicMaterialEditor/Private/UI/Widgets/Editor/SDMMaterialPreview.cpp
+++ b/Engine/Plugins/VirtualProduction/DynamicMaterial/Source/DynamicMaterialEditor/Private/UI/Widgets/Editor/SDMMaterialPreview.cpp
@@ -122,7 +122,16 @@ void SDMMaterialPreview::Construct(const FArguments& InArgs, const TSharedRef<SD
 	}
 }
 
-TSharedPtr<SWidget> SDMMaterialPreview::MakeViewportToolbar()
+void SDMMaterialPreview::PopulateViewportOverlays(TSharedRef<SOverlay> Overlay)
+{
+	Overlay->AddSlot()
+		.VAlign(VAlign_Top)
+		[
+			MakeViewportToolbarOverlay()
+		];
+}
+
+TSharedRef<SWidget> SDMMaterialPreview::MakeViewportToolbarOverlay()
 {
 	if (!bShowMenu)
 	{
diff --git a/Engine/Plugins/VirtualProduction/DynamicMaterial/Source/DynamicMaterialEditor/Private/UI/Widgets/Editor/SDMMaterialPreview.h b/Engine/Plugins/VirtualProduction/DynamicMaterial/Source/DynamicMaterialEditor/Private/UI/Widgets/Editor/SDMMaterialPreview.h
--- a/Engine/Plugins/VirtualProduction/DynamicMaterial/Source/DynamicMaterialEditor/Private/UI/Widgets/Editor/SDMMaterialPreview.h
+++ b/Engine/Plugins/VirtualProduction/DynamicMaterial/Source/DynamicMaterialEditor/Private/UI/Widgets/Editor/SDMMaterialPreview.h
@@ -65,7 +65,7 @@ public:
 		UDynamicMaterialModelBase* InMaterialModelBase);
 
 	//~ Begin SEditorViewport
-	virtual TSharedPtr<SWidget> MakeViewportToolbar() override;
+	virtual void PopulateViewportOverlays(TSharedRef<SOverlay> Overlay) override;
 	//~ End SEditorViewport
 
 	//~ Begin FGCObject
@@ -85,6 +85,7 @@ protected:
 	TObjectPtr<UMaterialInterface> PreviewMaterial;
 	TObjectPtr<APostProcessVolume> PostProcessVolumeActor;
 
+	TSharedRef<SWidget> MakeViewportToolbarOverlay();
 	void RefreshViewport();
 
 	void SetPreviewType(EDMMaterialPreviewMesh InPrimitiveType);
diff --git a/Engine/Source/Editor/LevelEditor/Private/SLevelViewport.cpp b/Engine/Source/Editor/LevelEditor/Private/SLevelViewport.cpp
--- a/Engine/Source/Editor/LevelEditor/Private/SLevelViewport.cpp
+++ b/Engine/Source/Editor/LevelEditor/Private/SLevelViewport.cpp
@@ -2236,176 +2236,178 @@ TSharedRef<FEditorViewportClient> SLevelViewport::MakeEditorViewportClient()
 	return LevelViewportClient.ToSharedRef();
 }
 
-TSharedPtr<SWidget> SLevelViewport::BuildViewportToolbar()
+FName SLevelViewport::RegisterViewportToolbar() const
 {
 	const FName ViewportToolbarName = "LevelEditor.ViewportToolbar";
-
-	// Register the viewport toolbar if another viewport hasn't already (it's shared).
+	
+	if (!UToolMenus::Get()->IsMenuRegistered(ViewportToolbarName))
 	{
-		if (!UToolMenus::Get()->IsMenuRegistered(ViewportToolbarName))
-		{
-			UToolMenu* const ViewportToolbarMenu = UToolMenus::Get()->RegisterMenu(
-				ViewportToolbarName, NAME_None /* parent */, EMultiBoxType::SlimHorizontalToolBar
-			);
+		UToolMenu* const ViewportToolbarMenu = UToolMenus::Get()->RegisterMenu(
+			ViewportToolbarName, NAME_None /* parent */, EMultiBoxType::SlimHorizontalToolBar
+		);
 
-			ViewportToolbarMenu->StyleName = "ViewportToolbar";
-			ViewportToolbarMenu->bSeparateSections = false;
+		ViewportToolbarMenu->StyleName = "ViewportToolbar";
+		ViewportToolbarMenu->bSeparateSections = false;
 
-			// Add the left-aligned part of the viewport toolbar.
-			{
-				FToolMenuSection& LeftSection = ViewportToolbarMenu->AddSection("Left");
+		// Add the left-aligned part of the viewport toolbar.
+		{
+			FToolMenuSection& LeftSection = ViewportToolbarMenu->AddSection("Left");
 
-				LeftSection.AddEntry(UE::UnrealEd::CreateTransformsSubmenu());
-				UE::LevelEditor::ExtendTransformSubmenu("LevelEditor.ViewportToolbar.Transform");
+			LeftSection.AddEntry(UE::UnrealEd::CreateTransformsSubmenu());
+			UE::LevelEditor::ExtendTransformSubmenu("LevelEditor.ViewportToolbar.Transform");
 
-				LeftSection.AddEntry(UE::UnrealEd::CreateSnappingSubmenu());
-				UE::LevelEditor::ExtendSnappingSubmenu("LevelEditor.ViewportToolbar.Snapping");
-			}
+			LeftSection.AddEntry(UE::UnrealEd::CreateSnappingSubmenu());
+			UE::LevelEditor::ExtendSnappingSubmenu("LevelEditor.ViewportToolbar.Snapping");
+		}
 
-			// Add the right-aligned part of the viewport toolbar.
+		// Add the right-aligned part of the viewport toolbar.
+		{
+			FToolMenuSection& RightSection = ViewportToolbarMenu->AddSection("Right");
+			RightSection.Alignment = EToolMenuSectionAlign::Last;
+
+			// Add the "Camera" submenu.
 			{
-				FToolMenuSection& RightSection = ViewportToolbarMenu->AddSection("Right");
-				RightSection.Alignment = EToolMenuSectionAlign::Last;
+				// This matches the name used in SLevelViewportToolBar::GenerateCameraMenu for
+				// the old level editor viewport toolbar camera menu.
+				const FName ParentSubmenuName = "LevelEditor.LevelViewportToolBar.Camera";
 
-				// Add the "Camera" submenu.
+				// Create our parent menu.
+				if (!UToolMenus::Get()->IsMenuRegistered(ParentSubmenuName))
 				{
-					// This matches the name used in SLevelViewportToolBar::GenerateCameraMenu for
-					// the old level editor viewport toolbar camera menu.
-					const FName ParentSubmenuName = "LevelEditor.LevelViewportToolBar.Camera";
-
-					// Create our parent menu.
-					if (!UToolMenus::Get()->IsMenuRegistered(ParentSubmenuName))
-					{
-						UToolMenus::Get()->RegisterMenu(ParentSubmenuName);
-					}
+					UToolMenus::Get()->RegisterMenu(ParentSubmenuName);
+				}
 
-					// Build the menu name our Camera menu will be using so we can extend it and set
-					// the old viewport toolbar camera menu as our parent.
-					const FName SubmenuName = UToolMenus::JoinMenuPaths(ViewportToolbarName, "Camera");
+				// Build the menu name our Camera menu will be using so we can extend it and set
+				// the old viewport toolbar camera menu as our parent.
+				const FName SubmenuName = UToolMenus::JoinMenuPaths(ViewportToolbarName, "Camera");
 
-					// Register our submenu explicitly so we can set the old viewport toolbar's camera menu as our parent.
-					UToolMenus::Get()->RegisterMenu(SubmenuName, ParentSubmenuName);
+				// Register our submenu explicitly so we can set the old viewport toolbar's camera menu as our parent.
+				UToolMenus::Get()->RegisterMenu(SubmenuName, ParentSubmenuName);
 
-					// Extending using both Level Editor specific and UnrealEd generic entries
-					RightSection.AddEntry(UE::LevelEditor::CreateToolbarCameraSubmenu());
-					UE::LevelEditor::ExtendCameraSubmenu(SubmenuName);
-				}
+				// Extending using both Level Editor specific and UnrealEd generic entries
+				RightSection.AddEntry(UE::LevelEditor::CreateToolbarCameraSubmenu());
+				UE::LevelEditor::ExtendCameraSubmenu(SubmenuName);
+			}
 
-				// Add the "View Modes" sub menu.
+			// Add the "View Modes" sub menu.
+			{
+				// Stay backward-compatible with the old viewport toolbar.
 				{
-					// Stay backward-compatible with the old viewport toolbar.
+					// Create our grandparent menu.
+					if (!UToolMenus::Get()->IsMenuRegistered("UnrealEd.ViewportToolbar.View"))
 					{
-						// Create our grandparent menu.
-						if (!UToolMenus::Get()->IsMenuRegistered("UnrealEd.ViewportToolbar.View"))
-						{
-							UToolMenus::Get()->RegisterMenu("UnrealEd.ViewportToolbar.View");
-						}
-
-						// Create our parent menu.
-						if (!UToolMenus::Get()->IsMenuRegistered("LevelEditor.LevelViewportToolbar.View"))
-						{
-							UToolMenus::Get()->RegisterMenu(
-								"LevelEditor.LevelViewportToolbar.View", "UnrealEd.ViewportToolbar.View"
-							);
-						}
+						UToolMenus::Get()->RegisterMenu("UnrealEd.ViewportToolbar.View");
+					}
 
-						// Create our menu.
+					// Create our parent menu.
+					if (!UToolMenus::Get()->IsMenuRegistered("LevelEditor.LevelViewportToolbar.View"))
+					{
 						UToolMenus::Get()->RegisterMenu(
-							"LevelEditor.ViewportToolbar.ViewModes", "LevelEditor.LevelViewportToolbar.View"
+							"LevelEditor.LevelViewportToolbar.View", "UnrealEd.ViewportToolbar.View"
 						);
 					}
 
-					// Add the level editor specific entries.
-					UE::LevelEditor::ExtendViewModesSubmenu("LevelEditor.ViewportToolbar.ViewModes");
-
-					// Create and add the submenu entry to make the menu we just created and extended appear in the
-					// viewport toolbar itself.
-					RightSection.AddEntry(UE::UnrealEd::CreateViewModesSubmenu());
+					// Create our menu.
+					UToolMenus::Get()->RegisterMenu(
+						"LevelEditor.ViewportToolbar.ViewModes", "LevelEditor.LevelViewportToolbar.View"
+					);
 				}
 
-				// Add the "Show" submenu.
+				// Add the level editor specific entries.
+				UE::LevelEditor::ExtendViewModesSubmenu("LevelEditor.ViewportToolbar.ViewModes");
+
+				// Create and add the submenu entry to make the menu we just created and extended appear in the
+				// viewport toolbar itself.
+				RightSection.AddEntry(UE::UnrealEd::CreateViewModesSubmenu());
+			}
+
+			// Add the "Show" submenu.
+			{
+				// Stay backward-compatible with the old viewport toolbar.
 				{
-					// Stay backward-compatible with the old viewport toolbar.
+					if (!UToolMenus::Get()->IsMenuRegistered("LevelEditor.LevelViewportToolbar.Show"))
 					{
-						if (!UToolMenus::Get()->IsMenuRegistered("LevelEditor.LevelViewportToolbar.Show"))
-						{
-							UToolMenus::Get()->RegisterMenu("LevelEditor.LevelViewportToolbar.Show");
-						}
-						UToolMenus::Get()->RegisterMenu(
-							"LevelEditor.ViewportToolbar.Show", "LevelEditor.LevelViewportToolbar.Show"
-						);
+						UToolMenus::Get()->RegisterMenu("LevelEditor.LevelViewportToolbar.Show");
 					}
-
-					RightSection.AddEntry(UE::LevelEditor::CreateShowSubmenu());
+					UToolMenus::Get()->RegisterMenu(
+						"LevelEditor.ViewportToolbar.Show", "LevelEditor.LevelViewportToolbar.Show"
+					);
 				}
 
-				RightSection.AddEntry(UE::LevelEditor::CreatePerformanceAndScalabilitySubmenu());
-				RightSection.AddEntry(UE::LevelEditor::CreateSettingsSubmenu());
+				RightSection.AddEntry(UE::LevelEditor::CreateShowSubmenu());
 			}
 
-			// Sections added for compatibility with old toolbar transform menu extensions
+			RightSection.AddEntry(UE::LevelEditor::CreatePerformanceAndScalabilitySubmenu());
+			RightSection.AddEntry(UE::LevelEditor::CreateSettingsSubmenu());
+		}
+
+		// Sections added for compatibility with old toolbar transform menu extensions
+		{
+			// Existing code might add entries via FToolBarBuilder Extenders,
+			// These extenders are based off of toolbar section & extension names.
+			// The following sections names are used in STransformViewportToolBar::MakeTransformToolBar(...),
+			// This allows legacy extensions to be added directly to the toolbar.
+
+			ViewportToolbarMenu->AddSection("Transform").Alignment = EToolMenuSectionAlign::Last;
+			ViewportToolbarMenu->AddSection("LocationGridSnap").Alignment = EToolMenuSectionAlign::Last;
+			ViewportToolbarMenu->AddSection("RotationGridSnap").Alignment = EToolMenuSectionAlign::Last;
+			ViewportToolbarMenu->AddSection("Layer2DSnap").Alignment = EToolMenuSectionAlign::Last;
+			ViewportToolbarMenu->AddSection("ScaleGridSnap").Alignment = EToolMenuSectionAlign::Last;
+			ViewportToolbarMenu->AddSection("CameraSpeed").Alignment = EToolMenuSectionAlign::Last;
+			
+			// TODO: Remove these once the deprecation period is complete
+			FToolMenuSection& FallbackSection = ViewportToolbarMenu->AddSection("Fallback");
+			FallbackSection.Alignment = EToolMenuSectionAlign::Last;
+			
+			FallbackSection.AddDynamicEntry("FallbackExtensionPanels", FNewToolMenuSectionDelegate::CreateLambda([](FToolMenuSection& Section)
 			{
-				// Existing code might add entries via FToolBarBuilder Extenders,
-				// These extenders are based off of toolbar section & extension names.
-				// The following sections names are used in STransformViewportToolBar::MakeTransformToolBar(...),
-				// This allows legacy extensions to be added directly to the toolbar.
-
-				ViewportToolbarMenu->AddSection("Transform").Alignment = EToolMenuSectionAlign::Last;
-				ViewportToolbarMenu->AddSection("LocationGridSnap").Alignment = EToolMenuSectionAlign::Last;
-				ViewportToolbarMenu->AddSection("RotationGridSnap").Alignment = EToolMenuSectionAlign::Last;
-				ViewportToolbarMenu->AddSection("Layer2DSnap").Alignment = EToolMenuSectionAlign::Last;
-				ViewportToolbarMenu->AddSection("ScaleGridSnap").Alignment = EToolMenuSectionAlign::Last;
-				ViewportToolbarMenu->AddSection("CameraSpeed").Alignment = EToolMenuSectionAlign::Last;
+				ULegacyLevelViewportToolbarContext* Context = Section.FindContext<ULegacyLevelViewportToolbarContext>();
+				if (!Context || !Context->LevelViewport.IsValid() || !Context->LevelViewportToolBarWidget.IsValid())
+				{
+					return;
+				}
 				
-				// TODO: Remove these once the deprecation period is complete
-				FToolMenuSection& FallbackSection = ViewportToolbarMenu->AddSection("Fallback");
-				FallbackSection.Alignment = EToolMenuSectionAlign::Last;
+				UViewportToolBarContext* ExtensionContext = NewObject<UViewportToolBarContext>();
+				ExtensionContext->Viewport = Context->LevelViewport;
+				PRAGMA_DISABLE_DEPRECATION_WARNINGS
+				ExtensionContext->ViewportToolBar = Context->LevelViewportToolBarWidget;
+				PRAGMA_ENABLE_DEPRECATION_WARNINGS
 				
-				FallbackSection.AddDynamicEntry("FallbackExtensionPanels", FNewToolMenuSectionDelegate::CreateLambda([](FToolMenuSection& Section)
+				auto MakeExtensionPanel = [ExtensionContext, &Section](FName ExtensionId)
 				{
-					ULegacyLevelViewportToolbarContext* Context = Section.FindContext<ULegacyLevelViewportToolbarContext>();
-					if (!Context || !Context->LevelViewport.IsValid() || !Context->LevelViewportToolBarWidget.IsValid())
-					{
-						return;
-					}
-					
-					UViewportToolBarContext* ExtensionContext = NewObject<UViewportToolBarContext>();
-					ExtensionContext->Viewport = Context->LevelViewport;
-					PRAGMA_DISABLE_DEPRECATION_WARNINGS
-					ExtensionContext->ViewportToolBar = Context->LevelViewportToolBarWidget;
-					PRAGMA_ENABLE_DEPRECATION_WARNINGS
-					
-					auto MakeExtensionPanel = [ExtensionContext, &Section](FName ExtensionId)
-					{
-						FToolMenuEntry& Entry = Section.AddEntry(FToolMenuEntry::InitWidget(
-							ExtensionId,
-							SNew(SExtensionPanel)
-								.ExtensionPanelID(ExtensionId)
-								.ExtensionContext(ExtensionContext),
-							FText::GetEmpty(),
-							true, false, true
-						));
-						
-						Entry.ToolBarData.ResizeParams.VisibleInOverflow = false;
-					};
+					FToolMenuEntry& Entry = Section.AddEntry(FToolMenuEntry::InitWidget(
+						ExtensionId,
+						SNew(SExtensionPanel)
+							.ExtensionPanelID(ExtensionId)
+							.ExtensionContext(ExtensionContext),
+						FText::GetEmpty(),
+						true, false, true
+					));
 					
-					MakeExtensionPanel("LevelViewportToolBar.LeftExtension");
-					MakeExtensionPanel("LevelViewportToolBar.MiddleExtension");
-					MakeExtensionPanel("LevelViewportToolBar.RightExtension");
-					MakeExtensionPanel("LevelViewportToolBar.RightmostExtension");
-				}));
-			}
-			
-			{
-				// The viewport sizing entries should _always_ be the last section.
-				// Doing otherwise breaks the customized behavior of the overflow menu.
-				FToolMenuSection& SizingSection = ViewportToolbarMenu->AddSection("ViewportSizing");
-				SizingSection.Alignment = EToolMenuSectionAlign::Last;
-				SizingSection.AddEntry(UE::LevelEditor::CreateViewportSizingSubmenu());
-			}
+					Entry.ToolBarData.ResizeParams.VisibleInOverflow = false;
+				};
+				
+				MakeExtensionPanel("LevelViewportToolBar.LeftExtension");
+				MakeExtensionPanel("LevelViewportToolBar.MiddleExtension");
+				MakeExtensionPanel("LevelViewportToolBar.RightExtension");
+				MakeExtensionPanel("LevelViewportToolBar.RightmostExtension");
+			}));
+		}
+		
+		{
+			// The viewport sizing entries should _always_ be the last section.
+			// Doing otherwise breaks the customized behavior of the overflow menu.
+			FToolMenuSection& SizingSection = ViewportToolbarMenu->AddSection("ViewportSizing");
+			SizingSection.Alignment = EToolMenuSectionAlign::Last;
+			SizingSection.AddEntry(UE::LevelEditor::CreateViewportSizingSubmenu());
 		}
 	}
+	
+	return ViewportToolbarName;	
+}
 
+TSharedPtr<SWidget> SLevelViewport::BuildViewportToolbar()
+{
 	FToolMenuContext ViewportToolbarContext;
 	{
 		ViewportToolbarContext.AppendCommandList(GetCommandList());
@@ -2494,7 +2496,7 @@ TSharedPtr<SWidget> SLevelViewport::BuildViewportToolbar()
 	.Visibility(this, &SLevelViewport::GetToolBarVisibility)
 	+ SWidgetSwitcher::Slot()
 	[
-		UToolMenus::Get()->GenerateWidget(ViewportToolbarName, ViewportToolbarContext)
+		UToolMenus::Get()->GenerateWidget(RegisterViewportToolbar(), ViewportToolbarContext)
 	]
 	+ SWidgetSwitcher::Slot()
 	[
diff --git a/Engine/Source/Editor/LevelEditor/Public/SLevelViewport.h b/Engine/Source/Editor/LevelEditor/Public/SLevelViewport.h
--- a/Engine/Source/Editor/LevelEditor/Public/SLevelViewport.h
+++ b/Engine/Source/Editor/LevelEditor/Public/SLevelViewport.h
@@ -505,6 +505,11 @@ public:
 	UE_API bool IsViewportToolbarVisible() const; 
 
 protected:
+	/**
+	 * Registers the viewport toolbar UToolMenu and returns its name.
+	 */
+	UE_API virtual FName RegisterViewportToolbar() const;
+
 	/** SEditorViewport interface */
 	UE_API virtual TSharedRef<FEditorViewportClient> MakeEditorViewportClient() override;
 	UE_API virtual TSharedPtr<SWidget> BuildViewportToolbar() override;
diff --git a/Engine/Source/Editor/UnrealEd/Private/SCommonEditorViewportToolbarBase.cpp b/Engine/Source/Editor/UnrealEd/Private/SCommonEditorViewportToolbarBase.cpp
--- a/Engine/Source/Editor/UnrealEd/Private/SCommonEditorViewportToolbarBase.cpp
+++ b/Engine/Source/Editor/UnrealEd/Private/SCommonEditorViewportToolbarBase.cpp
@@ -526,7 +526,9 @@ void SCommonEditorViewportToolbarBase::Construct(const FArguments& InArgs, TShar
 					SNew(SBox)
 					.Visibility_Lambda([]
 					{
+						PRAGMA_DISABLE_DEPRECATION_WARNINGS
 						return UE::UnrealEd::ShowNewViewportToolbars() ? EVisibility::Visible: EVisibility::Collapsed;
+						PRAGMA_ENABLE_DEPRECATION_WARNINGS
 					})
 					[
 						ToolMenuWidget
@@ -538,7 +540,9 @@ void SCommonEditorViewportToolbarBase::Construct(const FArguments& InArgs, TShar
 					SNew(SBox)
 					.Visibility_Lambda([]
 					{
+						PRAGMA_DISABLE_DEPRECATION_WARNINGS
 						return UE::UnrealEd::ShowOldViewportToolbars() ? EVisibility::Visible: EVisibility::Collapsed;
+						PRAGMA_ENABLE_DEPRECATION_WARNINGS
 					})
 					[
 						MainBox
diff --git a/Engine/Source/Editor/UnrealEd/Private/SEditorViewport.cpp b/Engine/Source/Editor/UnrealEd/Private/SEditorViewport.cpp
--- a/Engine/Source/Editor/UnrealEd/Private/SEditorViewport.cpp
+++ b/Engine/Source/Editor/UnrealEd/Private/SEditorViewport.cpp
@@ -136,24 +136,38 @@ void SEditorViewport::Construct( const FArguments& InArgs )
 
 	// Set up viewport toolbars.
 	{
-		const TSharedPtr<SWidget> OldViewportToolbar = MakeViewportToolbar();
+		// Continue to support switching behavior through the deprecation period to aid licensees in their conversion process. 
+		PRAGMA_DISABLE_DEPRECATION_WARNINGS
+		TSharedPtr<SWidget> OldViewportToolbar = MakeViewportToolbar();
+		PRAGMA_ENABLE_DEPRECATION_WARNINGS
 		TSharedPtr<SWidget> NewViewportToolbar = BuildViewportToolbar();
 		
 		// Allow programatically-migrated viewports (e.g. SCommonEditorViewportToolbarBase) to register themselves on construction as supporting either position
 		// This indirection is done to minimize API changes (e.g. using a specific toolbar interface). The goal is that clients do nothing and get upgrades.
-		if (!NewViewportToolbar && OldViewportToolbar)
+		// Overrides of SCommonEditorViewportToolbarBase can construct the full base toolbar and then overwrite the child slot with their own content.
+		// This check ensures that the automatic upgrade is only allowed if we end up with a widget containing the child that we expected to have
+		// when the upgrade was requested.
+		if (TSharedPtr<SWidget> ExpectedChild = AutoUpgradeWidgetChild.Pin())
 		{
-			// Overrides of SCommonEditorViewportToolbarBase can construct the full base toolbar and then overwrite the child slot with their own content.
-			// This check ensures that the automatic upgrade is only allowed if we end up with a widget containing the child that we expected to have
-			// when the upgrade was requested.
-			if (TSharedPtr<SWidget> ExpectedChild = AutoUpgradeWidgetChild.Pin())
+			if (!NewViewportToolbar && OldViewportToolbar)
 			{
+				// The auto-upgrade toolbar was created with the legacy function call.
 				if (OldViewportToolbar->GetChildren()->Num() == 1 && OldViewportToolbar->GetChildren()->GetChildAt(0) == ExpectedChild)
 				{
 					bLegacyToolbarIsAutomaticallyUpgradable = true;
 					NewViewportToolbar = OldViewportToolbar;
 				}
 			}
+			else if (NewViewportToolbar && !OldViewportToolbar)
+			{
+				// The auto-upgrade toolbar was created with BuildViewportToolbar()
+				// In this case, we still want to be able to swap back and forth.
+				if (NewViewportToolbar->GetChildren()->Num() == 1 && NewViewportToolbar->GetChildren()->GetChildAt(0) == ExpectedChild)
+				{
+					bLegacyToolbarIsAutomaticallyUpgradable = true;
+					OldViewportToolbar = NewViewportToolbar;
+				}
+			}			
 		}
 
 		if (OldViewportToolbar)
@@ -174,7 +188,9 @@ void SEditorViewport::Construct( const FArguments& InArgs )
 								return EVisibility::Visible;
 							}
 							
+							PRAGMA_DISABLE_DEPRECATION_WARNINGS
 							return UE::UnrealEd::ShowOldViewportToolbars() ? EVisibility::Visible: EVisibility::Collapsed;
+							PRAGMA_ENABLE_DEPRECATION_WARNINGS
 						}
 					)
 					[
@@ -203,7 +219,9 @@ void SEditorViewport::Construct( const FArguments& InArgs )
 							return EVisibility::Visible;
 						}
 						// In the case of the old toolbar being upgraded, don't show it twice
+						PRAGMA_DISABLE_DEPRECATION_WARNINGS
 						return (UE::UnrealEd::ShowNewViewportToolbars() && (!bHasLegacyUpgradedToolbar || !UE::UnrealEd::ShowOldViewportToolbars())) ? EVisibility::Visible: EVisibility::Collapsed;
+						PRAGMA_ENABLE_DEPRECATION_WARNINGS
 					}
 				)
 				[
diff --git a/Engine/Source/Editor/UnrealEd/Private/ViewportToolbar/UnrealEdViewportToolbar.cpp b/Engine/Source/Editor/UnrealEd/Private/ViewportToolbar/UnrealEdViewportToolbar.cpp
--- a/Engine/Source/Editor/UnrealEd/Private/ViewportToolbar/UnrealEdViewportToolbar.cpp
+++ b/Engine/Source/Editor/UnrealEd/Private/ViewportToolbar/UnrealEdViewportToolbar.cpp
@@ -656,6 +656,7 @@ float GetCamSpeedScalarSliderPosition(const TSharedRef<SEditorViewport>& InEdito
 
 } // namespace UE::UnrealEd::Private
 
+// The accessor functions are deprecated, but the functionality remains so that licensee implementers can swap back and forth easily while transitioning.
 static FAutoConsoleVariableRef CVarToolMenusViewportToolbars(
 	TEXT("ToolMenusViewportToolbars"),
 	UE::UnrealEd::Private::CVarToolMenusViewportToolbarsValue,
diff --git a/Engine/Source/Editor/UnrealEd/Public/SEditorViewport.h b/Engine/Source/Editor/UnrealEd/Public/SEditorViewport.h
--- a/Engine/Source/Editor/UnrealEd/Public/SEditorViewport.h
+++ b/Engine/Source/Editor/UnrealEd/Public/SEditorViewport.h
@@ -164,7 +164,8 @@ protected:
 	virtual TSharedPtr<SWidget> BuildViewportToolbar() { return TSharedPtr<SWidget>(nullptr); }
 
 	// Implement this to add a viewport toolbar to the inside top of the viewport
-	virtual TSharedPtr<SWidget> MakeViewportToolbar() { return TSharedPtr<SWidget>(nullptr); }
+	UE_DEPRECATED(5.7, "Please use BuildViewportToolbar() to create a toolbar that lives above the viewport instead. Alternatively, insert the toolbar in PopulateViewportOverlays()")
+	virtual TSharedPtr<SWidget> MakeViewportToolbar() final { return TSharedPtr<SWidget>(nullptr); }
 
 	// Implement this to add an arbitrary set of toolbars or other overlays to the inside of the viewport
 	virtual void PopulateViewportOverlays(TSharedRef<SOverlay> Overlay) { }
diff --git a/Engine/Source/Editor/UnrealEd/Public/ViewportToolbar/UnrealEdViewportToolbar.h b/Engine/Source/Editor/UnrealEd/Public/ViewportToolbar/UnrealEdViewportToolbar.h
--- a/Engine/Source/Editor/UnrealEd/Public/ViewportToolbar/UnrealEdViewportToolbar.h
+++ b/Engine/Source/Editor/UnrealEd/Public/ViewportToolbar/UnrealEdViewportToolbar.h
@@ -39,9 +39,11 @@ enum EHidableViewModeMenuSections : uint8
 };
 
 /** The value of this function is controlled by the CVAR "ToolMenusViewportToolbars". */
+UE_DEPRECATED(5.7, "The legacy viewport toolbars are no longer supported and are no longer ever visible.")
 UNREALED_API bool ShowOldViewportToolbars();
 
 /** The value of this function is controlled by the CVAR "ToolMenusViewportToolbars". */
+UE_DEPRECATED(5.7, "The legacy viewport toolbars are no longer supported. The new toolbars are always visible.")
 UNREALED_API bool ShowNewViewportToolbars();
 
 UNREALED_API FSlateIcon GetIconFromCoordSystem(ECoordSystem InCoordSystem);
