Make the "Show Selected" action for actors always visible in the context menu. Add "can execute" handlers for Hide & Show Selected actions.

diff --git a/Engine/Source/Editor/LevelEditor/Private/LevelEditor.cpp b/Engine/Source/Editor/LevelEditor/Private/LevelEditor.cpp
--- a/Engine/Source/Editor/LevelEditor/Private/LevelEditor.cpp
+++ b/Engine/Source/Editor/LevelEditor/Private/LevelEditor.cpp
@@ -1102,12 +1102,14 @@ void FLevelEditorModule::BindGlobalLevelEditorCommands()
 
 	ActionList.MapAction(
 		Commands.ShowSelected,
-		FExecuteAction::CreateStatic( &FLevelEditorActionCallbacks::ExecuteExecCommand, FString( TEXT("ACTOR UNHIDE SELECTED") ) )
+		FExecuteAction::CreateStatic( &FLevelEditorActionCallbacks::ExecuteExecCommand, FString( TEXT("ACTOR UNHIDE SELECTED") ) ),
+		FCanExecuteAction::CreateStatic( &FLevelEditorActionCallbacks::CanShowSelectedActors )
 		);
 
 	ActionList.MapAction(
 		Commands.HideSelected,
-		FExecuteAction::CreateStatic( &FLevelEditorActionCallbacks::ExecuteExecCommand, FString( TEXT("ACTOR HIDE SELECTED") ) )
+		FExecuteAction::CreateStatic( &FLevelEditorActionCallbacks::ExecuteExecCommand, FString( TEXT("ACTOR HIDE SELECTED") ) ),
+		FCanExecuteAction::CreateStatic( &FLevelEditorActionCallbacks::CanHideSelectedActors )
 		);
 
 	ActionList.MapAction(
diff --git a/Engine/Source/Editor/LevelEditor/Private/LevelEditorActions.cpp b/Engine/Source/Editor/LevelEditor/Private/LevelEditorActions.cpp
--- a/Engine/Source/Editor/LevelEditor/Private/LevelEditorActions.cpp
+++ b/Engine/Source/Editor/LevelEditor/Private/LevelEditorActions.cpp
@@ -2207,6 +2207,32 @@ bool FLevelEditorActionCallbacks::CanSelectComponentOwnerActor()
 	return GEditor->GetSelectedComponentCount() > 0;
 }
 
+bool FLevelEditorActionCallbacks::CanHideSelectedActors()
+{
+	for ( FSelectionIterator It( GEditor->GetSelectedActorIterator() ) ; It ; ++It )
+	{
+		AActor* Actor = static_cast<AActor*>( *It );
+		if (!Actor->IsHiddenEd())
+		{
+			return true;
+		}
+	}
+	return false;
+}
+
+bool FLevelEditorActionCallbacks::CanShowSelectedActors()
+{
+	for ( FSelectionIterator It( GEditor->GetSelectedActorIterator() ) ; It ; ++It )
+	{
+		AActor* Actor = static_cast<AActor*>( *It );
+		if (Actor->IsHiddenEd())
+		{
+			return true;
+		}
+	}
+	return false;
+}
+
 void FLevelEditorActionCallbacks::OnSelectOwningHLODCluster()
 {
 	if (GEditor->GetSelectedActorCount() > 0)
diff --git a/Engine/Source/Editor/LevelEditor/Private/LevelEditorContextMenu.cpp b/Engine/Source/Editor/LevelEditor/Private/LevelEditorContextMenu.cpp
--- a/Engine/Source/Editor/LevelEditor/Private/LevelEditorContextMenu.cpp
+++ b/Engine/Source/Editor/LevelEditor/Private/LevelEditorContextMenu.cpp
@@ -1069,23 +1069,19 @@ void FLevelEditorContextMenu::BuildGroupMenu(FToolMenuSection& Section, const FS
 void FLevelEditorContextMenuImpl::FillActorVisibilityMenu(UToolMenu* Menu)
 {
 	{
-		FToolMenuSection& Section = Menu->AddSection("VisibilitySelected");
-		// Show 'Show Selected' only if the selection has any hidden actors
-		if ( SelectionInfo.bHaveHidden )
-		{
-			Section.AddMenuEntry(FLevelEditorCommands::Get().ShowSelected);
-		}
+		FToolMenuSection& Section = Menu->AddSection("VisibilitySelected", LOCTEXT("VisibilitySelectedHeader", "Selection"));
+		Section.AddMenuEntry(FLevelEditorCommands::Get().ShowSelected);
 		Section.AddMenuEntry(FLevelEditorCommands::Get().HideSelected);
 	}
 
 	{
-		FToolMenuSection& Section = Menu->AddSection("VisibilityAll");
+		FToolMenuSection& Section = Menu->AddSection("VisibilityAll", LOCTEXT("VisibilityAllHeader", "World"));
 		Section.AddMenuEntry(FLevelEditorCommands::Get().ShowSelectedOnly);
 		Section.AddMenuEntry(FLevelEditorCommands::Get().ShowAll);
 	}
 
 	{
-		FToolMenuSection& Section = Menu->AddSection("VisibilityStartup");
+		FToolMenuSection& Section = Menu->AddSection("VisibilityStartup", LOCTEXT("VisibilityStartupHeader", "Startup"));
 		Section.AddMenuEntry(FLevelEditorCommands::Get().ShowAllStartup);
 		Section.AddMenuEntry(FLevelEditorCommands::Get().ShowSelectedStartup);
 		Section.AddMenuEntry(FLevelEditorCommands::Get().HideSelectedStartup);
diff --git a/Engine/Source/Editor/LevelEditor/Public/LevelEditorActions.h b/Engine/Source/Editor/LevelEditor/Public/LevelEditorActions.h
--- a/Engine/Source/Editor/LevelEditor/Public/LevelEditorActions.h
+++ b/Engine/Source/Editor/LevelEditor/Public/LevelEditorActions.h
@@ -1128,6 +1128,12 @@ public:
 	/** Called to see if any components are selected */
 	static UE_API bool CanSelectComponentOwnerActor();
 
+	/** Called to see if selected actors can be hidden */
+	static UE_API bool CanHideSelectedActors();
+	
+	/** Called to see if selected actors can be hidden */
+	static UE_API bool CanShowSelectedActors();
+
 	/**
 	 * Called to select all lights
 	 */
