[ITF Viewport Interactions] Enable ITF Viewport Interactions for Camera Movements in Level Editor

Better working if new gizmos are enabled (Editor Preferences --> General/Transform Gizmos --> Enable Experimental Gizmos)

CVar to enable Camera Movements: `ViewportInteractions.EnableITFInteractions 1`

Also suggested to use ITF DragTools: `DragTools.EnableITFTools 1`

CL adds camera movement controls which should work approximately as the legacy ones:
- Zoom
- View Angle
- Panning
- Ortho Panning
- Orbit
- Move + Yaw
- FOV change
- Camera Translate
- Camera Rotate
- Camera Speed update with mouse wheel

Known issues/differences

- Numpad +/- will zoom instead of moving the camera
- LMB + RMB drag panning in perspective view will likely not work (use MMB drag)
- LMB + RMB drag zoom in ortho view will work only if LMB is pressed first

diff --git a/Engine/Source/Editor/LevelEditor/LevelEditor.Build.cs b/Engine/Source/Editor/LevelEditor/LevelEditor.Build.cs
--- a/Engine/Source/Editor/LevelEditor/LevelEditor.Build.cs
+++ b/Engine/Source/Editor/LevelEditor/LevelEditor.Build.cs
@@ -91,7 +91,8 @@ public class LevelEditor : ModuleRules
 				"DeveloperSettings",
 				"ActionableMessage",
 				"Json",
-				"JsonUtilities"
+				"JsonUtilities", 
+				"EditorInteractiveToolsFramework",
 			}
 		);
 
diff --git a/Engine/Source/Editor/LevelEditor/Private/SLevelEditor.cpp b/Engine/Source/Editor/LevelEditor/Private/SLevelEditor.cpp
--- a/Engine/Source/Editor/LevelEditor/Private/SLevelEditor.cpp
+++ b/Engine/Source/Editor/LevelEditor/Private/SLevelEditor.cpp
@@ -72,6 +72,7 @@
 #include "Elements/SMInstance/SMInstanceElementId.h"
 #include "Elements/SMInstance/SMInstanceElementLevelEditorSelectionCustomization.h"
 #include "DerivedDataEditor.h"
+#include "EditorInteractiveGizmoManager.h"
 #include "ZenEditor.h"
 #include "EditorModeManager.h"
 #include "EditorViewportLayout.h"
@@ -80,10 +81,16 @@
 #include "LevelEditorOutlinerSettings.h"
 #include "IWorldHierarchy.h"
 #include "LevelViewportActions.h"
+#include "EditorDragTools/DragToolsBehaviourSource.h"
 #include "Tools/EdModeInteractiveToolsContext.h"
+#include "Framework/Commands/GenericCommands.h"
+#include "ViewportInteractions/ViewportCommandsInteraction.h"
+#include "ViewportInteractions/ViewportInteractionsBehaviorSource.h"
+#include "ViewportInteractions/ViewportInteractionsUtilities.h"
 
 #define LOCTEXT_NAMESPACE "SLevelEditor"
 
+class UViewportOrbitInteraction;
 static const FName MainFrameModuleName("MainFrame");
 static const FName LevelEditorModuleName("LevelEditor");
 static const FName WorldBrowserHierarchyTab("WorldBrowserHierarchy");
@@ -100,7 +107,20 @@ SLevelEditor::SLevelEditor()
 			SLevelEditor::GetEditorModeManager().GetInteractiveToolsContext())
 	{
 		ModeManagerInteractiveToolsContext->SetDragToolsEnabled(true);
+		ModeManagerInteractiveToolsContext->SetInputToolsEnabled(true);
+
+		RegisterViewportInteractions();
 	}
+
+	// This is used while working with both legacy and ITF interactions. Once ITF is shipped as default, it shouldn't be needed anymore
+	UE::Editor::ViewportInteractions::OnEditorViewportInteractionsActivated().AddRaw(
+		this, &SLevelEditor::RegisterViewportInteractions
+	);
+
+	UEditorInteractiveGizmoManager::OnUsesNewTRSGizmosChangedDelegate().AddLambda([this](bool bInUsesExperimentalGizmos)
+	{
+		RegisterViewportInteractions();
+	});
 }
 
 void SLevelEditor::BindCommands()
@@ -473,6 +493,10 @@ SLevelEditor::~SLevelEditor()
 		GUnrealEd->GetSelectedActors()->SetElementSelectionSet(nullptr);
 		GUnrealEd->GetSelectedComponents()->SetElementSelectionSet(nullptr);
 	}
+
+	// This will be likely removed once ITF interactions become the default, and legacy ones are removed/deactivated
+	UE::Editor::ViewportInteractions::OnEditorViewportInteractionsActivated().RemoveAll(this);
+	UEditorInteractiveGizmoManager::OnUsesNewTRSGizmosChangedDelegate().RemoveAll(this);
 }
 
 FText SLevelEditor::GetTabTitle() const
@@ -765,6 +789,24 @@ void SLevelEditor::ResetMostRecentOutliner()
 	SceneOutlinerPtr.Reset();
 }
 
+void SLevelEditor::RegisterViewportInteractions() const
+{
+	if (UModeManagerInteractiveToolsContext* const ModeManagerInteractiveToolsContext =
+			SLevelEditor::GetEditorModeManager().GetInteractiveToolsContext())
+	{
+		if (UViewportInteractionsBehaviorSource* InputToolsSource =
+				ModeManagerInteractiveToolsContext->GetViewportInteractionsBehaviorSource())
+		{
+			InputToolsSource->Reset();
+
+			UE::Editor::ViewportInteractions::AddDefaultCameraMovementInteractions(InputToolsSource);
+
+			// make sure to register all BehaviorSources
+			InputToolsSource->RegisterBehaviorSources();
+		}
+	}
+}
+
 TSharedPtr<ISceneOutliner> SLevelEditor::GetMostRecentlyUsedSceneOutliner()
 {
 	// If the outliner marked as "most recent" is still open, simply return it
diff --git a/Engine/Source/Editor/LevelEditor/Private/SLevelEditor.h b/Engine/Source/Editor/LevelEditor/Private/SLevelEditor.h
--- a/Engine/Source/Editor/LevelEditor/Private/SLevelEditor.h
+++ b/Engine/Source/Editor/LevelEditor/Private/SLevelEditor.h
@@ -277,6 +277,8 @@ private:
 	/** Sets the first available outliner as the most recent outliner */
 	void ResetMostRecentOutliner();
 
+	void RegisterViewportInteractions() const;
+
 private:
 
 	// Tracking the active viewports in this level editor.
