This change addresses an inconsistency in the Viewport right-click context menu, where the Transform options were always ordered according to FLU rather than the expected LUF, even with LUF coordinate system in use.

Behavior is now consistent and context-aware across coordinate systems.

diff --git a/Engine/Source/Editor/LevelEditor/Private/LevelEditorContextMenu.cpp b/Engine/Source/Editor/LevelEditor/Private/LevelEditorContextMenu.cpp
--- a/Engine/Source/Editor/LevelEditor/Private/LevelEditorContextMenu.cpp
+++ b/Engine/Source/Editor/LevelEditor/Private/LevelEditorContextMenu.cpp
@@ -54,6 +54,7 @@
 #include "SLevelEditor.h"
 #include "SLevelViewport.h"
 #include "PropertyEditorModule.h"
+#include "Misc/AxisDisplayInfo.h"
 
 #define LOCTEXT_NAMESPACE "LevelViewportContextMenu"
 
@@ -1151,9 +1152,20 @@ void FLevelEditorContextMenuImpl::FillTransformMenu(UToolMenu* Menu)
 		if (LevelEditorContext->CurrentSelection->HasSelectedObjects<AActor>())
 		{
 			FToolMenuSection& Section = Menu->AddSection("MirrorLock");
-			Section.AddMenuEntry(FLevelEditorCommands::Get().MirrorActorX);
-			Section.AddMenuEntry(FLevelEditorCommands::Get().MirrorActorY);
-			Section.AddMenuEntry(FLevelEditorCommands::Get().MirrorActorZ);
+
+			if (AxisDisplayInfo::GetAxisDisplayCoordinateSystem() == EAxisList::LeftUpForward)
+			{
+				Section.AddMenuEntry(FLevelEditorCommands::Get().MirrorActorY);
+				Section.AddMenuEntry(FLevelEditorCommands::Get().MirrorActorZ);
+				Section.AddMenuEntry(FLevelEditorCommands::Get().MirrorActorX);
+			}
+			else
+			{
+				Section.AddMenuEntry(FLevelEditorCommands::Get().MirrorActorX);
+				Section.AddMenuEntry(FLevelEditorCommands::Get().MirrorActorY);
+				Section.AddMenuEntry(FLevelEditorCommands::Get().MirrorActorZ);
+			}
+
 			Section.AddMenuEntry(FLevelEditorCommands::Get().LockActorMovement);
 		}
 	}
