Recognize external text as a valid drag source.  Needed, else OnDrop() won't called for external text.

diff --git a/Engine/Source/Editor/LevelEditor/Private/SLevelViewport.cpp b/Engine/Source/Editor/LevelEditor/Private/SLevelViewport.cpp
--- a/Engine/Source/Editor/LevelEditor/Private/SLevelViewport.cpp
+++ b/Engine/Source/Editor/LevelEditor/Private/SLevelViewport.cpp
@@ -918,6 +918,10 @@ bool SLevelViewport::HandleDragObjects(const FGeometry& MyGeometry, const FDragD
 	{
 		bValidDrag = true;
 	}
+	else if (Operation->IsOfType<FExternalDragOperation>() && StaticCastSharedPtr<FExternalDragOperation>(Operation)->HasText())
+	{
+		bValidDrag = true;
+	}
 	else
 	{
 		GetAssetsFromDrag(DragDropEvent, SelectedAssetDatas);
@@ -1083,9 +1087,9 @@ bool SLevelViewport::HandlePlaceDraggedObjects(const FGeometry& MyGeometry, cons
 			const FAssetData& AssetData = DroppedAssetDatas[AssetIdx];
 
 			UObject* Asset = AssetData.GetAsset();
-			if ( Asset != nullptr )
+			if (Asset != nullptr)
 			{
-				DroppedObjects.Add( Asset );
+				DroppedObjects.Add(Asset);
 			}
 			else
 			{
