Changing the default help hot key to shift+F1
F1 will be use to summon the ai assistant contextual help promp.

diff --git a/Engine/Plugins/Experimental/AIAssistant/Source/AIAssistant/Private/AIAssistantCommands.cpp b/Engine/Plugins/Experimental/AIAssistant/Source/AIAssistant/Private/AIAssistantCommands.cpp
--- a/Engine/Plugins/Experimental/AIAssistant/Source/AIAssistant/Private/AIAssistantCommands.cpp
+++ b/Engine/Plugins/Experimental/AIAssistant/Source/AIAssistant/Private/AIAssistantCommands.cpp
@@ -19,7 +19,7 @@ FAIAssistantCommands::FAIAssistantCommands() :
 {
 	UI_COMMAND(OpenAIAssistantTab, "AI Assistant", "Bring up AI Assistant window", EUserInterfaceActionType::Button, FInputChord());
 
-	UI_COMMAND(SlateQueryCommand, "Query UI with AI Assistant", "Query the UI widget under the mouse with AI Assistant", EUserInterfaceActionType::Button, FInputChord(EModifierKey::Shift, EKeys::F1));
+	UI_COMMAND(SlateQueryCommand, "Query UI with AI Assistant", "Query the UI widget under the mouse with AI Assistant", EUserInterfaceActionType::Button, FInputChord(EKeys::F1));
 }
 
 
diff --git a/Engine/Source/Editor/LevelEditor/Private/LevelEditorActions.cpp b/Engine/Source/Editor/LevelEditor/Private/LevelEditorActions.cpp
--- a/Engine/Source/Editor/LevelEditor/Private/LevelEditorActions.cpp
+++ b/Engine/Source/Editor/LevelEditor/Private/LevelEditorActions.cpp
@@ -4101,7 +4101,7 @@ FLevelEditorCommands::FLevelEditorCommands()
 UE_DISABLE_OPTIMIZATION_SHIP
 void FLevelEditorCommands::RegisterCommands()
 {
-	UI_COMMAND( BrowseDocumentation, "Level Editor Documentation", "Details on how to use the Level Editor", EUserInterfaceActionType::Button, FInputChord( EKeys::F1 ) );
+	UI_COMMAND( BrowseDocumentation, "Level Editor Documentation", "Details on how to use the Level Editor", EUserInterfaceActionType::Button, FInputChord( EModifierKey::Shift, EKeys::F1 ) );
 	UI_COMMAND( BrowseViewportControls, "Viewport Controls", "Ways to move around in the 3D viewport", EUserInterfaceActionType::Button, FInputChord() );
 	UI_COMMAND( NewLevel, "New Level...", "Create a new level, or choose a level template to start from.", EUserInterfaceActionType::Button, FInputChord( EModifierKey::Control, EKeys::N ) );
 	UI_COMMAND( Save, "Save Current Level", "Saves the current level to disk", EUserInterfaceActionType::Button, FInputChord(EModifierKey::Control, EKeys::S) );
diff --git a/Engine/Source/Editor/UnrealEd/Private/Toolkits/GlobalEditorCommonCommands.cpp b/Engine/Source/Editor/UnrealEd/Private/Toolkits/GlobalEditorCommonCommands.cpp
--- a/Engine/Source/Editor/UnrealEd/Private/Toolkits/GlobalEditorCommonCommands.cpp
+++ b/Engine/Source/Editor/UnrealEd/Private/Toolkits/GlobalEditorCommonCommands.cpp
@@ -58,7 +58,7 @@ void FGlobalEditorCommonCommands::RegisterCommands()
 	
 	UI_COMMAND(OpenOutputLogDrawer, "Open Output Log Drawer", "Opens the output log drawer from the active asset editor status bar", EUserInterfaceActionType::Button, FInputChord(EModifierKey::Alt, EKeys::Tilde));
 
-	UI_COMMAND(OpenDocumentation, "Open Documentation...", "Opens documentation for this tool", EUserInterfaceActionType::Button, FInputChord(EKeys::F1));
+	UI_COMMAND(OpenDocumentation, "Open Documentation...", "Opens documentation for this tool", EUserInterfaceActionType::Button, FInputChord(EModifierKey::Shift, EKeys::F1));
 
 	UI_COMMAND(OpenLevel, "Open Level...", "Loads an existing level", EUserInterfaceActionType::Button, FInputChord(EModifierKey::Control, EKeys::O));
 
diff --git a/Engine/Source/Runtime/Slate/Private/Framework/MultiBox/Mac/MacMenu.cpp b/Engine/Source/Runtime/Slate/Private/Framework/MultiBox/Mac/MacMenu.cpp
--- a/Engine/Source/Runtime/Slate/Private/Framework/MultiBox/Mac/MacMenu.cpp
+++ b/Engine/Source/Runtime/Slate/Private/Framework/MultiBox/Mac/MacMenu.cpp
@@ -131,7 +131,7 @@ public:
 	virtual void RegisterCommands() override
 	{
 		UI_COMMAND( Preferences, "Preferences...", "Preferences", EUserInterfaceActionType::None, FInputChord( EModifierKey::Control, EKeys::Comma ) )
-		UI_COMMAND( Help, "Help", "Help", EUserInterfaceActionType::None, FInputChord( EModifierKey::None, EKeys::F1 ) )
+		UI_COMMAND( Help, "Help", "Help", EUserInterfaceActionType::None, FInputChord( EModifierKey::Shift, EKeys::F1 ) )
 		UI_COMMAND( Hide, "Hide", "Hide", EUserInterfaceActionType::None, FInputChord( EModifierKey::Control, EKeys::H ) )
 		UI_COMMAND( HideOthers, "Hide Others", "Hide Others", EUserInterfaceActionType::None, FInputChord( EModifierKey::Control | EModifierKey::Alt, EKeys::H ) )
 		UI_COMMAND( ShowAll, "Show All", "Show All", EUserInterfaceActionType::None, FInputChord() )
