[HomeScreen] Set the HomeScreen/LevelEditor and LevelEditorViewport LockedState

diff --git a/Engine/Source/Editor/LevelEditor/Private/LevelEditor.cpp b/Engine/Source/Editor/LevelEditor/Private/LevelEditor.cpp
--- a/Engine/Source/Editor/LevelEditor/Private/LevelEditor.cpp
+++ b/Engine/Source/Editor/LevelEditor/Private/LevelEditor.cpp
@@ -317,6 +317,7 @@ void FLevelEditorModule::StartupModule()
 	const IWorkspaceMenuStructure& MenuStructure = WorkspaceMenu::GetMenuStructure();
 
 	FGlobalTabmanager::Get()->RegisterTabSpawner("LevelEditor", FOnSpawnTab::CreateRaw(this, &FLevelEditorModule::SpawnLevelEditor))
+		.SetIsTabLocked(true)
 		.SetDisplayName(NSLOCTEXT("LevelEditor", "LevelEditorTab", "Level Editor"))
 		.SetAutoGenerateMenuEntry(false);
 
diff --git a/Engine/Source/Editor/LevelEditor/Private/SLevelEditor.cpp b/Engine/Source/Editor/LevelEditor/Private/SLevelEditor.cpp
--- a/Engine/Source/Editor/LevelEditor/Private/SLevelEditor.cpp
+++ b/Engine/Source/Editor/LevelEditor/Private/SLevelEditor.cpp
@@ -1410,24 +1410,28 @@ TSharedRef<SWidget> SLevelEditor::RestoreContentArea( const TSharedRef<SDockTab>
 				.SetTooltipText(ViewportTooltip)
 				.SetGroup(MenuStructure.GetLevelEditorViewportsCategory())
 				.SetIcon(ViewportIcon)
+				.SetIsTabLocked(true)
 				.SetCanSidebarTab(false);
 
 			LevelEditorTabManager->RegisterTabSpawner(LevelEditorTabIds::LevelEditorViewport_Clone1, FOnSpawnTab::CreateSP(this, &SLevelEditor::SpawnLevelEditorTab, LevelEditorTabIds::LevelEditorViewport_Clone1, FString()) )
 				.SetDisplayName(NSLOCTEXT("LevelEditorTabs", "LevelEditorViewport_Clone1", "Viewport 2"))
 				.SetTooltipText(ViewportTooltip)
 				.SetGroup( MenuStructure.GetLevelEditorViewportsCategory() )
+				.SetIsTabLocked(true)
 				.SetIcon(ViewportIcon);
 
 			LevelEditorTabManager->RegisterTabSpawner(LevelEditorTabIds::LevelEditorViewport_Clone2, FOnSpawnTab::CreateSP(this, &SLevelEditor::SpawnLevelEditorTab, LevelEditorTabIds::LevelEditorViewport_Clone2, FString()) )
 				.SetDisplayName(NSLOCTEXT("LevelEditorTabs", "LevelEditorViewport_Clone2", "Viewport 3"))
 				.SetTooltipText(ViewportTooltip)
 				.SetGroup( MenuStructure.GetLevelEditorViewportsCategory() )
+				.SetIsTabLocked(true)
 				.SetIcon(ViewportIcon);
 
 				LevelEditorTabManager->RegisterTabSpawner(LevelEditorTabIds::LevelEditorViewport_Clone3, FOnSpawnTab::CreateSP(this, &SLevelEditor::SpawnLevelEditorTab, LevelEditorTabIds::LevelEditorViewport_Clone3, FString()) )
 				.SetDisplayName(NSLOCTEXT("LevelEditorTabs", "LevelEditorViewport_Clone3", "Viewport 4"))
 				.SetTooltipText(ViewportTooltip)
 				.SetGroup( MenuStructure.GetLevelEditorViewportsCategory() )
+				.SetIsTabLocked(true)
 				.SetIcon(ViewportIcon);
 		}
 		{
diff --git a/Engine/Source/Editor/MainFrame/Private/MainFrameModule.cpp b/Engine/Source/Editor/MainFrame/Private/MainFrameModule.cpp
--- a/Engine/Source/Editor/MainFrame/Private/MainFrameModule.cpp
+++ b/Engine/Source/Editor/MainFrame/Private/MainFrameModule.cpp
@@ -686,6 +686,7 @@ void FMainFrameModule::StartupModule( )
     		FGlobalTabmanager::Get()->RegisterTabSpawner(
     			"HomeScreen", FOnSpawnTab::CreateRaw(this, &FMainFrameModule::SpawnHomeScreenWindow))
     			.SetAutoGenerateMenuEntry(false)
+				.SetIsTabLocked(true)
     			.SetMenuType(ETabSpawnerMenuType::Hidden)
     			.SetDisplayName(FText::GetEmpty());
     
