[LevelEditor] Revert lock state default state for LevelEditorViewport

diff --git a/Engine/Source/Editor/LevelEditor/Private/SLevelEditor.cpp b/Engine/Source/Editor/LevelEditor/Private/SLevelEditor.cpp
--- a/Engine/Source/Editor/LevelEditor/Private/SLevelEditor.cpp
+++ b/Engine/Source/Editor/LevelEditor/Private/SLevelEditor.cpp
@@ -1410,28 +1410,24 @@ TSharedRef<SWidget> SLevelEditor::RestoreContentArea( const TSharedRef<SDockTab>
 				.SetTooltipText(ViewportTooltip)
 				.SetGroup(MenuStructure.GetLevelEditorViewportsCategory())
 				.SetIcon(ViewportIcon)
-				.SetIsTabLocked(true)
 				.SetCanSidebarTab(false);
 
 			LevelEditorTabManager->RegisterTabSpawner(LevelEditorTabIds::LevelEditorViewport_Clone1, FOnSpawnTab::CreateSP(this, &SLevelEditor::SpawnLevelEditorTab, LevelEditorTabIds::LevelEditorViewport_Clone1, FString()) )
 				.SetDisplayName(NSLOCTEXT("LevelEditorTabs", "LevelEditorViewport_Clone1", "Viewport 2"))
 				.SetTooltipText(ViewportTooltip)
 				.SetGroup( MenuStructure.GetLevelEditorViewportsCategory() )
-				.SetIsTabLocked(true)
 				.SetIcon(ViewportIcon);
 
 			LevelEditorTabManager->RegisterTabSpawner(LevelEditorTabIds::LevelEditorViewport_Clone2, FOnSpawnTab::CreateSP(this, &SLevelEditor::SpawnLevelEditorTab, LevelEditorTabIds::LevelEditorViewport_Clone2, FString()) )
 				.SetDisplayName(NSLOCTEXT("LevelEditorTabs", "LevelEditorViewport_Clone2", "Viewport 3"))
 				.SetTooltipText(ViewportTooltip)
 				.SetGroup( MenuStructure.GetLevelEditorViewportsCategory() )
-				.SetIsTabLocked(true)
 				.SetIcon(ViewportIcon);
 
 				LevelEditorTabManager->RegisterTabSpawner(LevelEditorTabIds::LevelEditorViewport_Clone3, FOnSpawnTab::CreateSP(this, &SLevelEditor::SpawnLevelEditorTab, LevelEditorTabIds::LevelEditorViewport_Clone3, FString()) )
 				.SetDisplayName(NSLOCTEXT("LevelEditorTabs", "LevelEditorViewport_Clone3", "Viewport 4"))
 				.SetTooltipText(ViewportTooltip)
 				.SetGroup( MenuStructure.GetLevelEditorViewportsCategory() )
-				.SetIsTabLocked(true)
 				.SetIcon(ViewportIcon);
 		}
 		{
