Certain features that Vulkan SM5 supports are not supported in D3D11, so we disable this preview platform for that RHI explicitly.

diff --git a/Engine/Source/Editor/LevelEditor/Private/LevelEditorActions.cpp b/Engine/Source/Editor/LevelEditor/Private/LevelEditorActions.cpp
--- a/Engine/Source/Editor/LevelEditor/Private/LevelEditorActions.cpp
+++ b/Engine/Source/Editor/LevelEditor/Private/LevelEditorActions.cpp
@@ -837,6 +837,13 @@ bool FLevelEditorActionCallbacks::CanExecutePreviewPlatform(FPreviewPlatformInfo
 		return false;
 	}
 
+	// Disable shader format SF_VULKAN_SM5 when the preferred shader code format for preview platforms is DXBC (D3D11).
+	// Many features supported in Vulkan SM5 are not supported by D3D11, such as UAV limits.
+	if (NewPreviewPlatform.PreviewShaderFormatName == TEXT("SF_VULKAN_SM5") && GRHIGlobals.PreferredPreviewShaderCodeFormat == TEXT("DXBC"))
+	{
+		return false;
+	}
+
 	const EShaderPlatform PreviewShaderPlatform = NewPreviewPlatform.ShaderPlatform;
 
 	if (FDataDrivenShaderPlatformInfo::IsValid(PreviewShaderPlatform) && FDataDrivenShaderPlatformInfo::GetIsPreviewPlatform(PreviewShaderPlatform))
