[ITF Viewport Interactions] Add Selection interactions to Level Editor

Add click based interactions to level editor

- LMB for selection
- RMB for selection + summoning context menu
- LMB double click

diff --git a/Engine/Source/Editor/LevelEditor/LevelEditor.Build.cs b/Engine/Source/Editor/LevelEditor/LevelEditor.Build.cs
--- a/Engine/Source/Editor/LevelEditor/LevelEditor.Build.cs
+++ b/Engine/Source/Editor/LevelEditor/LevelEditor.Build.cs
@@ -94,6 +94,7 @@ public class LevelEditor : ModuleRules
 				"Json",
 				"JsonUtilities", 
 				"EditorInteractiveToolsFramework",
+				"InteractiveToolsFramework"
 			}
 		);
 
diff --git a/Engine/Source/Editor/LevelEditor/Private/SLevelEditor.cpp b/Engine/Source/Editor/LevelEditor/Private/SLevelEditor.cpp
--- a/Engine/Source/Editor/LevelEditor/Private/SLevelEditor.cpp
+++ b/Engine/Source/Editor/LevelEditor/Private/SLevelEditor.cpp
@@ -81,10 +81,9 @@
 #include "LevelEditorOutlinerSettings.h"
 #include "IWorldHierarchy.h"
 #include "LevelViewportActions.h"
-#include "EditorDragTools/DragToolsBehaviourSource.h"
 #include "Tools/EdModeInteractiveToolsContext.h"
 #include "Framework/Commands/GenericCommands.h"
-#include "ViewportInteractions/ViewportCommandsInteraction.h"
+#include "ViewportInteractions/LevelViewportClickSelection.h"
 #include "ViewportInteractions/ViewportInteractionsBehaviorSource.h"
 #include "ViewportInteractions/ViewportInteractionsUtilities.h"
 
@@ -801,6 +800,15 @@ void SLevelEditor::RegisterViewportInteractions() const
 
 			UE::Editor::ViewportInteractions::AddDefaultCameraMovementInteractions(InputToolsSource);
 
+			// Left Mouse Click
+			InputToolsSource->AddInteraction<ULevelViewportClickSelection>();
+
+			// Right Mouse Click - summons context menu
+			InputToolsSource->AddInteraction<ULevelViewportClickContextMenu>();
+
+			// Double click
+			InputToolsSource->AddInteraction<ULevelViewportDoubleClick>();
+
 			// make sure to register all BehaviorSources
 			InputToolsSource->RegisterBehaviorSources();
 		}
diff --git a/Engine/Source/Editor/LevelEditor/Private/ViewportInteractions/LevelViewportClickSelection.cpp b/Engine/Source/Editor/LevelEditor/Private/ViewportInteractions/LevelViewportClickSelection.cpp
new file mode 100644
--- /dev/null
+++ b/Engine/Source/Editor/LevelEditor/Private/ViewportInteractions/LevelViewportClickSelection.cpp
@@ -0,0 +1,131 @@
+// Copyright Epic Games, Inc. All Rights Reserved.
+
+#include "LevelViewportClickSelection.h"
+#include "BaseBehaviors/DoubleClickBehavior.h"
+#include "Editor/UnrealEdEngine.h"
+#include "EditorModeManager.h"
+#include "EngineUtils.h"
+#include "HModel.h"
+#include "SceneView.h"
+#include "Subsystems/BrushEditingSubsystem.h"
+#include "UnrealEdGlobals.h"
+#include "ViewportInteractions/ViewportInteractionsBehaviorSource.h"
+#include "ViewportSelectionUtilities.h"
+
+namespace UE::Editor::LevelViewportInteractions
+{
+
+// See FLevelEditorViewportClient::ProcessClick for reference.
+// This might look simplified, since ViewportSelectionUtilities::ClickElement is likely handling some cases which are never hit in the original code
+void ProcessClick(FEditorViewportClient* InEditorViewportClient, const FViewportClick& InViewportClick, HHitProxy* InHitProxy)
+{
+	if (!InEditorViewportClient)
+	{
+		return;
+	}
+
+	if (!InEditorViewportClient->IsLevelEditorClient())
+	{
+		return;
+	}
+
+	FEditorModeTools* ModeTools = InEditorViewportClient->GetModeTools();
+	if (!ModeTools)
+	{
+		return;
+	}
+
+	UBrushEditingSubsystem* BrushSubsystem = GEditor->GetEditorSubsystem<UBrushEditingSubsystem>();
+
+	if (ModeTools->HandleClick(InEditorViewportClient, InHitProxy, InViewportClick))
+	{
+		return;
+	}
+
+	const FTypedElementHandle HitElement = InHitProxy ? InHitProxy->GetElementHandle() : FTypedElementHandle();
+
+	if (InHitProxy == nullptr)
+	{
+		UE::Editor::ViewportSelectionUtilities::ClickBackdrop(InEditorViewportClient,InViewportClick);
+	}
+	else if (GUnrealEd->ComponentVisManager.HandleClick(InEditorViewportClient, InHitProxy, InViewportClick))
+	{
+		// Component vis manager handled the click
+	}
+	else if (HitElement && UE::Editor::ViewportSelectionUtilities::ClickElement(InEditorViewportClient, HitElement, InViewportClick))
+	{
+		// Element handled the click
+	}
+	else if (InHitProxy->IsA(HBSPBrushVert::StaticGetType()) && static_cast<HBSPBrushVert*>(InHitProxy)->Brush.IsValid())
+	{
+		FVector Vertex = FVector(*static_cast<HBSPBrushVert*>(InHitProxy)->Vertex);
+		UE::Editor::ViewportSelectionUtilities::ClickBrushVertex(InEditorViewportClient, static_cast<HBSPBrushVert*>(InHitProxy)->Brush.Get(),&Vertex,InViewportClick);
+	}
+	else if (InHitProxy->IsA(HStaticMeshVert::StaticGetType()))
+	{
+		UE::Editor::ViewportSelectionUtilities::ClickStaticMeshVertex(InEditorViewportClient,
+			static_cast<HStaticMeshVert*>(InHitProxy)->Actor,
+			static_cast<HStaticMeshVert*>(InHitProxy)->Vertex,InViewportClick);
+	}
+	else if (BrushSubsystem && BrushSubsystem->ProcessClickOnBrushGeometry(InEditorViewportClient, InHitProxy, InViewportClick))
+	{
+		// Handled by the brush subsystem
+	}
+	else if (InHitProxy->IsA(HModel::StaticGetType()))
+	{
+		HModel* ModelHit = static_cast<HModel*>(InHitProxy);
+
+		// Compute the viewport's current view family.
+		FSceneViewFamilyContext ViewFamily(FSceneViewFamily::ConstructionValues( InEditorViewportClient->Viewport, InEditorViewportClient->GetScene(), InEditorViewportClient->EngineShowFlags ));
+		FSceneView* SceneView = InEditorViewportClient->CalcSceneView( &ViewFamily );
+
+		uint32 SurfaceIndex = INDEX_NONE;
+		if(ModelHit->ResolveSurface(SceneView,InViewportClick.GetClickPos().X,InViewportClick.GetClickPos().Y,SurfaceIndex))
+		{
+			UE::Editor::ViewportSelectionUtilities::ClickSurface(InEditorViewportClient, ModelHit->GetModel(),SurfaceIndex,InViewportClick);
+		}
+	}
+	//TODO: check if this case is ever hit (cannot use StaticGetType as above, since it's not exported from LevelEditor module)
+	else if (static_cast<HLevelSocketProxy*>(InHitProxy))
+	{
+		UE::Editor::ViewportSelectionUtilities::ClickLevelSocket(InEditorViewportClient, InHitProxy, InViewportClick);
+	}
+}
+}
+
+ULevelViewportClickSelection::ULevelViewportClickSelection()
+{
+	ToolType = UE::Editor::ViewportInteractions::SingleClick;
+	InteractionName = TEXT("Level Viewport Click Selection");
+}
+
+void ULevelViewportClickSelection::ProcessClick_Internal(const FViewportClick& InViewportClick, HHitProxy* InHitProxy, FSceneView& View)
+{
+	UE::Editor::ViewportInteractions::LOG(*InteractionName->ToString());
+	UE::Editor::LevelViewportInteractions::ProcessClick(GetEditorViewportClient(), InViewportClick, InHitProxy);
+}
+
+ULevelViewportClickContextMenu::ULevelViewportClickContextMenu()
+{
+	ToolType = UE::Editor::ViewportInteractions::ContextMenu;
+	InteractionName = TEXT("Level Viewport Context Menu");
+
+	MouseButton = EKeys::RightMouseButton;
+
+	if (TStrongObjectPtr<USingleClickInputBehavior> SingleClickInputBehavior = SingleClickInputBehaviorWeak.Pin())
+	{
+		SingleClickInputBehavior->SetUseRightMouseButton();
+	}
+}
+
+ULevelViewportDoubleClick::ULevelViewportDoubleClick()
+{
+	UDoubleClickInputBehavior* DoubleClickInputBehavior = NewObject<UDoubleClickInputBehavior>();
+	DoubleClickInputBehavior->Initialize(this);
+	DoubleClickInputBehavior->SetDefaultPriority(FInputCapturePriority::DEFAULT_GIZMO_PRIORITY - 1);
+
+	ToolType = UE::Editor::ViewportInteractions::DoubleClick;
+	InteractionName = TEXT("Level Viewport Double Click");
+
+	RegisterInputBehavior(DoubleClickInputBehavior);
+}
diff --git a/Engine/Source/Editor/LevelEditor/Private/ViewportInteractions/LevelViewportClickSelection.h b/Engine/Source/Editor/LevelEditor/Private/ViewportInteractions/LevelViewportClickSelection.h
new file mode 100644
--- /dev/null
+++ b/Engine/Source/Editor/LevelEditor/Private/ViewportInteractions/LevelViewportClickSelection.h
@@ -0,0 +1,51 @@
+// Copyright Epic Games, Inc. All Rights Reserved.
+
+#pragma once
+
+#include "EditorViewportClient.h"
+#include "ViewportInteractions/ViewportClickInteraction.h"
+#include "LevelViewportClickSelection.generated.h"
+
+/**
+ * Left-mouse click in level editor viewport, for elements selection
+ */
+UCLASS()
+class LEVELEDITOR_API ULevelViewportClickSelection : public UViewportClickInteraction
+{
+	GENERATED_BODY()
+
+public:
+	ULevelViewportClickSelection();
+
+protected:
+	virtual void ProcessClick_Internal(const FViewportClick& InViewportClick, HHitProxy* InHitProxy, FSceneView& View) override;
+};
+
+/**
+ * Right-mouse click in level editor viewport, for selecting + summoning context menu
+ */
+UCLASS()
+class LEVELEDITOR_API ULevelViewportClickContextMenu : public ULevelViewportClickSelection
+{
+	GENERATED_BODY()
+
+public:
+	ULevelViewportClickContextMenu();
+};
+
+/**
+ * Left-mouse double click
+ */
+UCLASS()
+class LEVELEDITOR_API ULevelViewportDoubleClick : public ULevelViewportClickSelection
+{
+	GENERATED_BODY()
+
+public:
+	ULevelViewportDoubleClick();
+
+	virtual EInputEvent GetMouseEvent() const override
+	{
+		return IE_DoubleClick;
+	}
+};
