[ITF Drag Tools] Toggling "Use Experimental Gizmo" puts the Transform Tools menu in a weird state

Marquee Selection section is now added by `UE::LevelEditor::ExtendTransformSubmenu()` instead of `UModeManagerInteractiveToolsContext::RegisterDragTools()`.

Marquee Selection submenu code is currently specific to Level Editor settings.
If this menu will be needed in the future by other editors/toolbars, we could make a dedicated function and put it somewhere else.

diff --git a/Engine/Source/Editor/LevelEditor/Private/ViewportToolbar/LevelEditorSubmenus.cpp b/Engine/Source/Editor/LevelEditor/Private/ViewportToolbar/LevelEditorSubmenus.cpp
--- a/Engine/Source/Editor/LevelEditor/Private/ViewportToolbar/LevelEditorSubmenus.cpp
+++ b/Engine/Source/Editor/LevelEditor/Private/ViewportToolbar/LevelEditorSubmenus.cpp
@@ -6,6 +6,7 @@
 #include "Bookmarks/BookmarkUI.h"
 #include "Bookmarks/IBookmarkTypeTools.h"
 #include "Camera/CameraActor.h"
+#include "EditorInteractiveGizmoManager.h"
 #include "EditorViewportCommands.h"
 #include "Engine/SceneCapture.h"
 #include "Engine/GameViewportClient.h"
@@ -37,6 +38,7 @@
 #include "Widgets/SToolTip.h"
 #include "WorldPartition/IWorldPartitionEditorModule.h"
 #include "VirtualShadowMapVisualizationData.h"
+#include "Tools/EdModeInteractiveToolsContext.h"
 
 #define LOCTEXT_NAMESPACE "LevelEditorViewportToolbar"
 
@@ -3107,6 +3109,87 @@ void ExtendTransformSubmenu(FName InTransformSubmenuName)
 		SelectionSection.AddMenuEntry(FLevelEditorCommands::Get().TransparentBoxSelect);
 		SelectionSection.AddMenuEntry(FLevelEditorCommands::Get().ShowSelectionSubcomponents);
 		SelectionSection.AddMenuEntry(FLevelEditorCommands::Get().EnableViewportHoverFeedback);
+
+		// Marquee selection options
+		{
+			FToolMenuEntry& MarqueeMode = SelectionSection.AddSubMenu(
+				"StrictBoxSelect", // Overwrite the other setting
+				LOCTEXT("MarqueeSelectionLabel", "Marquee Selection"),
+				LOCTEXT("MarqueeSelectionTooltip", "Define how objects are selected by marquee selection"),
+				FNewToolMenuDelegate::CreateLambda(
+					[](UToolMenu* Menu)
+					{
+						FProperty* Property = ULevelEditorViewportSettings::StaticClass()->FindPropertyByName(
+							GET_MEMBER_NAME_CHECKED(ULevelEditorViewportSettings, MarqueeSelectionMode)
+						);
+						if (!Property || !Property->IsA<FEnumProperty>())
+						{
+							return;
+						}
+
+						FEnumProperty* EnumProperty = StaticCast<FEnumProperty*>(Property);
+						UEnum* Enum = EnumProperty->GetEnum();
+						if (!Enum)
+						{
+							return;
+						}
+
+						TWeakObjectPtr<UEnum> WeakEnum = Enum;
+
+						FToolMenuSection& Section = Menu->AddSection("Mode", LOCTEXT("MarqueeModeSection", "Mode"));
+
+						const int32 MaxIndex = Enum->NumEnums() - (Enum->ContainsExistingMax() ? 1 : 0);
+						for (int32 Index = 0; Index < MaxIndex; Index++)
+						{
+							Section.AddMenuEntry(
+								NAME_None,
+								Enum->GetDisplayNameTextByIndex(Index),
+								Enum->GetToolTipTextByIndex(Index),
+								FSlateIcon(),
+								FUIAction(
+									FExecuteAction::CreateLambda(
+										[WeakEnum, Index]
+										{
+											if (UEnum* Enum = WeakEnum.Get())
+											{
+												ULevelEditorViewportSettings* Settings =
+													GetMutableDefault<ULevelEditorViewportSettings>();
+												Settings->MarqueeSelectionMode =
+													static_cast<EMarqueeSelectionMode>(Enum->GetValueByIndex(Index));
+												Settings->SaveConfig();
+											}
+										}
+									),
+									FCanExecuteAction(),
+									FIsActionChecked::CreateLambda(
+										[WeakEnum, Index]
+										{
+											if (UEnum* Enum = WeakEnum.Get())
+											{
+												const ULevelEditorViewportSettings* Settings =
+													GetDefault<ULevelEditorViewportSettings>();
+												return Settings->MarqueeSelectionMode
+													== static_cast<EMarqueeSelectionMode>(Enum->GetValueByIndex(Index));
+											}
+											return false;
+										}
+									)
+								),
+								EUserInterfaceActionType::RadioButton
+							);
+						}
+					}
+				)
+
+			);
+			MarqueeMode.Icon = FSlateIcon(FAppStyle::GetAppStyleSetName(), "LevelEditor.StrictBoxSelect");
+			MarqueeMode.InsertPosition = FToolMenuInsert("TransparentBoxSelect", EToolMenuInsertType::Before);
+			MarqueeMode.Visibility = TAttribute<bool>::CreateLambda([]()
+			{
+				// ITF Drag Tools are automatically enabled by new TRS Gizmos
+				return UEditorInteractiveGizmoManager::UsesNewTRSGizmos() && GLevelEditorModeTools().GetInteractiveToolsContext()->UsesDragTools();
+			});
+		}
 	}
 }
 
diff --git a/Engine/Source/Editor/UnrealEd/Private/Tools/EdModeInteractiveToolsContext.cpp b/Engine/Source/Editor/UnrealEd/Private/Tools/EdModeInteractiveToolsContext.cpp
--- a/Engine/Source/Editor/UnrealEd/Private/Tools/EdModeInteractiveToolsContext.cpp
+++ b/Engine/Source/Editor/UnrealEd/Private/Tools/EdModeInteractiveToolsContext.cpp
@@ -1474,74 +1474,6 @@ void UModeManagerInteractiveToolsContext::RegisterDragTools()
 		DragToolsBehaviorSource = NewObject<UDragToolsBehaviorSource>();
 		DragToolsBehaviorSource->Initialize(this);
 		DragToolsBehaviorSource->RegisterSource();
-		
-		{
-			// Register menu extensions
-        	FToolMenuOwnerScoped Owner(this);
-        	
-        	UToolMenu* TransformMenu = UToolMenus::Get()->ExtendMenu("LevelEditor.ViewportToolbar.Transform");
-			FToolMenuSection& SelectionSection = TransformMenu->FindOrAddSection("Selection");
-			
-			FToolMenuEntry& MarqueeMode = SelectionSection.AddSubMenu(
-				"StrictBoxSelect", // Overwrite the other setting
-				LOCTEXT("MarqueeSelectionLabel", "Marquee Selection"),
-				LOCTEXT("MarqueeSelectionTooltip", "Define how objects are selected by marquee selection"),
-				FNewToolMenuDelegate::CreateLambda([](UToolMenu* Menu)
-				{
-					FProperty* Property = ULevelEditorViewportSettings::StaticClass()->FindPropertyByName(GET_MEMBER_NAME_CHECKED(ULevelEditorViewportSettings, MarqueeSelectionMode));
-					if (!Property || !Property->IsA<FEnumProperty>())
-					{
-						return;
-					}
-					
-					FEnumProperty* EnumProperty = StaticCast<FEnumProperty*>(Property);
-					UEnum* Enum = EnumProperty->GetEnum();
-					if (!Enum)
-					{
-						return;
-					}
-					
-					TWeakObjectPtr<UEnum> WeakEnum = Enum;
-					
-					FToolMenuSection& Section = Menu->AddSection("Mode", LOCTEXT("MarqueeModeSection", "Mode"));
-					
-					const int32 MaxIndex = Enum->NumEnums() - (Enum->ContainsExistingMax() ? 1 : 0);
-					for (int32 Index = 0; Index < MaxIndex; Index++)
-					{
-						Section.AddMenuEntry(
-							NAME_None,
-							Enum->GetDisplayNameTextByIndex(Index),
-							Enum->GetToolTipTextByIndex(Index),
-							FSlateIcon(),
-							FUIAction(
-								FExecuteAction::CreateLambda([WeakEnum, Index]
-								{
-									if (UEnum* Enum = WeakEnum.Get())
-									{
-										ULevelEditorViewportSettings* Settings = GetMutableDefault<ULevelEditorViewportSettings>();
-										Settings->MarqueeSelectionMode = static_cast<EMarqueeSelectionMode>(Enum->GetValueByIndex(Index));
-										Settings->SaveConfig();
-									}
-								}),
-								FCanExecuteAction(),
-								FIsActionChecked::CreateLambda([WeakEnum, Index]
-								{
-									if (UEnum* Enum = WeakEnum.Get())
-									{
-										const ULevelEditorViewportSettings* Settings = GetDefault<ULevelEditorViewportSettings>();
-										return Settings->MarqueeSelectionMode == static_cast<EMarqueeSelectionMode>(Enum->GetValueByIndex(Index));
-									}
-									return false;
-								})
-							),
-							EUserInterfaceActionType::RadioButton
-						);
-					}
-				})
-			);
-			MarqueeMode.Icon = FSlateIcon(FAppStyle::GetAppStyleSetName(), "LevelEditor.StrictBoxSelect");
-			MarqueeMode.InsertPosition = FToolMenuInsert("TransparentBoxSelect", EToolMenuInsertType::Before);
-		}
 	}
 }
 
