[Gizmos] Fix an issue with scale snapping when Preserve Aspect Ratios is set to true.

#FailedValidations Preflight

diff --git a/Engine/Source/Editor/LevelEditor/Private/Elements/Component/ComponentElementLevelEditorViewportInteractionCustomization.cpp b/Engine/Source/Editor/LevelEditor/Private/Elements/Component/ComponentElementLevelEditorViewportInteractionCustomization.cpp
--- a/Engine/Source/Editor/LevelEditor/Private/Elements/Component/ComponentElementLevelEditorViewportInteractionCustomization.cpp
+++ b/Engine/Source/Editor/LevelEditor/Private/Elements/Component/ComponentElementLevelEditorViewportInteractionCustomization.cpp
@@ -184,7 +184,7 @@ void FComponentElementLevelEditorViewportInteractionCustomization::ValidateScale
 	* @param	InOutScaleDelta				The scale delta we are potentially transforming.
 	* @return true if the axes should be snapped individually, according to the snap setting (i.e. this function had no effect)
 	*/
-	auto ApplyScalingOptions = [](const FVector& InOriginalPreDragScale, const bool bActiveAxes[3], FVector& InOutScaleDelta)
+	auto ApplyScalingOptions = [](const FVector& InOriginalPreDragScale, const FVector& InCurrentScale, const bool bActiveAxes[3], FVector& InOutScaleDelta)
 	{
 		int ActiveAxisCount = 0;
 		bool CurrentValueSameSign = true;
@@ -253,10 +253,10 @@ void FComponentElementLevelEditorViewportInteractionCustomization::ValidateScale
 
 			check(MaxAxisIndex != -1);
 
-			FVector::FReal AbsoluteScaleValue = InOriginalPreDragScale[MaxAxisIndex] + InOutScaleDelta[MaxAxisIndex];
+			FVector::FReal AbsoluteScaleValue = InCurrentScale[MaxAxisIndex] + InOutScaleDelta[MaxAxisIndex];
 			if (ViewportSettings->SnapScaleEnabled)
 			{
-				AbsoluteScaleValue = FMath::GridSnap(InOriginalPreDragScale[MaxAxisIndex] + InOutScaleDelta[MaxAxisIndex], (FVector::FReal)GEditor->GetScaleGridSize());
+				AbsoluteScaleValue = FMath::GridSnap(AbsoluteScaleValue, (FVector::FReal)GEditor->GetScaleGridSize());
 				SnapScaleAfter = false;
 			}
 
@@ -265,7 +265,7 @@ void FComponentElementLevelEditorViewportInteractionCustomization::ValidateScale
 			{
 				if (bActiveAxes[Axis])
 				{
-					InOutScaleDelta[Axis] = (InOriginalPreDragScale[Axis] * ScaleRatioMax) - InOriginalPreDragScale[Axis];
+					InOutScaleDelta[Axis] = (InOriginalPreDragScale[Axis] * ScaleRatioMax) - InCurrentScale[Axis];
 				}
 			}
 		}
@@ -281,7 +281,7 @@ void FComponentElementLevelEditorViewportInteractionCustomization::ValidateScale
 
 	//When scaling with more then one active axis, We must make sure we apply the correct delta sign to each delta scale axis
 	//We also want to support the PreserveNonUniformScale option
-	bool bSnapAxes = ApplyScalingOptions(InOriginalPreDragScale, bActiveAxes, InOutScaleDelta);
+	bool bSnapAxes = ApplyScalingOptions(InOriginalPreDragScale, InCurrentScale, bActiveAxes, InOutScaleDelta);
 
 	// check each axis
 	for (int Axis = 0; Axis < 3; ++Axis)
