[ITF Interactions] Clicking on Level sequence handles in Level Viewport crashes the Editor or triggers a check

Exporting HLevelSocketProxy so that we can properly parse proxies in `UE::Editor::LevelViewportInteractions::ProcessClick` and possibly other places in the future.

diff --git a/Engine/Source/Editor/LevelEditor/Private/ViewportInteractions/LevelViewportClickSelection.cpp b/Engine/Source/Editor/LevelEditor/Private/ViewportInteractions/LevelViewportClickSelection.cpp
--- a/Engine/Source/Editor/LevelEditor/Private/ViewportInteractions/LevelViewportClickSelection.cpp
+++ b/Engine/Source/Editor/LevelEditor/Private/ViewportInteractions/LevelViewportClickSelection.cpp
@@ -85,8 +85,7 @@ void ProcessClick(FEditorViewportClient* InEditorViewportClient, const FViewport
 			UE::Editor::ViewportSelectionUtilities::ClickSurface(InEditorViewportClient, ModelHit->GetModel(),SurfaceIndex,InViewportClick);
 		}
 	}
-	//TODO: check if this case is ever hit (cannot use StaticGetType as above, since it's not exported from LevelEditor module)
-	else if (static_cast<HLevelSocketProxy*>(InHitProxy))
+	else if (InHitProxy->IsA(HLevelSocketProxy::StaticGetType()))
 	{
 		UE::Editor::ViewportSelectionUtilities::ClickLevelSocket(InEditorViewportClient, InHitProxy, InViewportClick);
 	}
diff --git a/Engine/Source/Editor/UnrealEd/Public/ViewportSelectionUtilities.h b/Engine/Source/Editor/UnrealEd/Public/ViewportSelectionUtilities.h
--- a/Engine/Source/Editor/UnrealEd/Public/ViewportSelectionUtilities.h
+++ b/Engine/Source/Editor/UnrealEd/Public/ViewportSelectionUtilities.h
@@ -24,7 +24,7 @@ struct HGeomVertexProxy;
  */
 struct HLevelSocketProxy : public HHitProxy
 {
-	DECLARE_HIT_PROXY();
+	DECLARE_HIT_PROXY(UNREALED_API);
 
 	TObjectPtr<AActor> Actor;
 	TObjectPtr<USceneComponent> SceneComponent;
