BSP Brush is rebuilt after being moved with the new TRS Gizmo

diff --git a/Engine/Source/Editor/LevelEditor/Private/Elements/Actor/ActorElementLevelEditorViewportInteractionCustomization.cpp b/Engine/Source/Editor/LevelEditor/Private/Elements/Actor/ActorElementLevelEditorViewportInteractionCustomization.cpp
--- a/Engine/Source/Editor/LevelEditor/Private/Elements/Actor/ActorElementLevelEditorViewportInteractionCustomization.cpp
+++ b/Engine/Source/Editor/LevelEditor/Private/Elements/Actor/ActorElementLevelEditorViewportInteractionCustomization.cpp
@@ -78,6 +78,11 @@ void FActorElementLevelEditorViewportInteractionCustomization::GizmoManipulation
 		FCoreUObjectDelegates::OnObjectPropertyChanged.Broadcast(Actor, PropertyChangedEvent);
 	}
 
+	if (Actor->IsA(ABrush::StaticClass()))
+	{
+		GEditor->RebuildAlteredBSP();
+	} 
+
 	Actor->PostEditMove(true);
 	GEditor->BroadcastEndObjectMovement(*Actor);
 }
diff --git a/Engine/Source/Editor/UnrealEd/Private/LevelEditorViewport.cpp b/Engine/Source/Editor/UnrealEd/Private/LevelEditorViewport.cpp
--- a/Engine/Source/Editor/UnrealEd/Private/LevelEditorViewport.cpp
+++ b/Engine/Source/Editor/UnrealEd/Private/LevelEditorViewport.cpp
@@ -3852,6 +3852,8 @@ void FLevelEditorViewportClient::TrackingStopped()
 		bIsTrackingBrushModification = false;
 		if ( bDidAnythingActuallyChange && bWidgetAxisControlledByDrag )
 		{
+			// Note that while the brush is now rebuilt in FActorElementLevelEditorViewportInteractionCustomization::GizmoManipulationStopped,
+			// we still call it here as the call above has different conditions to this one. There's a re-entrant guard in RebuildAlteredBSP to prevent a second rebuild. 
 			GEditor->RebuildAlteredBSP();
 		}
 	}
