EditorUsability : UnrealEd
- Added browse to asset when right clicking the tab of an editor asset tab, action is only enabled for assets (non transient)
- Added browse to level when right clicking the tab of an editor level tab, action is only enabled for saved level
- Reordering includes alphabetically

diff --git a/Engine/Source/Editor/LevelEditor/Private/LevelEditor.cpp b/Engine/Source/Editor/LevelEditor/Private/LevelEditor.cpp
--- a/Engine/Source/Editor/LevelEditor/Private/LevelEditor.cpp
+++ b/Engine/Source/Editor/LevelEditor/Private/LevelEditor.cpp
@@ -1,43 +1,46 @@
 // Copyright Epic Games, Inc. All Rights Reserved.
 
 #include "LevelEditor.h"
-#include "Model.h"
-#include "Widgets/Text/STextBlock.h"
-#include "Widgets/Layout/SBorder.h"
-#include "Misc/App.h"
-#include "Modules/ModuleManager.h"
-#include "Layout/WidgetPath.h"
-#include "Framework/Application/SlateApplication.h"
-#include "Widgets/Layout/SBox.h"
-#include "Styling/AppStyle.h"
+
+#include "CommonMenuExtensionsModule.h"
+#include "ContentBrowserModule.h"
+#include "DataDrivenShaderPlatformInfo.h"
 #include "Editor/EditorPerProjectUserSettings.h"
 #include "Editor/UnrealEdEngine.h"
-#include "LevelEditorViewport.h"
+#include "Editor.h"
 #include "EditorModes.h"
-#include "UnrealEdGlobals.h"
-#include "Misc/ConfigCacheIni.h"
+#include "EditorViewportCommands.h"
+#include "Framework/Application/SlateApplication.h"
+#include "IContentBrowserSingleton.h"
+#include "Interfaces/IProjectManager.h"
+#include "ISlateReflectorModule.h"
+#include "Layout/WidgetPath.h"
+#include "LevelEditorViewport.h"
+#include "LevelViewportActions.h"
 #include "LevelViewportLayout.h"
-#include "SLevelEditor.h"
-#include "LightmapResRatioAdjust.h"
+#include "LevelViewportLayoutEntity.h"
 #include "LevelEditorActions.h"
 #include "LevelEditorModesActions.h"
-#include "WorkspaceMenuStructure.h"
-#include "WorkspaceMenuStructureModule.h"
+#include "LightmapResRatioAdjust.h"
 #include "MessageLogModule.h"
-#include "EditorViewportCommands.h"
-#include "LevelViewportActions.h"
-#include "Toolkits/GlobalEditorCommonCommands.h"
-#include "ISlateReflectorModule.h"
-#include "Widgets/Docking/SDockTab.h"
-#include "Interfaces/IProjectManager.h"
-#include "LevelViewportLayoutEntity.h"
+#include "Misc/App.h"
+#include "Misc/ConfigCacheIni.h"
+#include "Model.h"
+#include "Modules/ModuleManager.h"
 #include "PixelInspectorModule.h"
-#include "CommonMenuExtensionsModule.h"
-#include "ProjectDescriptor.h"
 #include "PlatformInfo.h"
-#include "Editor.h"
+#include "ProjectDescriptor.h"
+#include "SLevelEditor.h"
 #include "SLevelViewport.h"
-#include "DataDrivenShaderPlatformInfo.h"
+#include "Styling/AppStyle.h"
+#include "Toolkits/GlobalEditorCommonCommands.h"
+#include "UnrealEdGlobals.h"
+#include "Widgets/Docking/SDockTab.h"
+#include "Widgets/Layout/SBorder.h"
+#include "Widgets/Layout/SBox.h"
+#include "Widgets/Text/STextBlock.h"
+#include "WorkspaceMenuStructure.h"
+#include "WorkspaceMenuStructureModule.h"
 
 // @todo Editor: remove this circular dependency
 #include "Interfaces/IMainFrameModule.h"
@@ -222,6 +225,7 @@
 	LLM_SCOPE_BYTAG(LevelEditor);
 	
 	TSharedRef<SDockTab> LevelEditorTab = SNew(SDockTab)
+		.OnExtendContextMenu(SDockTab::FExtendContextMenu::CreateRaw(this, &FLevelEditorModule::OnExtendTabContextMenu))
 		.TabRole(ETabRole::MajorTab)
 		.ContentPadding(FMargin(0))
 		.IconColor(FAppStyle::Get().GetColor("LevelEditor.AssetColor")); // Same color as FAssetTypeActions_World
@@ -275,6 +279,57 @@
 	return LevelEditorTab;
 }
 
+void FLevelEditorModule::OnExtendTabContextMenu(FMenuBuilder& InMenuBuilder)
+{
+	InMenuBuilder.BeginSection("DockingTabLevelOptions", LOCTEXT("LevelTabSection", "Level"));
+	{
+		InMenuBuilder.AddMenuEntry(
+			LOCTEXT("BrowseLevelInContentBrowser", "Browse to Level"),
+			LOCTEXT("BrowseLevelInContentBrowserTooltip", "Browse to level in content browser"),
+			FSlateIcon(FAppStyle::GetAppStyleSetName(), "Icons.BrowseContent"),
+			FUIAction(
+				FExecuteAction::CreateRaw(this, &FLevelEditorModule::OnBrowseLevelInContentBrowser),
+				FCanExecuteAction::CreateRaw(this, &FLevelEditorModule::OnCanBrowseLevelInContentBrowser)
+			)
+		);
+	}
+	InMenuBuilder.EndSection();
+}
 
+bool FLevelEditorModule::OnCanBrowseLevelInContentBrowser()
+{
+	if (const TSharedPtr<ILevelEditor> LevelEditor = GetLevelEditorInstance().Pin())
+	{
+		const UWorld* World = LevelEditor->GetWorld();
+
+		return !!World
+			&& !!World->PersistentLevel
+			&& !!World->PersistentLevel->GetPackage()
+			&& !World->PersistentLevel->GetPackage()->GetName().StartsWith("/Temp/");
+	}
+
+	return false;
+}
+
+void FLevelEditorModule::OnBrowseLevelInContentBrowser()
+{
+	if (const TSharedPtr<ILevelEditor> LevelEditor = GetLevelEditorInstance().Pin())
+	{
+		UWorld* World = LevelEditor->GetWorld();
+
+		if (!!World
+			&& !!World->PersistentLevel
+			&& !!World->PersistentLevel->GetPackage()
+			&& !World->PersistentLevel->GetPackage()->GetName().StartsWith("/Temp/"))
+		{
+			if (const FContentBrowserModule* ContentBrowserModule = FModuleManager::GetModulePtr<FContentBrowserModule>("ContentBrowser"))
+			{
+				ContentBrowserModule->Get().SyncBrowserToAssets({World->PersistentLevel});
+			}
+		}
+	}
+}
+
+
 /**
  * Called right after the module's DLL has been loaded and the module object has been created
diff --git a/Engine/Source/Editor/LevelEditor/Public/LevelEditor.h b/Engine/Source/Editor/LevelEditor/Public/LevelEditor.h
--- a/Engine/Source/Editor/LevelEditor/Public/LevelEditor.h
+++ b/Engine/Source/Editor/LevelEditor/Public/LevelEditor.h
@@ -336,7 +336,7 @@ public:
 	DECLARE_EVENT( FLevelEditorModule, FTakeHighResScreenShotsEvent );
 	virtual FTakeHighResScreenShotsEvent& OnTakeHighResScreenShots() { return TakeHighResScreenShotsEvent; }
 
-	/** Delegate used to capture skeltal meshes to single-frame animations when 'keeping simulation changes' */
+	/** Delegate used to capture skeletal meshes to single-frame animations when 'keeping simulation changes' */
 	DECLARE_DELEGATE_RetVal_OneParam(UAnimSequence*, FCaptureSingleFrameAnimSequence, USkeletalMeshComponent* /*Component*/);
 	virtual FCaptureSingleFrameAnimSequence& OnCaptureSingleFrameAnimSequence() { return CaptureSingleFrameAnimSequenceDelegate; }
 	
@@ -440,6 +440,15 @@ private:
 	/** Spawn the main level editor tab */
 	TSharedRef<SDockTab> SpawnLevelEditor( const FSpawnTabArgs& InArgs );
 
+	/** Extend the dock tab context menu */
+	void OnExtendTabContextMenu(FMenuBuilder& InMenuBuilder);
+
+	/** Is the active level an asset that can be browsed to */
+	bool OnCanBrowseLevelInContentBrowser();
+
+	/** Browse to active level in content browser */
+	void OnBrowseLevelInContentBrowser();
+
 private:
 	TSharedPtr<FExtensibilityManager> MenuExtensibilityManager;
 	TSharedPtr<FExtensibilityManager> ToolBarExtensibilityManager;
@@ -487,7 +496,7 @@ private:
 	/** Multicast delegate executed when a map is changed (loaded,saved,new map, etc) */
 	FMapChangedEvent MapChangedEvent;
 
-	/** Delegate used to capture skeltal meshes to single-frame animations when 'keeping simulation changes' */
+	/** Delegate used to capture skeletal meshes to single-frame animations when 'keeping simulation changes' */
 	FCaptureSingleFrameAnimSequence CaptureSingleFrameAnimSequenceDelegate;
 
 	/** All extender delegates for the level viewport menus */
diff --git a/Engine/Source/Editor/UnrealEd/Private/Toolkits/AssetEditorToolkit.cpp b/Engine/Source/Editor/UnrealEd/Private/Toolkits/AssetEditorToolkit.cpp
--- a/Engine/Source/Editor/UnrealEd/Private/Toolkits/AssetEditorToolkit.cpp
+++ b/Engine/Source/Editor/UnrealEd/Private/Toolkits/AssetEditorToolkit.cpp
@@ -1,45 +1,48 @@
 // Copyright Epic Games, Inc. All Rights Reserved.
 
 #include "Toolkits/AssetEditorToolkit.h"
-#include "Widgets/Layout/SBorder.h"
-#include "Framework/MultiBox/MultiBoxDefs.h"
-#include "Framework/MultiBox/MultiBoxBuilder.h"
-#include "GameFramework/Actor.h"
+
+#include "AssetEditorModeManager.h"
+#include "AssetToolsModule.h"
+#include "CollectionManagerModule.h"
+#include "CollectionManagerTypes.h"
+#include "ContentBrowserModule.h"
 #include "Editor.h"
-#include "Misc/ConfigCacheIni.h"
-#include "Modules/ModuleManager.h"
-#include "Styling/AppStyle.h"
-#include "Settings/EditorStyleSettings.h"
 #include "EditorReimportHandler.h"
 #include "FileHelpers.h"
-#include "Toolkits/SStandaloneAssetEditorToolkitHost.h"
-#include "Toolkits/ToolkitManager.h"
-#include "Toolkits/AssetEditorCommonCommands.h"
-#include "Toolkits/GlobalEditorCommonCommands.h"
-#include "Styling/SlateIconFinder.h"
-#include "CollectionManagerTypes.h"
+#include "IAssetTools.h"
+#include "IAssetTypeActions.h"
 #include "ICollectionContainer.h"
 #include "ICollectionManager.h"
 #include "ICollectionSource.h"
-#include "CollectionManagerModule.h"
-#include "Widgets/SToolTip.h"
+#include "IContentBrowserSingleton.h"
 #include "IDocumentation.h"
-#include "Widgets/Docking/SDockTab.h"
-#include "IAssetTools.h"
-#include "IAssetTypeActions.h"
-#include "AssetToolsModule.h"
-#include "Toolkits/AssetEditorToolkitMenuContext.h"
+#include "StatusBarSubsystem.h"
 #include "ToolMenus.h"
-#include "Subsystems/AssetEditorSubsystem.h"
-#include "Logging/LogMacros.h"
-#include "AssetEditorModeManager.h"
-#include "Misc/Attribute.h"
-#include "Textures/SlateIcon.h"
 #include "WidgetDrawerConfig.h"
 #include "Framework/Commands/GenericCommands.h"
+#include "Framework/MultiBox/MultiBoxBuilder.h"
+#include "Framework/MultiBox/MultiBoxDefs.h"
+#include "GameFramework/Actor.h"
 #include "Interfaces/Interface_AsyncCompilation.h"
+#include "Logging/LogMacros.h"
+#include "Misc/Attribute.h"
+#include "Misc/ConfigCacheIni.h"
+#include "Modules/ModuleManager.h"
+#include "Settings/EditorStyleSettings.h"
+#include "Styling/AppStyle.h"
+#include "Styling/SlateIconFinder.h"
+#include "Subsystems/AssetEditorSubsystem.h"
+#include "Textures/SlateIcon.h"
+#include "Toolkits/AssetEditorCommonCommands.h"
+#include "Toolkits/AssetEditorToolkitMenuContext.h"
+#include "Toolkits/GlobalEditorCommonCommands.h"
+#include "Toolkits/SStandaloneAssetEditorToolkitHost.h"
+#include "Toolkits/ToolkitManager.h"
+#include "Widgets/SToolTip.h"
+#include "Widgets/Docking/SDockTab.h"
 #include "Widgets/Images/SImage.h"
-#include "StatusBarSubsystem.h"
+#include "Widgets/Layout/SBorder.h"
 
 #define LOCTEXT_NAMESPACE "AssetEditorToolkit"
 
@@ -140,6 +143,7 @@
 
 		// Create a new SlateToolkitHost
 		NewMajorTab = SNew(SDockTab)
+			.OnExtendContextMenu(SDockTab::FExtendContextMenu::CreateSP(this, &FAssetEditorToolkit::OnExtendTabContextMenu))
 			.ContentPadding(0.0f)
 			.TabRole(ETabRole::MajorTab)
 			.ToolTip(IDocumentation::Get()->CreateToolTip(ToolTipText, nullptr, DocLink, GetToolkitFName().ToString()))
@@ -923,6 +927,47 @@
 	return FReply::Unhandled();
 }
 
+void FAssetEditorToolkit::OnExtendTabContextMenu(FMenuBuilder& InMenuBuilder)
+{
+	InMenuBuilder.BeginSection("DockingTabAssetOptions", LOCTEXT("AssetTabSection", "Assets"));
+	{
+		InMenuBuilder.AddMenuEntry(
+			LOCTEXT("BrowseAssetInContentBrowser", "Browse to Asset"),
+			LOCTEXT("BrowseAssetInContentBrowserTooltip", "Browse to asset in content browser"),
+			FSlateIcon(FAppStyle::GetAppStyleSetName(), "Icons.BrowseContent"),
+			FUIAction(
+				FExecuteAction::CreateSP(this, &FAssetEditorToolkit::OnBrowseAssetInContentBrowser),
+				FCanExecuteAction::CreateSP(this, &FAssetEditorToolkit::OnCanBrowseAssetInContentBrowser)
+			)
+		);
+	}
+	InMenuBuilder.EndSection();
+}
+
+bool FAssetEditorToolkit::OnCanBrowseAssetInContentBrowser() const
+{
+	for (const TObjectPtr<UObject>& EditingObject : EditingObjects)
+	{
+		if (EditingObject && EditingObject->IsAsset())
+		{
+			return true;
+		}
+	}
+
+	return false;
+}
+
+void FAssetEditorToolkit::OnBrowseAssetInContentBrowser()
+{
+	if (HasEditingObject())
+	{
+		if (const FContentBrowserModule* ContentBrowserModule = FModuleManager::GetModulePtr<FContentBrowserModule>("ContentBrowser"))
+		{
+			ContentBrowserModule->Get().SyncBrowserToAssets(EditingObjects);
+		}
+	}
+}
+
 void FAssetEditorToolkit::RemoveEditingAsset(UObject* Asset)
 {
 	// Just close the editor tab if it's the last element
diff --git a/Engine/Source/Editor/UnrealEd/Public/Toolkits/AssetEditorToolkit.h b/Engine/Source/Editor/UnrealEd/Public/Toolkits/AssetEditorToolkit.h
--- a/Engine/Source/Editor/UnrealEd/Public/Toolkits/AssetEditorToolkit.h
+++ b/Engine/Source/Editor/UnrealEd/Public/Toolkits/AssetEditorToolkit.h
@@ -320,6 +320,21 @@ public:
 protected:
 	friend class UAssetEditorToolkitMenuContext;
 
+	/**
+	 * Allows to extend dock tab context menu for an asset editor
+	 */
+	UNREALED_API virtual void OnExtendTabContextMenu(FMenuBuilder& InMenuBuilder);
+
+	/**
+	 * Are there any assets to browse to
+	 */
+	bool OnCanBrowseAssetInContentBrowser() const;
+
+	/**
+	 * Browse to assets in content browser
+	 */
+	void OnBrowseAssetInContentBrowser();
+
 	/**
 	 * Perform any initialization that should happen after the basic toolkit needs are created for the asset editor.
 	 * For example, if you would like to activate a mode that is not the default from a mode manager, do that here.
