EditorUsability : LevelEditor
- Added Open Screenshot Folder button to High Resolution Screenshot tool
- Moved warning box above buttons to fill available space
- Resized window to take into account warning box message
- Refactoring window layout to avoid breaking content when resized to minimum
- Reordering includes alphabetically

diff --git a/Engine/Source/Editor/LevelEditor/Private/HighResScreenshotUI.cpp b/Engine/Source/Editor/LevelEditor/Private/HighResScreenshotUI.cpp
--- a/Engine/Source/Editor/LevelEditor/Private/HighResScreenshotUI.cpp
+++ b/Engine/Source/Editor/LevelEditor/Private/HighResScreenshotUI.cpp
@@ -1,22 +1,25 @@
 // Copyright Epic Games, Inc. All Rights Reserved.
 
 #include "HighresScreenshotUI.h"
+
 #include "Engine/World.h"
-#include "Framework/Docking/TabManager.h"
 #include "Framework/Application/SlateApplication.h"
-#include "Widgets/Layout/SBorder.h"
+#include "Framework/Docking/TabManager.h"
+#include "Styling/AppStyle.h"
+#include "Styling/StyleColors.h"
+#include "SWarningOrErrorBox.h"
 #include "Widgets/Images/SImage.h"
-#include "Widgets/Layout/SGridPanel.h"
 #include "Widgets/Input/SButton.h"
-#include "Widgets/Layout/SSeparator.h"
 #include "Widgets/Input/SSpinBox.h"
-#include "SWarningOrErrorBox.h"
-#include "Styling/StyleColors.h"
-#include "Styling/AppStyle.h"
+#include "Widgets/Layout/SBorder.h"
+#include "Widgets/Layout/SGridPanel.h"
+#include "Widgets/Layout/SScrollBox.h"
+#include "Widgets/Layout/SSeparator.h"
+#include "Widgets/Layout/SSpacer.h"
 
 #if !(UE_BUILD_SHIPPING || UE_BUILD_TEST)
-#include "Kismet/KismetSystemLibrary.h"
 #include "Engine/World.h"
+#include "Kismet/KismetSystemLibrary.h"
 #endif
 
 TWeakPtr<class SWindow> SHighResScreenshotDialog::CurrentWindow = NULL;
@@ -30,8 +33,9 @@ SHighResScreenshotDialog::SHighResScreenshotDialog()
 
 void SHighResScreenshotDialog::Construct( const FArguments& InArgs )
 {
+	const FMargin GridPadding(6.f, 3.f);
+	constexpr float WindowBorderPadding = 16.f;
 
-	FMargin GridPadding(6.f, 3.f);
 	this->ChildSlot
 	.Padding(0.f)
 	[
@@ -39,207 +43,236 @@ void SHighResScreenshotDialog::Construct( const FArguments& InArgs )
 		.Padding(0.f)
 		.BorderImage(FAppStyle::Get().GetBrush("ToolPanel.GroupBorder"))
 		[
-			SNew(SVerticalBox)
-			+SVerticalBox::Slot()
-			.AutoHeight()
-			[
-				SNew(SHorizontalBox)
-				+SHorizontalBox::Slot()
-				.AutoWidth()
-				.Padding(6.f)
-				[
-					SAssignNew(CaptureRegionButton, SButton)
-					.ButtonStyle(&FAppStyle::Get().GetWidgetStyle<FButtonStyle>("SimpleButton"))
-					.IsEnabled(this, &SHighResScreenshotDialog::IsCaptureRegionEditingAvailable)
-					.Visibility(this, &SHighResScreenshotDialog::GetSpecifyCaptureRegionVisibility)
-					.ToolTipText(NSLOCTEXT("HighResScreenshot", "ScreenshotSpecifyCaptureRectangleTooltip", "Specify the region which will be captured by the screenshot"))
-					.OnClicked(this, &SHighResScreenshotDialog::OnSelectCaptureRegionClicked)
-					.ContentPadding(4.f)
-					[
-						SNew(SImage)
-						.ColorAndOpacity(FSlateColor::UseForeground())
-						.Image(FAppStyle::Get().GetBrush("Icons.Crop"))
-					]
-				]
-				+SHorizontalBox::Slot()
-				.AutoWidth()
-				.Padding(6.f)
-				[
-					SNew( SButton )
-					.ButtonStyle(&FAppStyle::Get().GetWidgetStyle<FButtonStyle>("SimpleButton"))
-					.Visibility(this, &SHighResScreenshotDialog::GetCaptureRegionControlsVisibility)
-					.ToolTipText( NSLOCTEXT("HighResScreenshot", "ScreenshotAcceptCaptureRegionTooltip", "Accept any changes made to the capture region") )
-					.OnClicked( this, &SHighResScreenshotDialog::OnSelectCaptureAcceptRegionClicked )
-					.ContentPadding(4.f)
-					[
-						SNew(SImage)
-						.ColorAndOpacity(FSlateColor::UseForeground())
-						.Image(FAppStyle::Get().GetBrush("Icons.Check"))
-					]
-				]
-				+SHorizontalBox::Slot()
-				.AutoWidth()
-				.Padding(6.f)
-				[
-					SNew( SButton )
-					.ButtonStyle(&FAppStyle::Get().GetWidgetStyle<FButtonStyle>("SimpleButton"))
-					.ToolTipText( NSLOCTEXT("HighResScreenshot", "ScreenshotDiscardCaptureRegionTooltip", "Discard any changes made to the capture region") )
-					.Visibility(this, &SHighResScreenshotDialog::GetCaptureRegionControlsVisibility)
-					.OnClicked( this, &SHighResScreenshotDialog::OnSelectCaptureCancelRegionClicked )
-					.ContentPadding(4.f)
-					[
-						SNew(SImage)
-						.ColorAndOpacity(FSlateColor::UseForeground())
-						.Image(FAppStyle::Get().GetBrush("Icons.X"))
-					]
-				]
-				+SHorizontalBox::Slot()
-				.AutoWidth()
-				.Padding(6.f)
-				[
-					SNew( SButton )
-					.ButtonStyle(&FAppStyle::Get().GetWidgetStyle<FButtonStyle>("SimpleButton"))
-					.ToolTipText( NSLOCTEXT("HighResScreenshot", "ScreenshotFullViewportCaptureRegionTooltip", "Set the capture rectangle to the whole viewport") )
-					.Visibility(this, &SHighResScreenshotDialog::GetCaptureRegionControlsVisibility)
-					.OnClicked( this, &SHighResScreenshotDialog::OnSetFullViewportCaptureRegionClicked )
-					.ContentPadding(4.f)
-					[
-						SNew(SImage)
-						.ColorAndOpacity(FSlateColor::UseForeground())
-						.Image(FAppStyle::Get().GetBrush("Icons.Fullscreen"))
-					]
-				]
-			]
-
-			+SVerticalBox::Slot()
-			.AutoHeight()
-			[
-				SNew(SSeparator)
-				.Thickness(1.f)
-				.SeparatorImage(FAppStyle::Get().GetBrush("Brushes.Background"))
-			]
-
-			+SVerticalBox::Slot()
-			.Padding(6.0)
-			.HAlign(HAlign_Fill)
-			[
-				// Row/Column
-				SNew(SGridPanel) 
-				.FillColumn(1, 1.0f)
-
-				+SGridPanel::Slot(0, 0)
-				.Padding(GridPadding)
-				.VAlign(VAlign_Center)
-				[
-					SNew( STextBlock )
-					.Text( NSLOCTEXT("HighResScreenshot", "ScreenshotSizeMultiplier", "Screenshot Size Multiplier") )
-				]
-				+SGridPanel::Slot(0, 1)
-				.Padding(GridPadding)
-				[
-					SNew( STextBlock )
-					.Text( NSLOCTEXT("HighResScreenshot", "UseDateTimeAsImageName", "Use Date & Timestamp as Image name") )
-				]
-				+ SGridPanel::Slot(0, 2)
-				.Padding(GridPadding)
-				[
-					SNew(STextBlock)
-					.Text(NSLOCTEXT("HighResScreenshot", "IncludeBufferVisTargets", "Include Buffer Visualization Targets"))
-				]
-				+ SGridPanel::Slot(0, 3)
-				.Padding(GridPadding)
-				[
-					SAssignNew(HDRLabel, STextBlock)
-					.Text(NSLOCTEXT("HighResScreenshot", "CaptureHDR", "Write HDR format visualization targets"))
-				]
-				+SGridPanel::Slot(0, 4)
-				.Padding(GridPadding)
-				[
-					SAssignNew(Force128BitRenderingLabel, STextBlock)
-					.Text(NSLOCTEXT("HighResScreenshot", "Force128BitPipeline", "Force 128-bit buffers for rendering pipeline"))
-				]
-				+SGridPanel::Slot(0, 5)
-				.Padding(GridPadding)
-				[
-					SNew( STextBlock )
-					.Text( NSLOCTEXT("HighResScreenshot", "UseCustomDepth", "Use custom depth as mask") )
-				]
-
-				+SGridPanel::Slot(1, 0)
-				.Padding(GridPadding)
-				.HAlign(HAlign_Fill)
-				[
-						SNew( SSpinBox<float> )
-						.MinValue(FHighResScreenshotConfig::MinResolutionMultipler)
-						.MaxValue(FHighResScreenshotConfig::MaxResolutionMultipler)
-						.Delta(1.0f)
-						.Value(this, &SHighResScreenshotDialog::GetResolutionMultiplierSlider)
-						.OnValueChanged(this, &SHighResScreenshotDialog::OnResolutionMultiplierSliderChanged)
-				]
-				+ SGridPanel::Slot(1, 1)
-				.Padding(GridPadding)
-				[
-					SNew(SCheckBox)
-					.OnCheckStateChanged(this, &SHighResScreenshotDialog::OnDateTimeBasedNamingEnabledChanged)
-					.IsChecked(this, &SHighResScreenshotDialog::GetDateTimeBasedNamingEnabled)
-				]
-				+SGridPanel::Slot(1, 2)
-				.Padding(GridPadding)
-				[
-					SNew( SCheckBox )
-					.OnCheckStateChanged(this, &SHighResScreenshotDialog::OnBufferVisualizationDumpEnabledChanged)
-					.IsChecked(this, &SHighResScreenshotDialog::GetBufferVisualizationDumpEnabled)
-				]
-				+ SGridPanel::Slot(1, 3)
-				.Padding(GridPadding)
-				[
-					SAssignNew(HDRCheckBox, SCheckBox)
-					.OnCheckStateChanged(this, &SHighResScreenshotDialog::OnHDREnabledChanged)
-					.IsChecked(this, &SHighResScreenshotDialog::GetHDRCheckboxUIState)
-				]
-				+ SGridPanel::Slot(1, 4)
-				.Padding(GridPadding)
-				[
-					SAssignNew(Force128BitRenderingCheckBox, SCheckBox)
-					.OnCheckStateChanged(this, &SHighResScreenshotDialog::OnForce128BitRenderingChanged)
-					.IsChecked(this, &SHighResScreenshotDialog::GetForce128BitRenderingCheckboxUIState)
-				]
-				+SGridPanel::Slot(1, 5)
-				.Padding(GridPadding)
-				[
-					SNew( SCheckBox )
-					.OnCheckStateChanged(this, &SHighResScreenshotDialog::OnMaskEnabledChanged)
-					.IsChecked(this, &SHighResScreenshotDialog::GetMaskEnabled)
-				]
-			]
-
-			+SVerticalBox::Slot()
-			.AutoHeight()
-			.Padding(16.f)
-			[
-
-				SNew(SHorizontalBox)
-				+SHorizontalBox::Slot()
-				[
-					SNew(SWarningOrErrorBox)
-					.Padding(FMargin(16.f, 13.f, 16.f, 13.f))
-					.Message( NSLOCTEXT("HighResScreenshot", "CaptureWarningText", "Large multipliers may cause the graphics driver to crash.  Please try using a lower multiplier.") )
-					.Visibility_Lambda( [this] () { return Config.ResolutionMultiplier >= 3. ? EVisibility::Visible : EVisibility::Hidden; })
-				]
-
-				+SHorizontalBox::Slot()
-				.VAlign(VAlign_Bottom)
-				.HAlign(HAlign_Right)
-				.AutoWidth()
-				.Padding(FMargin(24.f, 0.f, 0.f, 0.f))
-				[
-					SNew( SButton )
-					.ButtonStyle(&FAppStyle::Get().GetWidgetStyle<FButtonStyle>("PrimaryButton"))
-					.ToolTipText(NSLOCTEXT("HighResScreenshot", "ScreenshotCaptureTooltop", "Take a screenshot") )
-					.OnClicked(this, &SHighResScreenshotDialog::OnCaptureClicked )
-					.Text(NSLOCTEXT("HighResScreenshot", "CaptureCommit", "Capture"))
-				]
-			]
+		    SNew(SVerticalBox)
+		    + SVerticalBox::Slot()
+		    .FillHeight(1.0f)
+			.Padding(0.f)
+		    [
+		    	SNew(SScrollBox)
+		    	+ SScrollBox::Slot()
+		    	[
+		    		SNew(SVerticalBox)
+		    	    +SVerticalBox::Slot()
+		    	    .AutoHeight()
+		    	    [
+		    	    	SNew(SHorizontalBox)
+		    	    	+SHorizontalBox::Slot()
+		    	    	.AutoWidth()
+		    	    	.Padding(6.f)
+		    	    	[
+		    	    		SAssignNew(CaptureRegionButton, SButton)
+		    	    		.ButtonStyle(&FAppStyle::Get().GetWidgetStyle<FButtonStyle>("SimpleButton"))
+		    	    		.IsEnabled(this, &SHighResScreenshotDialog::IsCaptureRegionEditingAvailable)
+		    	    		.Visibility(this, &SHighResScreenshotDialog::GetSpecifyCaptureRegionVisibility)
+		    	    		.ToolTipText(NSLOCTEXT("HighResScreenshot", "ScreenshotSpecifyCaptureRectangleTooltip", "Specify the region which will be captured by the screenshot"))
+		    	    		.OnClicked(this, &SHighResScreenshotDialog::OnSelectCaptureRegionClicked)
+		    	    		.ContentPadding(4.f)
+		    	    		[
+		    	    			SNew(SImage)
+		    	    			.ColorAndOpacity(FSlateColor::UseForeground())
+		    	    			.Image(FAppStyle::Get().GetBrush("Icons.Crop"))
+		    	    		]
+		    	    	]
+		    	    	+SHorizontalBox::Slot()
+		    	    	.AutoWidth()
+		    	    	.Padding(6.f)
+		    	    	[
+		    	    		SNew( SButton )
+		    	    		.ButtonStyle(&FAppStyle::Get().GetWidgetStyle<FButtonStyle>("SimpleButton"))
+		    	    		.Visibility(this, &SHighResScreenshotDialog::GetCaptureRegionControlsVisibility)
+		    	    		.ToolTipText( NSLOCTEXT("HighResScreenshot", "ScreenshotAcceptCaptureRegionTooltip", "Accept any changes made to the capture region") )
+		    	    		.OnClicked( this, &SHighResScreenshotDialog::OnSelectCaptureAcceptRegionClicked )
+		    	    		.ContentPadding(4.f)
+		    	    		[
+		    	    			SNew(SImage)
+		    	    			.ColorAndOpacity(FSlateColor::UseForeground())
+		    	    			.Image(FAppStyle::Get().GetBrush("Icons.Check"))
+		    	    		]
+		    	    	]
+		    	    	+SHorizontalBox::Slot()
+		    	    	.AutoWidth()
+		    	    	.Padding(6.f)
+		    	    	[
+		    	    		SNew( SButton )
+		    	    		.ButtonStyle(&FAppStyle::Get().GetWidgetStyle<FButtonStyle>("SimpleButton"))
+		    	    		.ToolTipText( NSLOCTEXT("HighResScreenshot", "ScreenshotDiscardCaptureRegionTooltip", "Discard any changes made to the capture region") )
+		    	    		.Visibility(this, &SHighResScreenshotDialog::GetCaptureRegionControlsVisibility)
+		    	    		.OnClicked( this, &SHighResScreenshotDialog::OnSelectCaptureCancelRegionClicked )
+		    	    		.ContentPadding(4.f)
+		    	    		[
+		    	    			SNew(SImage)
+		    	    			.ColorAndOpacity(FSlateColor::UseForeground())
+		    	    			.Image(FAppStyle::Get().GetBrush("Icons.X"))
+		    	    		]
+		    	    	]
+		    	    	+SHorizontalBox::Slot()
+		    	    	.AutoWidth()
+		    	    	.Padding(6.f)
+		    	    	[
+		    	    		SNew( SButton )
+		    	    		.ButtonStyle(&FAppStyle::Get().GetWidgetStyle<FButtonStyle>("SimpleButton"))
+		    	    		.ToolTipText( NSLOCTEXT("HighResScreenshot", "ScreenshotFullViewportCaptureRegionTooltip", "Set the capture rectangle to the whole viewport") )
+		    	    		.Visibility(this, &SHighResScreenshotDialog::GetCaptureRegionControlsVisibility)
+		    	    		.OnClicked( this, &SHighResScreenshotDialog::OnSetFullViewportCaptureRegionClicked )
+		    	    		.ContentPadding(4.f)
+		    	    		[
+		    	    			SNew(SImage)
+		    	    			.ColorAndOpacity(FSlateColor::UseForeground())
+		    	    			.Image(FAppStyle::Get().GetBrush("Icons.Fullscreen"))
+		    	    		]
+		    	    	]
+		    	    ]
+		    	    
+		    	    +SVerticalBox::Slot()
+		    	    .AutoHeight()
+		    	    [
+		    	    	SNew(SSeparator)
+		    	    	.Thickness(1.f)
+		    	    	.SeparatorImage(FAppStyle::Get().GetBrush("Brushes.Background"))
+		    	    ]
+		    	    
+		    	    +SVerticalBox::Slot()
+		    		.Padding(6.0)
+		    	    .AutoHeight()
+		    	    [
+		    	    	SNew(SWarningOrErrorBox)
+		    	    	.Padding(FMargin(16.f, 13.f, 16.f, 13.f))
+		    	    	.Message(NSLOCTEXT("HighResScreenshot", "CaptureWarningText", "Large multipliers may cause the graphics driver to crash.  Please try using a lower multiplier."))
+		    	    	.Visibility(this, &SHighResScreenshotDialog::GetCaptureWarningVisibility)
+		    	    ]
+		    	    
+		    	    +SVerticalBox::Slot()
+		    	    .Padding(6.0)
+		    	    .HAlign(HAlign_Fill)
+		    	    [
+		    	    	// Row/Column
+		    	    	SNew(SGridPanel)
+		    	    	.FillColumn(1, 1.0f)
+		    	    
+		    	    	+SGridPanel::Slot(0, 0)
+		    	    	.Padding(GridPadding)
+		    	    	.VAlign(VAlign_Center)
+		    	    	[
+		    	    		SNew( STextBlock )
+		    	    		.Text( NSLOCTEXT("HighResScreenshot", "ScreenshotSizeMultiplier", "Screenshot Size Multiplier") )
+		    	    	]
+		    	    	+SGridPanel::Slot(0, 1)
+		    	    	.Padding(GridPadding)
+		    	    	[
+		    	    		SNew( STextBlock )
+		    	    		.Text( NSLOCTEXT("HighResScreenshot", "UseDateTimeAsImageName", "Use Date & Timestamp as Image name") )
+		    	    	]
+		    	    	+ SGridPanel::Slot(0, 2)
+		    	    	.Padding(GridPadding)
+		    	    	[
+		    	    		SNew(STextBlock)
+		    	    		.Text(NSLOCTEXT("HighResScreenshot", "IncludeBufferVisTargets", "Include Buffer Visualization Targets"))
+		    	    	]
+		    	    	+ SGridPanel::Slot(0, 3)
+		    	    	.Padding(GridPadding)
+		    	    	[
+		    	    		SAssignNew(HDRLabel, STextBlock)
+		    	    		.Text(NSLOCTEXT("HighResScreenshot", "CaptureHDR", "Write HDR format visualization targets"))
+		    	    	]
+		    	    	+SGridPanel::Slot(0, 4)
+		    	    	.Padding(GridPadding)
+		    	    	[
+		    	    		SAssignNew(Force128BitRenderingLabel, STextBlock)
+		    	    		.Text(NSLOCTEXT("HighResScreenshot", "Force128BitPipeline", "Force 128-bit buffers for rendering pipeline"))
+		    	    	]
+		    	    	+SGridPanel::Slot(0, 5)
+		    	    	.Padding(GridPadding)
+		    	    	[
+		    	    		SNew( STextBlock )
+		    	    		.Text( NSLOCTEXT("HighResScreenshot", "UseCustomDepth", "Use custom depth as mask") )
+		    	    	]
+		    	    
+		    	    	+SGridPanel::Slot(1, 0)
+		    	    	.Padding(GridPadding)
+		    	    	.HAlign(HAlign_Fill)
+		    	    	[
+		    	    		SNew( SSpinBox<float> )
+		    	    		.MinValue(FHighResScreenshotConfig::MinResolutionMultipler)
+		    	    		.MaxValue(FHighResScreenshotConfig::MaxResolutionMultipler)
+		    	    		.Delta(1.0f)
+		    	    		.Value(this, &SHighResScreenshotDialog::GetResolutionMultiplierSlider)
+		    	    		.OnValueChanged(this, &SHighResScreenshotDialog::OnResolutionMultiplierSliderChanged)
+		    	    	]
+		    	    	+ SGridPanel::Slot(1, 1)
+		    	    	.Padding(GridPadding)
+		    	    	[
+		    	    		SNew(SCheckBox)
+		    	    		.OnCheckStateChanged(this, &SHighResScreenshotDialog::OnDateTimeBasedNamingEnabledChanged)
+		    	    		.IsChecked(this, &SHighResScreenshotDialog::GetDateTimeBasedNamingEnabled)
+		    	    	]
+		    	    	+SGridPanel::Slot(1, 2)
+		    	    	.Padding(GridPadding)
+		    	    	[
+		    	    		SNew( SCheckBox )
+		    	    		.OnCheckStateChanged(this, &SHighResScreenshotDialog::OnBufferVisualizationDumpEnabledChanged)
+		    	    		.IsChecked(this, &SHighResScreenshotDialog::GetBufferVisualizationDumpEnabled)
+		    	    	]
+		    	    	+ SGridPanel::Slot(1, 3)
+		    	    	.Padding(GridPadding)
+		    	    	[
+		    	    		SAssignNew(HDRCheckBox, SCheckBox)
+		    	    		.OnCheckStateChanged(this, &SHighResScreenshotDialog::OnHDREnabledChanged)
+		    	    		.IsChecked(this, &SHighResScreenshotDialog::GetHDRCheckboxUIState)
+		    	    	]
+		    	    	+ SGridPanel::Slot(1, 4)
+		    	    	.Padding(GridPadding)
+		    	    	[
+		    	    		SAssignNew(Force128BitRenderingCheckBox, SCheckBox)
+		    	    		.OnCheckStateChanged(this, &SHighResScreenshotDialog::OnForce128BitRenderingChanged)
+		    	    		.IsChecked(this, &SHighResScreenshotDialog::GetForce128BitRenderingCheckboxUIState)
+		    	    	]
+		    	    	+SGridPanel::Slot(1, 5)
+		    	    	.Padding(GridPadding)
+		    	    	[
+		    	    		SNew( SCheckBox )
+		    	    		.OnCheckStateChanged(this, &SHighResScreenshotDialog::OnMaskEnabledChanged)
+		    	    		.IsChecked(this, &SHighResScreenshotDialog::GetMaskEnabled)
+		    	    	]
+		    	    ]
+		    	]
+		    ]
+		    + SVerticalBox::Slot()
+		    .AutoHeight()
+			.Padding(WindowBorderPadding)
+		    [
+		    	SNew(SHorizontalBox)
+		    	+SHorizontalBox::Slot()
+		    	.VAlign(VAlign_Bottom)
+		    	.HAlign(HAlign_Left)
+		    	.AutoWidth()
+		    	.Padding(FMargin(0.f, 0.f, 24.f, 0.f))
+		    	[
+		    		SNew(SButton)
+		    		.ButtonStyle(&FAppStyle::Get().GetWidgetStyle<FButtonStyle>("Button"))
+		    		.ToolTipText(NSLOCTEXT("HighResScreenshot", "OpenScreenshotFolderToolTip", "Open the screenshot folder in the file explorer") )
+		    		.OnClicked(this, &SHighResScreenshotDialog::OnOpenScreenshotFolder)
+		    		.Text(NSLOCTEXT("HighResScreenshot", "OpenScreenshotFolder", "Open Screenshot Folder"))
+		    	]
+		    	
+		    	+ SHorizontalBox::Slot()
+		    	.FillWidth(1.f)
+		    	[
+		    		SNew(SSpacer)
+		    	]
+		    	
+		    	+SHorizontalBox::Slot()
+		    	.VAlign(VAlign_Bottom)
+		    	.HAlign(HAlign_Right)
+		    	.AutoWidth()
+		    	.Padding(FMargin(24.f, 0.f, 0.f, 0.f))
+		    	[
+		    		SNew(SButton)
+		    		.ButtonStyle(&FAppStyle::Get().GetWidgetStyle<FButtonStyle>("PrimaryButton"))
+		    		.ToolTipText(NSLOCTEXT("HighResScreenshot", "ScreenshotCaptureToolTip", "Take a screenshot"))
+		    		.OnClicked(this, &SHighResScreenshotDialog::OnCaptureClicked)
+		    		.Text(NSLOCTEXT("HighResScreenshot", "CaptureCommit", "Capture"))
+		    	]
+		    ]			
 		]
 	];
 
@@ -327,7 +360,9 @@ TWeakPtr<class SWindow> SHighResScreenshotDialog::OpenDialog(const TSharedPtr<FS
 		TSharedRef<SHighResScreenshotDialog> Dialog = SNew(SHighResScreenshotDialog);
 		TSharedRef<SWindow> Window = SNew(SWindow)
 			.Title( NSLOCTEXT("HighResScreenshot", "HighResolutionScreenshot", "High Resolution Screenshot") )
-			.ClientSize(FVector2D(480, 302))
+			.ClientSize(FVector2D(480, 350))
+			.MinHeight(250.f)
+			.MinWidth(400.f)
 			.SupportsMinimize(false)
 			.SupportsMaximize(false)
 			.FocusWhenFirstShown(true)
@@ -392,6 +427,17 @@ FReply SHighResScreenshotDialog::OnSelectCaptureRegionClicked()
 	return FReply::Handled();
 }
 
+EVisibility SHighResScreenshotDialog::GetCaptureWarningVisibility() const
+{
+	return Config.ResolutionMultiplier >= 3. ? EVisibility::Visible : EVisibility::Collapsed;
+}
+
+FReply SHighResScreenshotDialog::OnOpenScreenshotFolder()
+{
+	FPlatformProcess::ExploreFolder(*FPaths::ConvertRelativePathToFull(FPaths::ScreenShotDir()));
+	return FReply::Handled();
+}
+
 FReply SHighResScreenshotDialog::OnCaptureClicked()
 {
 	if (!GIsHighResScreenshot)
diff --git a/Engine/Source/Editor/LevelEditor/Private/HighresScreenshotUI.h b/Engine/Source/Editor/LevelEditor/Private/HighresScreenshotUI.h
--- a/Engine/Source/Editor/LevelEditor/Private/HighresScreenshotUI.h
+++ b/Engine/Source/Editor/LevelEditor/Private/HighresScreenshotUI.h
@@ -1,4 +1,5 @@
 // Copyright Epic Games, Inc. All Rights Reserved.
+
 #pragma once
 
 #include "CoreMinimal.h"
@@ -66,7 +67,9 @@ public:
 	static TWeakPtr<class SWindow> OpenDialog(const TSharedPtr<FSceneViewport>& InViewport, TSharedPtr<SCaptureRegionWidget> InCaptureRegionWidget = TSharedPtr<class SCaptureRegionWidget>());
 
 private:
+	EVisibility GetCaptureWarningVisibility() const;
 
+	FReply OnOpenScreenshotFolder();
 	FReply OnCaptureClicked();
 	FReply OnSelectCaptureRegionClicked();
 	FReply OnSelectCaptureCancelRegionClicked();
