[ITF Viewport Interactions] Deprecate the OnRegisterViewportInteractions delegate, replacing it with an OnBuildViewportInteractions delegate that provides the UViewportInteractionsBehaviorSource that should be registered against.

#rnx

diff --git a/Engine/Plugins/Animation/ControlRig/Source/ControlRigEditor/Private/EditMode/ControlRigEditMode.cpp b/Engine/Plugins/Animation/ControlRig/Source/ControlRigEditor/Private/EditMode/ControlRigEditMode.cpp
--- a/Engine/Plugins/Animation/ControlRig/Source/ControlRigEditor/Private/EditMode/ControlRigEditMode.cpp
+++ b/Engine/Plugins/Animation/ControlRig/Source/ControlRigEditor/Private/EditMode/ControlRigEditMode.cpp
@@ -535,13 +535,10 @@ void FControlRigEditMode::Enter()
 			if (UModeManagerInteractiveToolsContext* const InteractiveToolsContext =
 					LevelEditor->GetEditorModeManager().GetInteractiveToolsContext())
 			{
-				InteractiveToolsContext->OnRegisterViewportInteractions().AddSPLambda(this, [InteractiveToolsContext]
+				InteractiveToolsContext->OnBuildViewportInteractions().AddSPLambda(this, [](UViewportInteractionsBehaviorSource* Source)
 				{
-					if (UViewportInteractionsBehaviorSource* const InteractionsBehaviorSource = InteractiveToolsContext->GetViewportInteractionsBehaviorSource())
-					{
-						// Add interactions specific to Control Rig
-						InteractionsBehaviorSource->AddInteraction<UControlRigFrustumSelectInteraction>();
-					}
+					// Add interactions specific to Control Rig
+					Source->AddInteraction<UControlRigFrustumSelectInteraction>();
 				});
 			}
 		}
@@ -594,7 +591,7 @@ void FControlRigEditMode::Exit()
 		{
 			if (UModeManagerInteractiveToolsContext* const InteractiveToolsContext = LevelEditor->GetEditorModeManager().GetInteractiveToolsContext())
 			{
-				InteractiveToolsContext->OnRegisterViewportInteractions().RemoveAll(this);
+				InteractiveToolsContext->OnBuildViewportInteractions().RemoveAll(this);
 			}
 		}
 	}
diff --git a/Engine/Source/Editor/LevelEditor/Private/SLevelEditor.cpp b/Engine/Source/Editor/LevelEditor/Private/SLevelEditor.cpp
--- a/Engine/Source/Editor/LevelEditor/Private/SLevelEditor.cpp
+++ b/Engine/Source/Editor/LevelEditor/Private/SLevelEditor.cpp
@@ -107,10 +107,7 @@ SLevelEditor::SLevelEditor()
 	{
 		ModeManagerInteractiveToolsContext->SetViewportInteractionsEnabled(true);
 
-		ModeManagerInteractiveToolsContext->OnRegisterViewportInteractions().AddRaw(this, &SLevelEditor::RegisterViewportInteractions
-	);
-
-		RegisterViewportInteractions();
+		ModeManagerInteractiveToolsContext->OnBuildViewportInteractions().AddRaw(this, &SLevelEditor::RegisterViewportInteractions);
 	}
 }
 
@@ -469,7 +466,7 @@ SLevelEditor::~SLevelEditor()
 	if (UModeManagerInteractiveToolsContext* const ModeManagerInteractiveToolsContext =
 			SLevelEditor::GetEditorModeManager().GetInteractiveToolsContext())
 	{
-		ModeManagerInteractiveToolsContext->OnRegisterViewportInteractions().RemoveAll(this);
+		ModeManagerInteractiveToolsContext->OnBuildViewportInteractions().RemoveAll(this);
 	}
 }
 
@@ -763,30 +760,19 @@ void SLevelEditor::ResetMostRecentOutliner()
 	SceneOutlinerPtr.Reset();
 }
 
-void SLevelEditor::RegisterViewportInteractions() const
+void SLevelEditor::RegisterViewportInteractions(UViewportInteractionsBehaviorSource* Source) const
 {
-	if (UModeManagerInteractiveToolsContext* const ModeManagerInteractiveToolsContext =
-			SLevelEditor::GetEditorModeManager().GetInteractiveToolsContext())
-	{
-		if (UViewportInteractionsBehaviorSource* InputToolsSource =
-				ModeManagerInteractiveToolsContext->GetViewportInteractionsBehaviorSource())
-		{
-			if (ModeManagerInteractiveToolsContext->UsesViewportInteractions())
-			{
-				UE::Editor::ViewportInteractions::AddDefaultCameraMovementInteractions(InputToolsSource);
-				UE::Editor::ViewportInteractions::AddDefaultDragToolsInteractions(InputToolsSource);
-			}
+	UE::Editor::ViewportInteractions::AddDefaultCameraMovementInteractions(Source);
+	UE::Editor::ViewportInteractions::AddDefaultDragToolsInteractions(Source);
 
-			// Left Mouse Click
-			InputToolsSource->AddInteraction<ULevelViewportClickSelection>();
+	// Left Mouse Click
+	Source->AddInteraction<ULevelViewportClickSelection>();
 
-			// Right Mouse Click - summons context menu
-			InputToolsSource->AddInteraction<ULevelViewportClickContextMenu>();
+	// Right Mouse Click - summons context menu
+	Source->AddInteraction<ULevelViewportClickContextMenu>();
 
-			// Double click
-			InputToolsSource->AddInteraction<ULevelViewportDoubleClick>();
-		}
-	}
+	// Double click
+	Source->AddInteraction<ULevelViewportDoubleClick>();
 }
 
 TSharedPtr<ISceneOutliner> SLevelEditor::GetMostRecentlyUsedSceneOutliner()
diff --git a/Engine/Source/Editor/LevelEditor/Private/SLevelEditor.h b/Engine/Source/Editor/LevelEditor/Private/SLevelEditor.h
--- a/Engine/Source/Editor/LevelEditor/Private/SLevelEditor.h
+++ b/Engine/Source/Editor/LevelEditor/Private/SLevelEditor.h
@@ -277,7 +277,7 @@ private:
 	/** Sets the first available outliner as the most recent outliner */
 	void ResetMostRecentOutliner();
 
-	void RegisterViewportInteractions() const;
+	void RegisterViewportInteractions(UViewportInteractionsBehaviorSource* Source) const;
 
 private:
 
diff --git a/Engine/Source/Editor/Sequencer/Private/SequencerEdMode.cpp b/Engine/Source/Editor/Sequencer/Private/SequencerEdMode.cpp
--- a/Engine/Source/Editor/Sequencer/Private/SequencerEdMode.cpp
+++ b/Engine/Source/Editor/Sequencer/Private/SequencerEdMode.cpp
@@ -211,13 +211,10 @@ void FSequencerEdMode::Enter()
 	
 	if (UModeManagerInteractiveToolsContext* Context = GetModeManager()->GetInteractiveToolsContext())
 	{
-		Context->OnRegisterViewportInteractions().AddSPLambda(this, [this, Context]
+		Context->OnBuildViewportInteractions().AddSPLambda(this, [this](UViewportInteractionsBehaviorSource* Source)
 		{
-			if (UViewportInteractionsBehaviorSource* Source = Context->GetViewportInteractionsBehaviorSource())
-			{
-				UMoveSequencerTimeSliderViewportInteraction* MoveInteraction = Source->AddInteraction<UMoveSequencerTimeSliderViewportInteraction>(true);
-				MoveInteraction->EdMode = this;	
-			}
+			UMoveSequencerTimeSliderViewportInteraction* MoveInteraction = Source->AddInteraction<UMoveSequencerTimeSliderViewportInteraction>(true);
+			MoveInteraction->EdMode = this;
 		});
 	}
 	
@@ -230,7 +227,7 @@ void FSequencerEdMode::Exit()
 
 	if (UModeManagerInteractiveToolsContext* Context = GetModeManager()->GetInteractiveToolsContext())
 	{
-		Context->OnRegisterViewportInteractions().RemoveAll(this);
+		Context->OnBuildViewportInteractions().RemoveAll(this);
 	}
 	
 	FEdMode::Exit();
diff --git a/Engine/Source/Editor/UnrealEd/Private/SEditorViewport.cpp b/Engine/Source/Editor/UnrealEd/Private/SEditorViewport.cpp
--- a/Engine/Source/Editor/UnrealEd/Private/SEditorViewport.cpp
+++ b/Engine/Source/Editor/UnrealEd/Private/SEditorViewport.cpp
@@ -72,7 +72,7 @@ SEditorViewport::~SEditorViewport()
 	if (const TStrongObjectPtr<UModeManagerInteractiveToolsContext> ModeManagerInteractiveToolsContext =
 			ModeManagerInteractiveToolsContextWeak.Pin())
 	{
-		ModeManagerInteractiveToolsContext->OnRegisterViewportInteractions().RemoveAll(this);
+		ModeManagerInteractiveToolsContext->OnBuildViewportInteractions().RemoveAll(this);
 	}
 }
 
@@ -328,13 +328,12 @@ void SEditorViewport::Construct( const FArguments& InArgs )
 		// GetModeManagerInteractiveToolsContext is virtual, so we store the context as a weak ptr.
 		// This allows us to access it in class destructor
 		ModeManagerInteractiveToolsContextWeak = ModeManagerInteractiveToolsContext;
-		ModeManagerInteractiveToolsContext->OnRegisterViewportInteractions().AddSP(
-			this, &SEditorViewport::RegisterViewportInteractions
+		PRAGMA_DISABLE_DEPRECATION_WARNINGS
+		ModeManagerInteractiveToolsContext->OnBuildViewportInteractions().AddSP(
+			this, &SEditorViewport::OnRegisterViewportInteractions
 		);
+		PRAGMA_ENABLE_DEPRECATION_WARNINGS
 	}
-
-	// Register the Editor-specific viewport interactions
-	RegisterViewportInteractions();
 }
 
 FReply SEditorViewport::OnKeyDown( const FGeometry& MyGeometry, const FKeyEvent& InKeyEvent )
@@ -1396,17 +1395,6 @@ EVisibility SEditorViewport::GetActiveBorderVisibility() const
 	return BaseVisibility;
 }
 
-void SEditorViewport::RegisterViewportInteractions()
-{
-	if (UModeManagerInteractiveToolsContext* InteractiveToolsContext = GetModeManagerInteractiveToolsContext())
-	{
-		if (UViewportInteractionsBehaviorSource* InteractionsBehaviorSource = InteractiveToolsContext->GetViewportInteractionsBehaviorSource())
-		{
-			OnRegisterViewportInteractions(InteractionsBehaviorSource);
-		}
-	}
-}
-
 ///////////////////////////////////////////////////////////////////////////////
 // begin feature level control functions block
 ///////////////////////////////////////////////////////////////////////////////
diff --git a/Engine/Source/Editor/UnrealEd/Private/Tools/EdModeInteractiveToolsContext.cpp b/Engine/Source/Editor/UnrealEd/Private/Tools/EdModeInteractiveToolsContext.cpp
--- a/Engine/Source/Editor/UnrealEd/Private/Tools/EdModeInteractiveToolsContext.cpp
+++ b/Engine/Source/Editor/UnrealEd/Private/Tools/EdModeInteractiveToolsContext.cpp
@@ -1468,7 +1468,11 @@ void UModeManagerInteractiveToolsContext::RegisterViewportInteractions()
 		ViewportInteractionsBehaviorSource->Initialize(this);
 	}
 	
+	PRAGMA_DISABLE_DEPRECATION_WARNINGS
 	OnRegisterViewportInteractions().Broadcast();
+	PRAGMA_ENABLE_DEPRECATION_WARNINGS
+	
+	OnBuildViewportInteractions().Broadcast(ViewportInteractionsBehaviorSource);
 	
 	ViewportInteractionsBehaviorSource->RegisterBehaviorSources();
 }
diff --git a/Engine/Source/Editor/UnrealEd/Public/SEditorViewport.h b/Engine/Source/Editor/UnrealEd/Public/SEditorViewport.h
--- a/Engine/Source/Editor/UnrealEd/Public/SEditorViewport.h
+++ b/Engine/Source/Editor/UnrealEd/Public/SEditorViewport.h
@@ -279,6 +279,7 @@ protected:
 	 * Override this function to add the desired viewport interactions for this editor.
 	 * Requires virtual function GetModeManagerInteractiveToolsContext to return a valid UModeManagerInteractiveToolsContext
 	 */
+	UE_DEPRECATED(5.8, "Viewport interactions should be managed from the mode manager layer rather than from individual viewports.")
 	virtual void OnRegisterViewportInteractions(UViewportInteractionsBehaviorSource* InBehaviorSource)
 	{
 	}
@@ -364,11 +365,6 @@ private:
 	/** Gets the visibility of the active viewport border */
 	UNREALED_API EVisibility GetActiveBorderVisibility() const;
 
-	/**
-	 * Will register ITF viewport interactions, if OnRegisterViewportInteractions is overridden in your SEditorViewport
-	 */
-	void RegisterViewportInteractions();
-	
 	/** Called when the user wants to focus the viewport to the current selection and clip the near plane to the bounds of that selection */
 	void OnFocusAndClipViewportToSelection();
 
diff --git a/Engine/Source/Editor/UnrealEd/Public/Tools/EdModeInteractiveToolsContext.h b/Engine/Source/Editor/UnrealEd/Public/Tools/EdModeInteractiveToolsContext.h
--- a/Engine/Source/Editor/UnrealEd/Public/Tools/EdModeInteractiveToolsContext.h
+++ b/Engine/Source/Editor/UnrealEd/Public/Tools/EdModeInteractiveToolsContext.h
@@ -399,13 +399,30 @@ public:
 	 * Return the current Input Tools Behavior Source - nullptr if not enabled
 	 */
 	UE_API UViewportInteractionsBehaviorSource* GetViewportInteractionsBehaviorSource();
-
+	
+PRAGMA_DISABLE_DEPRECATION_WARNINGS
+	UE_DEPRECATED(5.8, "Use OnBuildViewportInteractions instead.")
 	DECLARE_MULTICAST_DELEGATE(FViewportInteractionsRegisterStateChanged)
+	UE_DEPRECATED(5.8, "Use OnBuildViewportInteractions instead.")
 	FViewportInteractionsRegisterStateChanged& OnRegisterViewportInteractions()
 	{
 		return OnRegisterViewportInteractionsDelegate;
 	}
+	UE_DEPRECATED(5.8, "Use OnBuildViewportInteractions instead.")
 	FViewportInteractionsRegisterStateChanged OnRegisterViewportInteractionsDelegate;
+PRAGMA_ENABLE_DEPRECATION_WARNINGS
+	
+	UE_EXPERIMENTAL(5.8, "Viewport interactions are experimental")
+	DECLARE_MULTICAST_DELEGATE_OneParam(FViewportInteractionsSourceMessage, UViewportInteractionsBehaviorSource*)
+	
+	UE_EXPERIMENTAL(5.8, "Viewport interactions are experimental")
+	/**
+	 * Called when viewport interactions are built. Add any custom interactions to the behavior source that is passed in via the event.
+	 * Interactions are automatically removed when editor modes change.
+	 */
+PRAGMA_DISABLE_EXPERIMENTAL_WARNINGS
+	FViewportInteractionsSourceMessage& OnBuildViewportInteractions() { return OnBuildViewportInteractionsEvent; }
+PRAGMA_ENABLE_EXPERIMENTAL_WARNINGS
 
 protected:
 	/**
@@ -455,6 +472,11 @@ protected:
 
 	UPROPERTY(Transient)
 	TObjectPtr<UViewportInteractionsBehaviorSource> ViewportInteractionsBehaviorSource;
+	
+	PRAGMA_DISABLE_EXPERIMENTAL_WARNINGS
+	UE_EXPERIMENTAL(5.8, "Viewport interactions are experimental")
+	FViewportInteractionsSourceMessage OnBuildViewportInteractionsEvent;
+	PRAGMA_ENABLE_EXPERIMENTAL_WARNINGS
 
 	bool bUsesViewportInteractions = false;
 };
