Editor: Allow the property matrix to respond correctly to actors that are stand-ins for other items, if PropertyEditor.AllowPropertyMatrixForNonActors is set to true.

diff --git a/Engine/Source/Editor/LevelEditor/Private/LevelEditorActions.cpp b/Engine/Source/Editor/LevelEditor/Private/LevelEditorActions.cpp
--- a/Engine/Source/Editor/LevelEditor/Private/LevelEditorActions.cpp
+++ b/Engine/Source/Editor/LevelEditor/Private/LevelEditorActions.cpp
@@ -50,10 +50,12 @@
 
 #include "Elements/Framework/TypedElementCommonActions.h"
 #include "Elements/Framework/TypedElementSelectionSet.h"
+#include "Elements/Interfaces/TypedElementDetailsInterface.h"
 #include "Elements/Interfaces/TypedElementObjectInterface.h"
 #include "Subsystems/EditorElementSubsystem.h"
 
 #include "LevelEditor.h"
+#include "LevelEditorInternalTools.h"
 #include "Engine/LevelScriptBlueprint.h"
 #include "LightingBuildOptions.h"
 #include "EditorSupportDelegates.h"
@@ -1549,16 +1551,49 @@ bool FLevelEditorActionCallbacks::EditAsset_CanExecute()
 
 void FLevelEditorActionCallbacks::OpenSelectionInPropertyMatrix_Clicked()
 {
+	// Gather the selected objects through the ITypedElementDetailsInterface so that any items that use
+	//  stand-in actors are properly gathered instead of their actors.
 	TArray<UObject*> SelectedObjects;
-	GEditor->GetSelectedActors()->GetSelectedObjects<UObject>(SelectedObjects);
-
-	FPropertyEditorModule& PropertyEditorModule = FModuleManager::LoadModuleChecked<FPropertyEditorModule>( "PropertyEditor" );
-	PropertyEditorModule.CreatePropertyEditorToolkit(TSharedPtr<IToolkitHost>(), SelectedObjects );
+	USelection* Selection = GEditor->GetSelectedActors();
+	UTypedElementSelectionSet* SelectionSet = Selection ? Selection->GetElementSelectionSet() : nullptr;
+	if (ensure(SelectionSet))
+	{
+		SelectionSet->ForEachSelectedElement<ITypedElementDetailsInterface>([&SelectedObjects](const TTypedElement<ITypedElementDetailsInterface>& InDetailsElement)
+		{
+			if (TUniquePtr<ITypedElementDetailsObject> ElementDetailsObject = InDetailsElement.GetDetailsObject())
+			{
+				if (UObject* Object = ElementDetailsObject->GetObject())
+				{
+					SelectedObjects.Add(Object);
+				}
+			}
+			return true;
+		});
+	}
+	
+	if (!SelectedObjects.IsEmpty())
+	{
+		FPropertyEditorModule& PropertyEditorModule = FModuleManager::LoadModuleChecked<FPropertyEditorModule>( "PropertyEditor" );
+		PropertyEditorModule.CreatePropertyEditorToolkit(TSharedPtr<IToolkitHost>(), SelectedObjects );
+	}
 }
 
 bool FLevelEditorActionCallbacks::OpenSelectionInPropertyMatrix_IsVisible()
 {
-	return FModuleManager::LoadModuleChecked<FPropertyEditorModule>( "PropertyEditor" ).GetCanUsePropertyMatrix() && GEditor->GetSelectedActorCount() > 1;
+	auto SelectionIsNonActors = []()->bool
+	{
+		USelection* Selection = GEditor->GetSelectedActors();
+		if (!Selection)
+		{
+			return false;
+		}
+		return InternalEditorLevelLibrary::DoesDetailSelectionContainNonActorsOrComponents(Selection->GetElementSelectionSet());
+	};
+
+	return FModuleManager::LoadModuleChecked<FPropertyEditorModule>( "PropertyEditor" ).GetCanUsePropertyMatrix() 
+		&& GEditor->GetSelectedActorCount() > 1
+		&& (FPropertyEditorModule::CVarAllowPropertyMatrixForNonActors.GetValueOnAnyThread()
+			|| !SelectionIsNonActors());
 }
 
 void FLevelEditorActionCallbacks::LockActorMovement_Clicked()
diff --git a/Engine/Source/Editor/LevelEditor/Private/LevelEditorContextMenu.cpp b/Engine/Source/Editor/LevelEditor/Private/LevelEditorContextMenu.cpp
--- a/Engine/Source/Editor/LevelEditor/Private/LevelEditorContextMenu.cpp
+++ b/Engine/Source/Editor/LevelEditor/Private/LevelEditorContextMenu.cpp
@@ -32,6 +32,7 @@
 #include "LevelEditorViewport.h"
 #include "EditorModes.h"
 #include "LevelEditor.h"
+#include "LevelEditorInternalTools.h"
 #include "Kismet2/KismetEditorUtilities.h"
 #include "AssetSelection.h"
 #include "LevelEditorActions.h"
@@ -369,13 +370,20 @@ void FLevelEditorContextMenu::RegisterActorContextMenu()
 			LevelEditorCreateActorMenu::FillAddReplaceContextMenuSections(Section, LevelEditorContext);
 		}
 		
+		if(FModuleManager::LoadModuleChecked<FPropertyEditorModule>( "PropertyEditor" ).GetCanUsePropertyMatrix() 
+			&& NumSelectedActors >= 2)
 		{
-			if(FModuleManager::LoadModuleChecked<FPropertyEditorModule>( "PropertyEditor" ).GetCanUsePropertyMatrix() && NumSelectedActors >= 2)
+			// Options that relate to bulk editing assets
+			// In most cases, you DO NOT want to extend this section; look at ActorUETools or ActorTypeTools below
+			
+			// Figure out whether the selected actors are actually the same thing we would see in a detail panel,
+			//  or whether they are actually stand-ins for some non-actors.
+			bool bDetailsObjectsAreNotActors = InternalEditorLevelLibrary::DoesDetailSelectionContainNonActorsOrComponents(LevelEditorContext->CurrentSelection);
+
+			if (!bDetailsObjectsAreNotActors || FPropertyEditorModule::CVarAllowPropertyMatrixForNonActors.GetValueOnAnyThread())
 			{
-				// Options that relate to bulk editing assets
-				// In most cases, you DO NOT want to extend this section; look at ActorUETools or ActorTypeTools below
 				FToolMenuSection& Section = InMenu->AddSection("ActorBulkEdit", LOCTEXT("ActorBulkEditHeading", "Bulk Editing"));
-				
+
 				Section.AddMenuEntry(
 					FLevelEditorCommands::Get().OpenSelectionInPropertyMatrix,
 					TAttribute<FText>(), // use command's label
@@ -383,16 +391,25 @@ void FLevelEditorContextMenu::RegisterActorContextMenu()
 					FSlateIcon(FAppStyle::GetAppStyleSetName(), "DetailsView.EditRawProperties")
 				);
 
-				Section.AddSubMenu(
-					"BulkEditComponentsSubmenu",
-					LOCTEXT("BulkEditComponentsSubmenuName", "Edit Components in the Property Matrix"),
-					LOCTEXT("BulkEditComponentsSubmenuTooltip", "Bulk Edit any editable components that are common between the selected actors in the Property Matrix"),
-					FNewToolMenuDelegate::CreateStatic(&FLevelEditorContextMenuImpl::FillBulkEditComponentsMenu),
-					FUIAction(),
-					EUserInterfaceActionType::Button,
-					false, // default value
-					FSlateIcon(FAppStyle::GetAppStyleSetName(), "DetailsView.EditRawProperties")
-				);
+				// The previous entry for opening the selection in the property matrix works even in cases where 
+				//  the selected actors are stand-ins for something different that is editable (since the true 
+				//  targets will be gotten via ITypedElementDetailsInterface). On the other hand, editing components
+				//  of the selection is harder to generalize in a way that is overridable by plugins. So, we will
+				//  not add this entry in that case. A plugin that uses stand-in actors can add its own "Edit Components"
+				//  entry if it likes.
+				if (!bDetailsObjectsAreNotActors)
+				{
+					Section.AddSubMenu(
+						"BulkEditComponentsSubmenu",
+						LOCTEXT("BulkEditComponentsSubmenuName", "Edit Components in the Property Matrix"),
+						LOCTEXT("BulkEditComponentsSubmenuTooltip", "Bulk Edit any editable components that are common between the selected actors in the Property Matrix"),
+						FNewToolMenuDelegate::CreateStatic(&FLevelEditorContextMenuImpl::FillBulkEditComponentsMenu),
+						FUIAction(),
+						EUserInterfaceActionType::Button,
+						false, // default value
+						FSlateIcon(FAppStyle::GetAppStyleSetName(), "DetailsView.EditRawProperties")
+					);
+				}
 			}
 		}
 
diff --git a/Engine/Source/Editor/LevelEditor/Private/LevelEditorInternalTools.cpp b/Engine/Source/Editor/LevelEditor/Private/LevelEditorInternalTools.cpp
--- a/Engine/Source/Editor/LevelEditor/Private/LevelEditorInternalTools.cpp
+++ b/Engine/Source/Editor/LevelEditor/Private/LevelEditorInternalTools.cpp
@@ -2,6 +2,9 @@
 
 #include "LevelEditorInternalTools.h"
 
+#include "Elements/Framework/TypedElementSelectionSet.h"
+#include "Elements/Interfaces/TypedElementDetailsInterface.h"
+
 namespace InternalEditorLevelLibrary
 {
 	TSharedPtr<SLevelViewport> GetLevelViewport(const FName& ViewportConfigKey)
@@ -60,4 +63,29 @@ namespace InternalEditorLevelLibrary
 	{
 		return InWorld && (InWorld->GetCurrentLevel()->OwningWorld->GetLevels().Num() > 1 && (!InWorld->IsPartitionedWorld() || GetEditingLevelInstance(InWorld) != nullptr));
 	}
+
+	bool DoesDetailSelectionContainNonActorsOrComponents(const UTypedElementSelectionSet* SelectionSet)
+	{
+		if (!ensure(SelectionSet))
+		{
+			return false;
+		}
+	
+		bool bFoundNonActorOrComponent = false;
+		SelectionSet->ForEachSelectedElement<ITypedElementDetailsInterface>([&bFoundNonActorOrComponent](const TTypedElement<ITypedElementDetailsInterface>& InDetailsElement)
+		{
+			if (TUniquePtr<ITypedElementDetailsObject> ElementDetailsObject = InDetailsElement.GetDetailsObject())
+			{
+				UObject* Object = ElementDetailsObject->GetObject();
+				if (!Object || (!Object->IsA<AActor>() && !Object->IsA<UActorComponent>()))
+				{
+					bFoundNonActorOrComponent = true;
+					return false;
+				}
+			}
+			return true;
+		});
+
+		return bFoundNonActorOrComponent;
+	}
 }
diff --git a/Engine/Source/Editor/LevelEditor/Private/LevelEditorInternalTools.h b/Engine/Source/Editor/LevelEditor/Private/LevelEditorInternalTools.h
--- a/Engine/Source/Editor/LevelEditor/Private/LevelEditorInternalTools.h
+++ b/Engine/Source/Editor/LevelEditor/Private/LevelEditorInternalTools.h
@@ -9,6 +9,8 @@
 #include "LevelInstance/LevelInstanceSubsystem.h"
 #include "LevelInstance/LevelInstanceInterface.h"
 
+class UTypedElementSelectionSet;
+
 namespace InternalEditorLevelLibrary
 {
 	TSharedPtr<SLevelViewport> GetLevelViewport(const FName& ViewportConfigKey);
@@ -16,4 +18,9 @@ namespace InternalEditorLevelLibrary
 	bool IsEditingLevelInstanceCurrentLevel(UWorld* InWorld);
 	AActor* GetEditingLevelInstance(UWorld* InWorld);
 	bool IsActorEditorContextVisible(UWorld* InWorld);
+	/**
+	 * When converted to objects that would be given to the detail panel, does the selection contain anything that
+	 *  is not an actor or component?
+	 */
+	bool DoesDetailSelectionContainNonActorsOrComponents(const UTypedElementSelectionSet* SelectionSet);
 }
diff --git a/Engine/Source/Editor/PropertyEditor/Private/PropertyEditorModule.cpp b/Engine/Source/Editor/PropertyEditor/Private/PropertyEditorModule.cpp
--- a/Engine/Source/Editor/PropertyEditor/Private/PropertyEditorModule.cpp
+++ b/Engine/Source/Editor/PropertyEditor/Private/PropertyEditorModule.cpp
@@ -37,6 +37,12 @@
 
 IMPLEMENT_MODULE( FPropertyEditorModule, PropertyEditor );
 
+TAutoConsoleVariable<bool> FPropertyEditorModule::CVarAllowPropertyMatrixForNonActors = TAutoConsoleVariable<bool>(
+	TEXT("PropertyEditor.AllowPropertyMatrixForNonActors"),
+	false,
+	TEXT("Specifies whether the property matrix can be launched when selected items are not actors or actor components.")
+);
+
 TSharedRef<IPropertyTypeCustomization> FPropertyTypeLayoutCallback::GetCustomizationInstance() const
 {
 	return PropertyTypeLayoutDelegate.Execute();
diff --git a/Engine/Source/Editor/PropertyEditor/Private/SDetailsView.cpp b/Engine/Source/Editor/PropertyEditor/Private/SDetailsView.cpp
--- a/Engine/Source/Editor/PropertyEditor/Private/SDetailsView.cpp
+++ b/Engine/Source/Editor/PropertyEditor/Private/SDetailsView.cpp
@@ -529,8 +529,22 @@ bool SDetailsView::CanOpenRawPropertyEditor() const
 
 EVisibility SDetailsView::CanShowRawPropertyEditorButton(const bool bAllowedByDetailsViewArgs) const
 {
-	const bool bShowPropertyMatrixOverride =  FModuleManager::LoadModuleChecked<FPropertyEditorModule>( "PropertyEditor" ).GetCanUsePropertyMatrix();
-	return (DetailsViewArgs.bShowPropertyMatrixButton && bShowPropertyMatrixOverride) ? EVisibility::Visible : EVisibility::Collapsed;
+	auto SelectionContainsNonActors = [this]() -> bool
+	{
+		for (const TWeakObjectPtr<UObject> Object : SelectedObjects)
+		{
+			if (Object.IsValid() && !Object->IsA<AActor>() && !Object->IsA<UActorComponent>())
+			{
+				return true;
+			}
+		}
+		return false;
+	};
+
+	const FPropertyEditorModule& PropertyModule = FModuleManager::LoadModuleChecked<FPropertyEditorModule>("PropertyEditor");
+	const bool bShowPropertyMatrixOverride = PropertyModule.GetCanUsePropertyMatrix();
+	return (DetailsViewArgs.bShowPropertyMatrixButton && bShowPropertyMatrixOverride 
+		&& (PropertyModule.CVarAllowPropertyMatrixForNonActors.GetValueOnAnyThread() || !SelectionContainsNonActors())) ? EVisibility::Visible : EVisibility::Collapsed;
 }
 
 FReply SDetailsView::OnOpenRawPropertyEditorClicked()
diff --git a/Engine/Source/Editor/PropertyEditor/Public/PropertyEditorModule.h b/Engine/Source/Editor/PropertyEditor/Public/PropertyEditorModule.h
--- a/Engine/Source/Editor/PropertyEditor/Public/PropertyEditorModule.h
+++ b/Engine/Source/Editor/PropertyEditor/Public/PropertyEditorModule.h
@@ -4,6 +4,7 @@
 #include "CoreMinimal.h"
 #include "IDetailsView.h"
 #include "IPropertyTypeCustomization.h"
+#include "HAL/IConsoleManager.h" // TAutoConsoleVariable
 #include "PropertyEditorDelegates.h"
 
 #include "Modules/ModuleInterface.h"
@@ -464,6 +465,8 @@ public:
 		bCanUsePropertyMatrixOverride = bInCanUsePropertyMatrix; 
 	} 
 
+	UE_EXPERIMENTAL(5.8, "This CVar may be removed.")
+	PROPERTYEDITOR_API static TAutoConsoleVariable<bool> CVarAllowPropertyMatrixForNonActors;
 private:
 
 	/**
