Enable individual level editor details panels to affect the add component visibility
* Added an internal customization for the SActorDetails that wraps the LevelEditor scoped customization to allow it to control AddComponent button visibility instead of having the SActorDetails change the visibility for the container of all buttons.
* Left a CVar to revert to legacy behaviour _just in case_.  The CVar will be removed in a later engine version.

diff --git a/Engine/Source/Editor/Kismet/Public/ISCSEditorUICustomization.h b/Engine/Source/Editor/Kismet/Public/ISCSEditorUICustomization.h
--- a/Engine/Source/Editor/Kismet/Public/ISCSEditorUICustomization.h
+++ b/Engine/Source/Editor/Kismet/Public/ISCSEditorUICustomization.h
@@ -24,21 +24,6 @@ class ISCSEditorUICustomization
 public:
 	virtual ~ISCSEditorUICustomization() {}
 
-	UE_DEPRECATED(5.5, "Use version with ObjectContext instead")
-	virtual bool HideComponentsTree() const { return false; }
-
-	UE_DEPRECATED(5.5, "Use version with ObjectContext instead")
-	virtual bool HideComponentsFilterBox() const { return false; }
-
-	UE_DEPRECATED(5.5, "Use version with ObjectContext instead")
-	virtual bool HideAddComponentButton() const { return false; }
-
-	UE_DEPRECATED(5.5, "Use version with ObjectContext instead")
-	virtual bool HideBlueprintButtons() const { return false; }
-
-	UE_DEPRECATED(5.5, "Use version with ObjectContext instead")
-	virtual TSubclassOf<UActorComponent> GetComponentTypeFilter() const { return nullptr; }
-
 	/** @return Whether to hide the components tree */
 	virtual bool HideComponentsTree(TArrayView<UObject*> Context) const { return false; }
 
diff --git a/Engine/Source/Editor/LevelEditor/Private/ActorDetailsSCSEditorUICustomization.cpp b/Engine/Source/Editor/LevelEditor/Private/ActorDetailsSCSEditorUICustomization.cpp
--- a/Engine/Source/Editor/LevelEditor/Private/ActorDetailsSCSEditorUICustomization.cpp
+++ b/Engine/Source/Editor/LevelEditor/Private/ActorDetailsSCSEditorUICustomization.cpp
@@ -1,15 +1,15 @@
 // Copyright Epic Games, Inc. All Rights Reserved.
 #include "ActorDetailsSCSEditorUICustomization.h"
 
-TSharedPtr<FActorDetailsSCSEditorUICustomization> FActorDetailsSCSEditorUICustomization::Instance;
+#include "SubobjectData.h"
 
-TSharedPtr<FActorDetailsSCSEditorUICustomization> FActorDetailsSCSEditorUICustomization::GetInstance()
+TSharedPtr<FActorDetailsSCSEditorUICustomization> FActorDetailsSCSEditorUICustomization::NewInstance()
+{
+	return MakeShared<FActorDetailsSCSEditorUICustomization>(FPrivateToken());
+}
+
+FActorDetailsSCSEditorUICustomization::FActorDetailsSCSEditorUICustomization(FPrivateToken)
 {
-	if (!Instance)
-	{
-		Instance = MakeShareable(new FActorDetailsSCSEditorUICustomization());
-	}
-	return Instance;
 }
 
 bool FActorDetailsSCSEditorUICustomization::HideComponentsTree(TArrayView<UObject*> Context) const
diff --git a/Engine/Source/Editor/LevelEditor/Private/ActorDetailsSCSEditorUICustomization.h b/Engine/Source/Editor/LevelEditor/Private/ActorDetailsSCSEditorUICustomization.h
--- a/Engine/Source/Editor/LevelEditor/Private/ActorDetailsSCSEditorUICustomization.h
+++ b/Engine/Source/Editor/LevelEditor/Private/ActorDetailsSCSEditorUICustomization.h
@@ -6,12 +6,15 @@
 #include "Containers/Array.h"
 #include "Templates/SharedPointer.h"
 
-class FActorDetailsSCSEditorUICustomization : public ISCSEditorUICustomization
+class FActorDetailsSCSEditorUICustomization final : public ISCSEditorUICustomization
 {
+private:
+	struct FPrivateToken { explicit FPrivateToken() = default; };
 public:
-	static TSharedPtr<FActorDetailsSCSEditorUICustomization> GetInstance();
-	
-	virtual ~FActorDetailsSCSEditorUICustomization() {}
+	static TSharedPtr<FActorDetailsSCSEditorUICustomization> NewInstance();
+
+	explicit FActorDetailsSCSEditorUICustomization(FPrivateToken);
+	virtual ~FActorDetailsSCSEditorUICustomization() override {}
 
 	virtual bool HideComponentsTree(TArrayView<UObject*> Context) const override;
 
@@ -42,8 +45,6 @@ public:
 
 private:
 	FActorDetailsSCSEditorUICustomization() {}
-
-private:
-	static TSharedPtr<FActorDetailsSCSEditorUICustomization> Instance;
+	
 	TArray<TSharedPtr<ISCSEditorUICustomization>> Customizations;
 };
\ No newline at end of file
diff --git a/Engine/Source/Editor/LevelEditor/Private/SActorDetails.cpp b/Engine/Source/Editor/LevelEditor/Private/SActorDetails.cpp
--- a/Engine/Source/Editor/LevelEditor/Private/SActorDetails.cpp
+++ b/Engine/Source/Editor/LevelEditor/Private/SActorDetails.cpp
@@ -221,7 +221,218 @@ public:
 	}
 };
 
-void SActorDetails::Construct(const FArguments& InArgs, UTypedElementSelectionSet* InSelectionSet, const FName TabIdentifier, TSharedPtr<FUICommandList> InCommandList, TSharedPtr<FTabManager> InTabManager)
+class SActorDetails::FInternalSubobjectCustomization : public ISCSEditorUICustomization
+{
+public:
+
+	explicit FInternalSubobjectCustomization(TSharedPtr<FActorDetailsSCSEditorUICustomization>& ParentCustomization);
+	virtual ~FInternalSubobjectCustomization() override = default;
+
+	void SetHideAddComponentButton(bool bShouldHide);
+	void SetHideBlueprintButton(bool bShouldHide);
+	
+	virtual bool HideAddComponentButton(TArrayView<UObject*> Context) const override;
+	virtual bool HideBlueprintButtons(TArrayView<UObject*> Context) const override;
+	virtual bool HideComponentsTree(TArrayView<UObject*> Context) const override;
+	virtual bool HideComponentsFilterBox(TArrayView<UObject*> Context) const override;
+	virtual const FSlateBrush* GetIconBrush(const FSubobjectData&) const override;
+	virtual TSharedPtr<SWidget> GetControlsWidget(const FSubobjectData&) const override;
+	virtual bool OverrideDesiredVisibility(const UE::Editor::FSubobjectEditorContext& Context, EVisibility& VisibilityOut) const override;
+	virtual EChildActorComponentTreeViewVisualizationMode GetChildActorVisualizationMode() const override;
+	virtual TSubclassOf<UActorComponent> GetComponentTypeFilter(TArrayView<UObject*> Context) const override;
+	virtual bool FilterSubobjectData(TArray<FSubobjectDataHandle>& SubobjectData) const override;
+	virtual bool SortSubobjectData(TArray<FSubobjectDataHandle>& SubobjectData) override;
+	virtual bool SupportsTreeViewOverlay(TArrayView<UObject*> Context) const override;
+	virtual TSharedRef<class SWidget> ConstructTreeViewOverlay(TWeakPtr<class SSubobjectEditor> Editor) const override;
+	virtual void PrePotentialTargetChange(SSubobjectEditor& SubobjectEditor) override;
+	virtual bool SelectDefaultSubobject(SSubobjectEditor& SubobjectEditor) override;
+
+private:
+	TSharedPtr<FActorDetailsSCSEditorUICustomization> ParentCustomization;
+	
+	bool bHideAddComponentButton = true;
+	bool bHideBlueprintButton = true;
+};
+
+bool SActorDetails::FInternalSubobjectCustomization::HideAddComponentButton(TArrayView<UObject*> Context) const
+{
+	if (const bool bParentHides = ParentCustomization && ParentCustomization->HideAddComponentButton(Context))
+	{
+		return bParentHides;
+	}
+	
+	return bHideAddComponentButton;
+}
+
+bool SActorDetails::FInternalSubobjectCustomization::HideBlueprintButtons(TArrayView<UObject*> Context) const
+{
+	if (const bool bParentHides = ParentCustomization && ParentCustomization->HideBlueprintButtons(Context))
+	{
+		return bParentHides;
+	}
+	
+	return bHideBlueprintButton;
+}
+
+bool SActorDetails::FInternalSubobjectCustomization::HideComponentsTree(TArrayView<UObject*> Context) const
+{
+	if (const bool bParentHides = ParentCustomization && ParentCustomization->HideComponentsTree(Context))
+	{
+		return bParentHides;
+	}
+
+	return false;
+}
+
+bool SActorDetails::FInternalSubobjectCustomization::HideComponentsFilterBox(TArrayView<UObject*> Context) const
+{
+	if (const bool bParentHides = ParentCustomization && ParentCustomization->HideComponentsFilterBox(Context))
+	{
+		return bParentHides;
+	}
+
+	return false;
+}
+
+const FSlateBrush* SActorDetails::FInternalSubobjectCustomization::GetIconBrush(const FSubobjectData& SubobjectData) const
+{
+	if (ParentCustomization)
+	{
+		return ParentCustomization->GetIconBrush(SubobjectData);
+	}
+	
+	return ISCSEditorUICustomization::GetIconBrush(SubobjectData);
+}
+
+TSharedPtr<SWidget> SActorDetails::FInternalSubobjectCustomization::GetControlsWidget(const FSubobjectData& SubobjectData) const
+{
+	if (ParentCustomization)
+	{
+		return ParentCustomization->GetControlsWidget(SubobjectData);
+	}
+
+	return ISCSEditorUICustomization::GetControlsWidget(SubobjectData);
+}
+
+bool SActorDetails::FInternalSubobjectCustomization::OverrideDesiredVisibility(const UE::Editor::FSubobjectEditorContext& Context, EVisibility& VisibilityOut) const
+{
+	if (ParentCustomization)
+	{
+		if (ParentCustomization->OverrideDesiredVisibility(Context, VisibilityOut))
+		{
+			return true;
+		}
+	}
+	return ISCSEditorUICustomization::OverrideDesiredVisibility(Context, VisibilityOut);
+}
+
+EChildActorComponentTreeViewVisualizationMode SActorDetails::FInternalSubobjectCustomization::GetChildActorVisualizationMode() const
+{
+	if (ParentCustomization)
+	{
+		return ParentCustomization->GetChildActorVisualizationMode();
+	}
+
+	return EChildActorComponentTreeViewVisualizationMode::UseDefault;
+}
+
+TSubclassOf<UActorComponent> SActorDetails::FInternalSubobjectCustomization::GetComponentTypeFilter(TArrayView<UObject*> Context) const
+{
+	if (ParentCustomization)
+	{
+		return ParentCustomization->GetComponentTypeFilter(Context);
+	}
+
+	return nullptr;
+}
+
+bool SActorDetails::FInternalSubobjectCustomization::FilterSubobjectData(TArray<FSubobjectDataHandle>& SubobjectData) const
+{
+	if (ParentCustomization)
+	{
+		return ParentCustomization->FilterSubobjectData(SubobjectData);
+	}
+	return false;
+}
+
+bool SActorDetails::FInternalSubobjectCustomization::SortSubobjectData(TArray<FSubobjectDataHandle>& SubobjectData)
+{
+	if (ParentCustomization)
+	{
+		return ParentCustomization->SortSubobjectData(SubobjectData);
+	}
+	return false;
+}
+
+bool SActorDetails::FInternalSubobjectCustomization::SupportsTreeViewOverlay(TArrayView<UObject*> Context) const
+{
+	if (ParentCustomization)
+	{
+		return ParentCustomization->SupportsTreeViewOverlay(Context);
+	}
+	return false;
+}
+
+TSharedRef<SWidget> SActorDetails::FInternalSubobjectCustomization::ConstructTreeViewOverlay(TWeakPtr<SSubobjectEditor> Editor) const
+{
+	if (ParentCustomization)
+	{
+		return ParentCustomization->ConstructTreeViewOverlay(Editor);
+	}
+	return SNullWidget::NullWidget;
+}
+
+void SActorDetails::FInternalSubobjectCustomization::PrePotentialTargetChange(SSubobjectEditor& InSubobjectEditor)
+{
+	if (ParentCustomization)
+	{
+		ParentCustomization->PrePotentialTargetChange(InSubobjectEditor);
+	}
+	ISCSEditorUICustomization::PrePotentialTargetChange(InSubobjectEditor);
+}
+
+bool SActorDetails::FInternalSubobjectCustomization::SelectDefaultSubobject(SSubobjectEditor& InSubobjectEditor)
+{
+	if (ParentCustomization)
+	{
+		return ParentCustomization->SelectDefaultSubobject(InSubobjectEditor);
+	}
+	return ISCSEditorUICustomization::SelectDefaultSubobject(InSubobjectEditor);
+}
+
+SActorDetails::FInternalSubobjectCustomization::FInternalSubobjectCustomization(
+	TSharedPtr<FActorDetailsSCSEditorUICustomization>& InParentCustomization)
+		: ParentCustomization(InParentCustomization)
+{
+}
+
+void SActorDetails::FInternalSubobjectCustomization::SetHideAddComponentButton(bool bShouldHide)
+{
+	bHideAddComponentButton = bShouldHide;
+}
+
+void SActorDetails::FInternalSubobjectCustomization::SetHideBlueprintButton(bool bShouldHide)
+{
+	bHideBlueprintButton = bShouldHide;
+}
+
+namespace ActorDetails
+{
+	bool bCVarDisableComponentEditorButtonsAreaOnMultiSelect = false;
+	FAutoConsoleVariableRef CVarLegacyAlwaysDisableComponentEditorButtonsAreaOnMultiSelect(
+		TEXT("Editor.LevelEditor.LegacyAlwaysDisableComponentEditorButtonsAreaOnMultiSelect"),
+		bCVarDisableComponentEditorButtonsAreaOnMultiSelect,
+		TEXT("Disables all component editor buttons in a level editor details panel.  CVar used for backwards compatibility. Please update extension widgets to change visibility based on selected subobject nodes."),
+		ECVF_Default);
+}
+
+void SActorDetails::Construct(
+	const FArguments& InArgs,
+	UTypedElementSelectionSet* InSelectionSet,
+	const FName TabIdentifier,
+	TSharedPtr<FUICommandList> InCommandList,
+	TSharedPtr<FTabManager> InTabManager,
+	TWeakPtr<FActorDetailsSCSEditorUICustomization> InLevelEditorSubobjectEditorCustomization)
 {
 	SelectionSet = InSelectionSet;
 	checkf(SelectionSet, TEXT("SActorDetails must be constructed with a valid selection set!"));
@@ -274,19 +485,35 @@ void SActorDetails::Construct(const FArguments& InArgs, UTypedElementSelectionSe
 		.Visibility(this, &SActorDetails::GetComponentEditorVisibility);
 
 	FModuleManager::LoadModuleChecked<FSubobjectEditorModule>("SubobjectEditor");
+
+	LevelEditorSubobjectEditorCustomization = InLevelEditorSubobjectEditorCustomization.Pin();
+	
+	TSharedPtr<ISCSEditorUICustomization> SubobjectEditorCustomization = nullptr;
+	if (ActorDetails::bCVarDisableComponentEditorButtonsAreaOnMultiSelect)
+	{
+		SubobjectEditorCustomization = LevelEditorSubobjectEditorCustomization;
+	}
+	else
+	{
+		InternalSubobjectEditorCustomization = MakeShared<FInternalSubobjectCustomization>(LevelEditorSubobjectEditorCustomization);
+		SubobjectEditorCustomization = InternalSubobjectEditorCustomization;
+	}
 	
 	SubobjectEditor = SNew(SSubobjectInstanceEditor)
 		.ObjectContext(this, &SActorDetails::GetObjectContext)
 		.AllowEditing(this, &SActorDetails::GetAllowComponentTreeEditing)
 		.OnSelectionUpdated(this, &SActorDetails::OnSubobjectEditorTreeViewSelectionChanged)
 		.OnItemDoubleClicked(this, &SActorDetails::OnSubobjectEditorTreeViewItemDoubleClicked)
-		.SCSEditorUICustomization(InArgs._SCSEditorUICustomization);
+		.SCSEditorUICustomization(SubobjectEditorCustomization);
 
 	ComponentsBox->SetContent(SubobjectEditor.ToSharedRef());
 
 	TSharedRef<SWidget> SubobjectEditorButtonBox = SubobjectEditor->GetToolButtonsBox().ToSharedRef();
-	SubobjectEditorButtonBox->SetVisibility(MakeAttributeSP(this, &SActorDetails::GetComponentEditorButtonsVisibility));
-
+	
+	if (ActorDetails::bCVarDisableComponentEditorButtonsAreaOnMultiSelect)
+	{
+		SubobjectEditorButtonBox->SetVisibility(MakeAttributeSP(this, &SActorDetails::GetComponentEditorButtonsVisibility));
+	}
 
 	TFunction<TArray<TTypedElement<ITypedElementDetailsInterface>>()> GetDetailsHandles = [this]()
 	{
@@ -492,6 +719,13 @@ void SActorDetails::RefreshTopLevelElements(TArrayView<const TTypedElement<IType
 		}
 	}
 
+	if (InternalSubobjectEditorCustomization)
+	{
+		const bool bHideButtons = TopLevelElements.Num() != 1;
+		InternalSubobjectEditorCustomization->SetHideAddComponentButton(bHideButtons);
+		InternalSubobjectEditorCustomization->SetHideBlueprintButton(bHideButtons);
+	}
+
 	// Update the underlying details view and the Elements buttons
 	SetElementDetailsObjects(TopLevelElements, bForceRefresh, bOverrideLock, &InDetailsElements);
 
@@ -675,14 +909,6 @@ FString SActorDetails::GetReferencerName() const
 	DetailsView->ForceRefresh();
 }
 
-void SActorDetails::SetSubobjectEditorUICustomization(TSharedPtr<ISCSEditorUICustomization> ActorDetailsSubobjectEditorUICustomization)
-{
-	if(SubobjectEditor.IsValid())
-	{
-		SubobjectEditor->SetUICustomization(ActorDetailsSubobjectEditorUICustomization);
-	}
-}
-
 void SActorDetails::OnComponentsEditedInWorld()
 {
 	if (AActor* Actor = GetActorContext())
diff --git a/Engine/Source/Editor/LevelEditor/Private/SActorDetails.h b/Engine/Source/Editor/LevelEditor/Private/SActorDetails.h
--- a/Engine/Source/Editor/LevelEditor/Private/SActorDetails.h
+++ b/Engine/Source/Editor/LevelEditor/Private/SActorDetails.h
@@ -40,10 +40,14 @@ class SActorDetails : public SCompoundWidget, public FEditorUndoClient, public F
 {
 public:
 	SLATE_BEGIN_ARGS(SActorDetails) {}
-		SLATE_ARGUMENT(TSharedPtr<ISCSEditorUICustomization>, SCSEditorUICustomization)
 	SLATE_END_ARGS()
 
-	void Construct(const FArguments& InArgs, UTypedElementSelectionSet* InSelectionSet, const FName TabIdentifier, TSharedPtr<FUICommandList> InCommandList, TSharedPtr<FTabManager> InTabManager);
+	void Construct(const FArguments& InArgs,
+		UTypedElementSelectionSet* InSelectionSet,
+		const FName TabIdentifier,
+		TSharedPtr<FUICommandList> InCommandList,
+		TSharedPtr<FTabManager> InTabManager,
+		TWeakPtr<FActorDetailsSCSEditorUICustomization> InLevelEditorSubobjectEditorCustomization);
 	~SActorDetails();
 
 	/**
@@ -76,13 +80,13 @@ public:
 	 */
 	void SetActorDetailsRootCustomization(TSharedPtr<FDetailsViewObjectFilter> InActorDetailsObjectFilter, TSharedPtr<class IDetailRootObjectCustomization> ActorDetailsRootCustomization);
 
-	UE_DEPRECATED(5.0, "SetSCSEditorUICustomization is deprecated, please use SetSubobjectEditorUICustomization instead.")
-	void SetSCSEditorUICustomization(TSharedPtr<ISCSEditorUICustomization> ActorDetailsSCSEditorUICustomization) { SetSubobjectEditorUICustomization(ActorDetailsSCSEditorUICustomization); }
-	
-	/** Sets the UI customization of the Suobject inside this details panel. */
-	void SetSubobjectEditorUICustomization(TSharedPtr<ISCSEditorUICustomization> ActorDetailsSCSEditorUICustomization);
-
 private:
+	/**
+	 * A customization that the ActorDetails widget manages that maintains the visibility of the
+	 * add component and convert to blueprint buttons.
+	 */
+	class FInternalSubobjectCustomization;
+	
 	void RefreshTopLevelElements(TArrayView<const TTypedElement<ITypedElementDetailsInterface>> InDetailsElements, const bool bForceRefresh, const bool bOverrideLock);
 	void RefreshSubobjectTreeElements(TArrayView<const TSharedPtr<FSubobjectEditorTreeNode>> InSelectedNodes, const bool bForceRefresh, const bool bOverrideLock);
 	void SetElementDetailsObjects(TArrayView<const TUniquePtr<ITypedElementDetailsObject>> InElementDetailsObjects, const bool bForceRefresh, const bool bOverrideLock, TArrayView<const TTypedElement<ITypedElementDetailsInterface>>* InDetailsElementsPtr = nullptr);
@@ -121,6 +125,11 @@ private:
 	TSharedPtr<class IDetailsView> DetailsView;
 	TSharedPtr<SBox> ComponentsBox;
 
+	// Subobject Editor customization that hosts other customizations. Scoped to the level editor
+	TSharedPtr<FActorDetailsSCSEditorUICustomization> LevelEditorSubobjectEditorCustomization;
+	// Subobject Editor customization that is maintained as last priority in the LevelEditorCustomization
+	// and is specific to the SActorDetails
+	TSharedPtr<FInternalSubobjectCustomization> InternalSubobjectEditorCustomization;
 	/** The subobject editor provides a tree widget that allows for editing of subobjects */
 	TSharedPtr<SSubobjectEditor> SubobjectEditor;
 
diff --git a/Engine/Source/Editor/LevelEditor/Private/SLevelEditor.cpp b/Engine/Source/Editor/LevelEditor/Private/SLevelEditor.cpp
--- a/Engine/Source/Editor/LevelEditor/Private/SLevelEditor.cpp
+++ b/Engine/Source/Editor/LevelEditor/Private/SLevelEditor.cpp
@@ -226,7 +226,7 @@ void SLevelEditor::RegisterMenus()
 
 void SLevelEditor::Construct( const SLevelEditor::FArguments& InArgs)
 {
-	ActorDetailsSCSEditorUICustomization = FActorDetailsSCSEditorUICustomization::GetInstance();
+	ActorDetailsSCSEditorUICustomization = FActorDetailsSCSEditorUICustomization::NewInstance();
 	
 	// Important: We use raw bindings here because we are releasing our binding in our destructor (where a weak pointer would be invalid)
 	// It's imperative that our delegate is removed in the destructor for the level editor module to play nicely with reloading.
@@ -2217,8 +2217,7 @@ void SLevelEditor::AddStandaloneLevelViewport( const TSharedRef<SLevelViewport>&
 
 TSharedRef<SWidget> SLevelEditor::CreateActorDetails( const FName TabIdentifier )
 {
-	TSharedRef<SActorDetails> ActorDetails = SNew( SActorDetails, GetMutableElementSelectionSet(), TabIdentifier, LevelEditorCommands, GetTabManager())
-		.SCSEditorUICustomization(ActorDetailsSCSEditorUICustomization);
+	TSharedRef<SActorDetails> ActorDetails = SNew( SActorDetails, GetMutableElementSelectionSet(), TabIdentifier, LevelEditorCommands, GetTabManager(), ActorDetailsSCSEditorUICustomization);
 
 	ActorDetails->SetActorDetailsRootCustomization(ActorDetailsObjectFilter, ActorDetailsRootCustomization);
 
diff --git a/Engine/Source/Editor/SubobjectEditor/Private/SSubobjectEditor.cpp b/Engine/Source/Editor/SubobjectEditor/Private/SSubobjectEditor.cpp
--- a/Engine/Source/Editor/SubobjectEditor/Private/SSubobjectEditor.cpp
+++ b/Engine/Source/Editor/SubobjectEditor/Private/SSubobjectEditor.cpp
@@ -3068,6 +3068,15 @@ FSubobjectDataHandle SSubobjectEditor::GetObjectContextHandle() const
 	
 }
 
+FSubobjectData* SSubobjectEditor::GetDataContextHandle() const
+{
+	if (RootNodes.Num() > 0)
+	{
+		return RootNodes[0]->GetDataHandle().GetData();
+	}
+	return nullptr;
+}
+
 UBlueprint* SSubobjectEditor::GetBlueprint() const
 {
 	if (AActor* Actor = Cast<AActor>(GetObjectContext()))
diff --git a/Engine/Source/Editor/SubobjectEditor/Public/SSubobjectEditor.h b/Engine/Source/Editor/SubobjectEditor/Public/SSubobjectEditor.h
--- a/Engine/Source/Editor/SubobjectEditor/Public/SSubobjectEditor.h
+++ b/Engine/Source/Editor/SubobjectEditor/Public/SSubobjectEditor.h
@@ -598,6 +598,13 @@ public:
 	/** SubobjectHandle of the current object that is represented by the subobject editor */
 	UE_API FSubobjectDataHandle GetObjectContextHandle() const;
 
+	/** Get the SubobjectData for the current context represented by the subobject editor
+	 * This is a newer API than GetObjectContext + GetObjectContextHandle.  It is introduced to allow
+	 * access to non-UObject based data and for cases where SubobjectData is a proxy for multiple objects
+	 * in bulk-edit cases.
+	 */
+	FSubobjectData* GetDataContextHandle() const;
+
 	/** Refresh the type list presented by the add component button when clicked */
 	UE_API void RefreshComponentTypesList();
 	
