Added the ability to display specific RGBA color channels in the active editor viewports, as well as options to control how the alpha channel specifically is displayed

diff --git a/Engine/Shaders/Private/PostProcessChannelMask.usf b/Engine/Shaders/Private/PostProcessChannelMask.usf
new file mode 100644
--- /dev/null
+++ b/Engine/Shaders/Private/PostProcessChannelMask.usf
@@ -0,0 +1,54 @@
+﻿// Copyright Epic Games, Inc. All Rights Reserved.
+
+/*=============================================================================
+	PostProcessChannelMask.usf: Masks to a single channel and outputs it as greyscale.
+=============================================================================*/
+
+#include "Common.ush"
+#include "GammaCorrectionCommon.ush"
+#include "ScreenPass.ush"
+
+Texture2D SceneColorTexture;
+SamplerState SceneColorSampler;
+
+float4 ChannelMask;
+uint bMaskChannels;
+uint bInvertAlphaChannelMask;
+uint bDisplayCheckerboard;
+FScreenTransform SvPositionToViewPos;
+
+#define CheckerSize 16
+#define InvCheckerSize (1.0f / CheckerSize)
+#define CheckerWhite float3(0.216f, 0.216f, 0.216f)
+#define CheckerGray float3(0.033f, 0.033f, 0.033f)
+
+float3 GetCheckerboardColor(float2 PixelPos)
+{
+	float2 ViewPos = ApplyScreenTransform(PixelPos, SvPositionToViewPos);
+	int2 CheckerCell = floor(ViewPos * float2(InvCheckerSize, InvCheckerSize) * View.ResolutionFractionAndInv.y);
+	return CheckerCell.x % 2 == CheckerCell.y % 2 ? CheckerWhite : CheckerGray;
+}
+
+void ChannelMask_MainPS(FScreenVertexOutput Input, out float4 OutColor : SV_Target0)
+{
+	float4 SceneColor = Texture2DSample(SceneColorTexture, SceneColorSampler, Input.UV);
+
+	// Always output the original scene alpha, regardless of what the output rgb ends up being
+	OutColor.a = SceneColor.a;
+	
+	// Non-alpha channels will be encoded in sRGB by this point, so to properly composite with checkerboard, convert to linear space
+	// TODO: Native HDR (linear ScRGB or Rec2020 PQ) or gamma 2.2 will need to be subsequently supported
+	SceneColor.rgb = sRGBToLinear(SceneColor.rgb);
+	
+	// Invert the alpha if needed
+	SceneColor.a = bInvertAlphaChannelMask ? saturate(1.0 - SceneColor.a) : SceneColor.a;
+	
+	float3 Channel = bMaskChannels ? dot(SceneColor.rgba, ChannelMask) : SceneColor.rgb;
+	float3 CheckerboardColor = bDisplayCheckerboard ? GetCheckerboardColor(Input.Position.xy) : 0.0f;
+
+	// Perform a naive premult alpha composite with the raw contents of the scene color's alpha.
+	// When not inverting the alpha, this is a technically incorrect blending, as Channel has been premultiplied with 1.0 - SceneColor.a
+	// and not the correct SceneColor.a (to show the correct value, the scene color would need to be un-pre-multiplied and re-pre-multiplied by the correct alphas)
+	// but we will simply output the "wrong" values to demonstrate the incorrectness
+	OutColor.rgb = LinearToSrgb(Channel + CheckerboardColor * (1.0 - SceneColor.a));
+}
\ No newline at end of file
diff --git a/Engine/Source/Editor/EditorStyle/Private/StarshipStyle.cpp b/Engine/Source/Editor/EditorStyle/Private/StarshipStyle.cpp
--- a/Engine/Source/Editor/EditorStyle/Private/StarshipStyle.cpp
+++ b/Engine/Source/Editor/EditorStyle/Private/StarshipStyle.cpp
@@ -4963,6 +4963,12 @@ void FStarshipEditorStyle::FStyle::SetupLevelEditorStyle()
 		Set( "EditorViewport.CollisionPawn",          new IMAGE_BRUSH_SVG("Starship/Common/PlayerCollision", Icon16x16 ) );
 		Set( "EditorViewport.CollisionVisibility",    new IMAGE_BRUSH_SVG("Starship/Common/VisibilityCollision", Icon16x16 ) );
 
+		Set( "EditorViewport.AllChannelsMask",   new IMAGE_BRUSH_SVG("Starship/Common/ColorChannelRGB", Icon16x16 ) );
+		Set( "EditorViewport.RedChannelMask",   new IMAGE_BRUSH_SVG("Starship/Common/ColorChannelRed", Icon16x16 ) );
+		Set( "EditorViewport.GreenChannelMask", new IMAGE_BRUSH_SVG("Starship/Common/ColorChannelGreen", Icon16x16 ) );
+		Set( "EditorViewport.BlueChannelMask",  new IMAGE_BRUSH_SVG("Starship/Common/ColorChannelBlue", Icon16x16 ) );
+		Set( "EditorViewport.AlphaChannelMask", new IMAGE_BRUSH_SVG("Starship/Common/ColorChannelAlpha", Icon16x16 ) );
+
 		Set( "EditorViewport.Perspective", new IMAGE_BRUSH_SVG("Starship/Common/ViewPerspective", Icon16x16 ) );
 		Set( "EditorViewport.Top",         new IMAGE_BRUSH_SVG("Starship/Common/ViewTop", Icon16x16 ) );
 		Set( "EditorViewport.Left",        new IMAGE_BRUSH_SVG("Starship/Common/ViewLeft", Icon16x16 ) );
diff --git a/Engine/Source/Editor/LevelEditor/Private/SLevelViewport.cpp b/Engine/Source/Editor/LevelEditor/Private/SLevelViewport.cpp
--- a/Engine/Source/Editor/LevelEditor/Private/SLevelViewport.cpp
+++ b/Engine/Source/Editor/LevelEditor/Private/SLevelViewport.cpp
@@ -632,6 +632,9 @@ void SLevelViewport::ConstructLevelEditorViewportClient(FLevelEditorViewportInst
 	LevelViewportClient->CurrentVirtualTextureVisualizationMode = ViewportInstanceSettings.VirtualTextureVisualizationMode;
 	LevelViewportClient->CurrentRayTracingDebugVisualizationMode = ViewportInstanceSettings.RayTracingDebugVisualizationMode;
 	LevelViewportClient->CurrentGPUSkinCacheVisualizationMode = ViewportInstanceSettings.GPUSkinCacheVisualizationMode;
+	LevelViewportClient->ChannelMaskParams.ColorChannelMask = (EColorChannelMask)ViewportInstanceSettings.ColorChannelMask;
+	LevelViewportClient->ChannelMaskParams.bInvertAlphaChannelMask = ViewportInstanceSettings.bInvertAlphaChannelMask;
+	LevelViewportClient->ChannelMaskParams.bDrawAlphaBlendedCheckerboard = ViewportInstanceSettings.bDrawAlphaBlendedCheckerboard;
 	LevelViewportClient->ExposureSettings = ViewportInstanceSettings.ExposureSettings;
 	if(ViewportInstanceSettings.ViewportType == LVT_Perspective)
 	{
@@ -2331,6 +2334,10 @@ FName SLevelViewport::RegisterViewportToolbar() const
 				RightSection.AddEntry(UE::UnrealEd::CreateViewModesSubmenu());
 			}
 
+			{
+				RightSection.AddEntry(UE::LevelEditor::CreateDynamicChannelsSubmenu());
+			}
+
 			// Add the "Show" submenu.
 			{
 				// Stay backward-compatible with the old viewport toolbar.
@@ -3048,6 +3055,9 @@ void SLevelViewport::SaveConfig(const FString& ConfigName) const
 		ViewportInstanceSettings.VirtualTextureVisualizationMode = LevelViewportClient->CurrentVirtualTextureVisualizationMode;
 		ViewportInstanceSettings.RayTracingDebugVisualizationMode = LevelViewportClient->CurrentRayTracingDebugVisualizationMode;
 		ViewportInstanceSettings.GPUSkinCacheVisualizationMode = LevelViewportClient->CurrentGPUSkinCacheVisualizationMode;
+		ViewportInstanceSettings.ColorChannelMask = (int32)LevelViewportClient->ChannelMaskParams.ColorChannelMask;
+		ViewportInstanceSettings.bInvertAlphaChannelMask = LevelViewportClient->ChannelMaskParams.bInvertAlphaChannelMask;
+		ViewportInstanceSettings.bDrawAlphaBlendedCheckerboard = LevelViewportClient->ChannelMaskParams.bDrawAlphaBlendedCheckerboard;
 		ViewportInstanceSettings.ExposureSettings = LevelViewportClient->ExposureSettings;
 		ViewportInstanceSettings.FOVAngle = LevelViewportClient->FOVAngle;
 		ViewportInstanceSettings.bAllowCinematicControl = LevelViewportClient->AllowsCinematicControl();
diff --git a/Engine/Source/Editor/LevelEditor/Private/ViewportToolbar/LevelEditorSubmenus.cpp b/Engine/Source/Editor/LevelEditor/Private/ViewportToolbar/LevelEditorSubmenus.cpp
--- a/Engine/Source/Editor/LevelEditor/Private/ViewportToolbar/LevelEditorSubmenus.cpp
+++ b/Engine/Source/Editor/LevelEditor/Private/ViewportToolbar/LevelEditorSubmenus.cpp
@@ -904,6 +904,165 @@ TSharedRef<SWidget> CreateOrthographicViewportGestureDirectionWidget()
 	return SNullWidget::NullWidget;
 }
 
+FToolMenuEntry CreateChannelsSubmenu(const TWeakPtr<::SLevelViewport>& LevelViewportWeak, bool bUseFullLabel)
+{
+	static FText ChannelLabels[] =
+	{
+		LOCTEXT("ChannelRGB_Label", "RGB"),
+		LOCTEXT("ChannelRed_Label", "Red"),
+		LOCTEXT("ChannelGreen_Label", "Green"),
+		LOCTEXT("ChannelBlue_Label", "Blue"),
+		LOCTEXT("ChannelAlpha_Label", "Alpha")
+	};
+
+	static FName ChannelIcons[] =
+	{
+		"EditorViewport.AllChannelsMask",
+		"EditorViewport.RedChannelMask",
+		"EditorViewport.GreenChannelMask",
+		"EditorViewport.BlueChannelMask",
+		"EditorViewport.AlphaChannelMask"
+	};
+
+	FToolMenuEntry Entry = FToolMenuEntry::InitSubMenu(
+		"Channels",
+		TAttribute<FText>::CreateLambda([LevelViewportWeak, bUseFullLabel]
+		{
+			if (const TSharedPtr<::SLevelViewport> Viewport = LevelViewportWeak.Pin())
+			{
+				const int8 ChannelMaskIndex = (int8)Viewport->GetViewportClient()->ChannelMaskParams.ColorChannelMask;
+				return bUseFullLabel
+					? FText::Format(LOCTEXT("ChannelsMenu_LabelFormat", "Channels ({0})"), ChannelLabels[ChannelMaskIndex + 1])
+					: ChannelLabels[ChannelMaskIndex + 1];
+			}
+
+			return FText::GetEmpty();
+		}),
+		LOCTEXT("ChannelsMenu_ToolTip", "Display specific color channels in this viewport"),
+		FNewToolMenuDelegate::CreateLambda([LevelViewportWeak](UToolMenu* InMenu)
+		{
+			FToolMenuSection& ChannelsSubMenuSection =
+				InMenu->AddSection("ChannelsSubMenuSection", LOCTEXT("ChannelsSubMenuSectionHeader", "Channels"));
+
+			ChannelsSubMenuSection.AddMenuEntry(FEditorViewportCommands::Get().AllChannelsMask);
+			ChannelsSubMenuSection.AddMenuEntry(FEditorViewportCommands::Get().RedChannelMask);
+			ChannelsSubMenuSection.AddMenuEntry(FEditorViewportCommands::Get().GreenChannelMask);
+			ChannelsSubMenuSection.AddMenuEntry(FEditorViewportCommands::Get().BlueChannelMask);
+			ChannelsSubMenuSection.AddMenuEntry(FEditorViewportCommands::Get().AlphaChannelMask);
+
+			FToolMenuSection& OptionsSubMenuSection =
+				InMenu->AddSection("OptionsSubMenuSection", LOCTEXT("OptionsSubMenuSectionHeader", "Options"));
+
+			OptionsSubMenuSection.AddMenuEntry(
+				"ShowChannelsInViewportToolbar",
+				LOCTEXT("ShowChannelsInToolbarLabel", "Show Channels in Viewport Toolbar"),
+				LOCTEXT("ShowChannelsInToolbarToolTip", "Show the channel selector in the level editor toolbar"),
+				FSlateIcon(),
+				FUIAction(
+					FExecuteAction::CreateLambda([]()
+					{
+						if (ULevelEditorViewportSettings* Settings = GetMutableDefault<ULevelEditorViewportSettings>())
+						{
+							Settings->bShowChannelsInToolbar = !Settings->bShowChannelsInToolbar;
+						}
+					}),
+					FCanExecuteAction(),
+					FGetActionCheckState::CreateLambda([]()
+					{
+						return GetDefault<ULevelEditorViewportSettings>()->bShowChannelsInToolbar ? ECheckBoxState::Checked : ECheckBoxState::Unchecked;
+					})
+				),
+				EUserInterfaceActionType::ToggleButton
+			);
+
+			OptionsSubMenuSection.AddMenuEntry(
+				"InvertAlpha",
+				LOCTEXT("InvertAlphaLabel", "View Alpha as Inverted"),
+				LOCTEXT("InvertAlphaToolTip", "Invert the alpha value displayed when viewing the alpha channel"),
+				FSlateIcon(),
+				FUIAction(
+					FExecuteAction::CreateLambda([LevelViewportWeak]
+					{
+						if (const TSharedPtr<::SLevelViewport> Viewport = LevelViewportWeak.Pin())
+						{
+							Viewport->GetViewportClient()->ChannelMaskParams.bInvertAlphaChannelMask = !Viewport->GetViewportClient()->ChannelMaskParams.bInvertAlphaChannelMask;
+							if (!Viewport->GetViewportClient()->ChannelMaskParams.bInvertAlphaChannelMask)
+							{
+								// Disable the checkerboard blend when not inverting alpha as it will be broken due to UE conventions around alpha
+								Viewport->GetViewportClient()->ChannelMaskParams.bDrawAlphaBlendedCheckerboard = false;
+							}
+							
+							Viewport->GetViewportClient()->Invalidate();
+						}
+					}),
+					FCanExecuteAction(),
+					FGetActionCheckState::CreateLambda([LevelViewportWeak]
+					{
+						bool bChecked = false;
+						if (const TSharedPtr<::SLevelViewport> Viewport = LevelViewportWeak.Pin())
+						{
+							bChecked = Viewport->GetViewportClient()->ChannelMaskParams.bInvertAlphaChannelMask; 
+						}
+
+						return bChecked ? ECheckBoxState::Checked : ECheckBoxState::Unchecked;
+					})
+				),
+				EUserInterfaceActionType::ToggleButton
+			);
+
+			OptionsSubMenuSection.AddMenuEntry(
+				"VisualizeAlphaWithCheckerboard",
+				LOCTEXT("VisualizeWithCheckerboardLabel", "Visualize Alpha with Checkerboard"),
+				LOCTEXT("VisualizeWithCheckerboardToolTip", "Displays a checkerboard underneath the viewport to help visualize the alpha channel"),
+				FSlateIcon(),
+				FUIAction(
+					FExecuteAction::CreateLambda([LevelViewportWeak]
+					{
+						if (const TSharedPtr<::SLevelViewport> Viewport = LevelViewportWeak.Pin())
+						{
+							Viewport->GetViewportClient()->ChannelMaskParams.bDrawAlphaBlendedCheckerboard = !Viewport->GetViewportClient()->ChannelMaskParams.bDrawAlphaBlendedCheckerboard;
+							Viewport->GetViewportClient()->Invalidate();
+						}
+					}),
+					FCanExecuteAction::CreateLambda([LevelViewportWeak]
+					{
+						if (const TSharedPtr<::SLevelViewport> Viewport = LevelViewportWeak.Pin())
+						{
+							return Viewport->GetViewportClient()->ChannelMaskParams.bInvertAlphaChannelMask;
+						}
+
+						return false;
+					}),
+					FGetActionCheckState::CreateLambda([LevelViewportWeak]
+					{
+						bool bChecked = false;
+						if (const TSharedPtr<::SLevelViewport> Viewport = LevelViewportWeak.Pin())
+						{
+							bChecked = Viewport->GetViewportClient()->ChannelMaskParams.bDrawAlphaBlendedCheckerboard; 
+						}
+
+						return bChecked ? ECheckBoxState::Checked : ECheckBoxState::Unchecked;
+					})
+				),
+				EUserInterfaceActionType::ToggleButton
+			);
+		}),
+		/*bInOpenSubMenuOnClick=*/false,
+		TAttribute<FSlateIcon>::CreateLambda(
+			[LevelViewportWeak]
+			{
+				if (const TSharedPtr<::SLevelViewport> Viewport = LevelViewportWeak.Pin())
+				{
+					const int8 ChannelMaskIndex = (int8)Viewport->GetViewportClient()->ChannelMaskParams.ColorChannelMask;
+					return FSlateIcon(FAppStyle::GetAppStyleSetName(), ChannelIcons[ChannelMaskIndex + 1]);
+				}
+				return FSlateIcon();
+			})
+	);
+
+	return Entry;
+}
+
 } // namespace UE::LevelEditor::Private
 
 namespace UE::LevelEditor
@@ -1188,6 +1347,14 @@ void PopulateViewModesMenu(UToolMenu* InMenu)
 		);
 	}
 
+	{
+		FToolMenuSection& Section =
+			InMenu->AddSection("LevelViewportChannels", LOCTEXT("ChannelsHeader", "Channels"), InsertPosition);
+
+		constexpr bool bUseFullLabel = true;
+		Section.AddEntry(UE::LevelEditor::Private::CreateChannelsSubmenu(LevelViewport.ToWeakPtr(), bUseFullLabel));
+	}
+
 	{
 		FToolMenuSection& Section =
 			InMenu->AddSection("LevelViewportLandscape", LOCTEXT("LandscapeHeader", "Landscape"), InsertPosition);
@@ -1938,6 +2105,29 @@ FToolMenuEntry CreateShowStatsSubmenu(bool bInAddToggleStatsCheckbox, TAttribute
 }
 #endif
 
+FToolMenuEntry CreateDynamicChannelsSubmenu()
+{
+	// This has to be a dynamic entry for the ViewModes submenu's label to be able to access the context.
+	return FToolMenuEntry::InitDynamicEntry(
+		"DynamicChannels",
+		FNewToolMenuSectionDelegate::CreateLambda(
+			[](FToolMenuSection& InDynamicSection) -> void
+			{
+				if (ULevelViewportContext* const LevelViewportContext = InDynamicSection.FindContext<ULevelViewportContext>())
+				{
+					constexpr bool bUseFullLabel = false;
+					FToolMenuEntry& Entry = InDynamicSection.AddEntry(UE::LevelEditor::Private::CreateChannelsSubmenu(LevelViewportContext->LevelViewport, bUseFullLabel));
+					Entry.Visibility = TAttribute<bool>::CreateLambda([]
+					{
+						return (bool)GetDefault<ULevelEditorViewportSettings>()->bShowChannelsInToolbar;
+					});
+					Entry.ToolBarData.ResizeParams.ClippingPriority = 800;
+				}
+			}
+		)
+	);
+}
+
 FToolMenuEntry CreateShowSubmenu()
 {
 	return UE::UnrealEd::CreateShowSubmenu(
diff --git a/Engine/Source/Editor/LevelEditor/Private/ViewportToolbar/LevelEditorSubmenus.h b/Engine/Source/Editor/LevelEditor/Private/ViewportToolbar/LevelEditorSubmenus.h
--- a/Engine/Source/Editor/LevelEditor/Private/ViewportToolbar/LevelEditorSubmenus.h
+++ b/Engine/Source/Editor/LevelEditor/Private/ViewportToolbar/LevelEditorSubmenus.h
@@ -38,6 +38,7 @@ FToolMenuEntry CreateShowStatsSubmenu(
 	bool bInAddToggleStatsCheckbox = false, TAttribute<FText> InLabelOverride = TAttribute<FText>()
 );
 #endif
+FToolMenuEntry CreateDynamicChannelsSubmenu();
 FToolMenuEntry CreateShowVolumesSubmenu();
 FToolMenuEntry CreateShowSubmenu();
 FToolMenuEntry CreatePIEShowSubmenu();
diff --git a/Engine/Source/Editor/UnrealEd/Classes/Settings/LevelEditorViewportSettings.h b/Engine/Source/Editor/UnrealEd/Classes/Settings/LevelEditorViewportSettings.h
--- a/Engine/Source/Editor/UnrealEd/Classes/Settings/LevelEditorViewportSettings.h
+++ b/Engine/Source/Editor/UnrealEd/Classes/Settings/LevelEditorViewportSettings.h
@@ -119,6 +119,9 @@ struct FLevelEditorViewportInstanceSettings
 		, VirtualTextureVisualizationMode()
 		, RayTracingDebugVisualizationMode()
 		, GPUSkinCacheVisualizationMode()
+		, ColorChannelMask(INDEX_NONE)
+		, bInvertAlphaChannelMask(true)
+		, bDrawAlphaBlendedCheckerboard(false)
 		, ExposureSettings()
 		, FOVAngle(EditorViewportDefs::DefaultPerspectiveFOVAngle)
 		, FarViewPlane(0)
@@ -192,6 +195,18 @@ struct FLevelEditorViewportInstanceSettings
 	UPROPERTY(config)
 	FName GPUSkinCacheVisualizationMode;
 
+	/** The color channel to display for this viewport */
+	UPROPERTY(config)
+	int32 ColorChannelMask;
+
+	/** Whether to invert the alpha channel when displaying it */
+	UPROPERTY(config)
+	bool bInvertAlphaChannelMask;
+
+	/** Whether to display an alpha-blended checkerboard beneath the viewport to help visualize alpha in the viewport */
+	UPROPERTY(config)
+	bool bDrawAlphaBlendedCheckerboard;
+	
 	/** Setting to allow designers to override the automatic expose. */
 	UPROPERTY(config)
 	FExposureSettings ExposureSettings;
@@ -644,6 +659,10 @@ public:
 	UPROPERTY(EditAnywhere, config, Category = LookAndFeel)
 	FColor PropertyColorationColorForMatchingObjects = FColor::Red;
 
+	/** When enabled, the channel mask menu will show up in the editor toolbar next to the View Modes button */
+	UPROPERTY(config)
+	uint32 bShowChannelsInToolbar:1;
+
 private:
 
 	// Per-instance viewport settings.
diff --git a/Engine/Source/Editor/UnrealEd/Private/EditorViewportClient.cpp b/Engine/Source/Editor/UnrealEd/Private/EditorViewportClient.cpp
--- a/Engine/Source/Editor/UnrealEd/Private/EditorViewportClient.cpp
+++ b/Engine/Source/Editor/UnrealEd/Private/EditorViewportClient.cpp
@@ -675,6 +675,13 @@ FEditorViewportClient::~FEditorViewportClient()
 	FEditorDelegates::MapChange.RemoveAll(this);
 }
 
+FEditorViewportClient::FColorChannelMaskParams::FColorChannelMaskParams()
+	: ColorChannelMask(EColorChannelMask::All)
+	, bInvertAlphaChannelMask(false)
+	, bDrawAlphaBlendedCheckerboard(false)
+{
+}
+
 FLevelEditorViewportClient::FRealtimeOverride::FRealtimeOverride(bool bInIsRealtime, FText InSystemDisplayName)
 	: SystemDisplayName(InSystemDisplayName)
 	, bIsRealtime(bInIsRealtime)
@@ -4792,6 +4799,9 @@ void FEditorViewportClient::Draw(FViewport* InViewport, FCanvas* Canvas)
 	ViewFamily.DebugDPIScale = GetDPIScale();
 
 	ViewFamily.EngineShowFlags = UseEngineShowFlags;
+	ViewFamily.ChannelMaskParams.ColorChannelMask = ChannelMaskParams.ColorChannelMask;
+	ViewFamily.ChannelMaskParams.bInvertAlphaChannelMask = ChannelMaskParams.bInvertAlphaChannelMask;
+	ViewFamily.ChannelMaskParams.bDrawAlphaBlendedCheckerboard = ChannelMaskParams.bDrawAlphaBlendedCheckerboard;
 
 	ViewFamily.bIsHDR = Viewport->IsHDRViewport();
 	
@@ -6376,6 +6386,26 @@ bool FEditorViewportClient::IsSetGeometryInspectionShowWireframe() const
 	return bIsGeometryInspectionWireframeOn;
 }
 
+void FEditorViewportClient::SetChannelMask(EColorChannelMask InColorChannelMask)
+{
+	// If the color channel mask is already set to the specified mask, toggle it off by setting the mask to All
+	if (ChannelMaskParams.ColorChannelMask == InColorChannelMask)
+	{
+		ChannelMaskParams.ColorChannelMask = EColorChannelMask::All;
+	}
+	else
+	{
+		ChannelMaskParams.ColorChannelMask = InColorChannelMask;
+	}
+	
+	Invalidate();
+}
+
+bool FEditorViewportClient::IsChannelBeingMasked(EColorChannelMask InColorChannelMask) const
+{
+	return ChannelMaskParams.ColorChannelMask == InColorChannelMask;
+}
+
 void FEditorViewportClient::SetViewMode(EViewModeIndex InViewModeIndex)
 {
 	ViewModeParam = -1; // Reset value when the viewmode changes
diff --git a/Engine/Source/Editor/UnrealEd/Private/EditorViewportCommands.cpp b/Engine/Source/Editor/UnrealEd/Private/EditorViewportCommands.cpp
--- a/Engine/Source/Editor/UnrealEd/Private/EditorViewportCommands.cpp
+++ b/Engine/Source/Editor/UnrealEd/Private/EditorViewportCommands.cpp
@@ -140,6 +140,12 @@ void FEditorViewportCommands::RegisterCommands()
 	UI_COMMAND( GeometryInspectionRandomColorMode, "Geometry Inspection Random Color Mode", "Assigns random colors to individual elements or components to improve visual distinction between overlapping or adjacent parts.", EUserInterfaceActionType::RadioButton, FInputChord());
 	UI_COMMAND( GeometryInspectionShowWireframe, "Use Wireframe", "Use Wireframe in Geometry Inspection Modes.", EUserInterfaceActionType::Check, FInputChord());
 
+	UI_COMMAND(AllChannelsMask, "RGB", "Show all color channels", EUserInterfaceActionType::RadioButton, FInputChord());
+	UI_COMMAND(RedChannelMask, "Red", "Show only the red color channel", EUserInterfaceActionType::RadioButton, FInputChord());
+	UI_COMMAND(GreenChannelMask, "Green", "Show only the green color channel", EUserInterfaceActionType::RadioButton, FInputChord());
+	UI_COMMAND(BlueChannelMask, "Blue", "Show only the blue color channel", EUserInterfaceActionType::RadioButton, FInputChord());
+	UI_COMMAND(AlphaChannelMask, "Alpha", "Show only the alpha color channel", EUserInterfaceActionType::RadioButton, FInputChord());
+
 	UI_COMMAND( ToggleRealTime, "Realtime", "Toggles real time rendering in this viewport", EUserInterfaceActionType::ToggleButton, FInputChord( EModifierKey::Control, EKeys::R ) );
 	UI_COMMAND( ToggleStats, "Show Stats", "Toggles the ability to show stats in this viewport (enables realtime)", EUserInterfaceActionType::ToggleButton, FInputChord( EModifierKey::Shift, EKeys::L ) );
 	UI_COMMAND( ToggleFPS, "Show FPS", "Toggles showing frames per second in this viewport (enables realtime)", EUserInterfaceActionType::ToggleButton, FInputChord( EModifierKey::Control|EModifierKey::Shift, EKeys::H ) );
diff --git a/Engine/Source/Editor/UnrealEd/Private/SEditorViewport.cpp b/Engine/Source/Editor/UnrealEd/Private/SEditorViewport.cpp
--- a/Engine/Source/Editor/UnrealEd/Private/SEditorViewport.cpp
+++ b/Engine/Source/Editor/UnrealEd/Private/SEditorViewport.cpp
@@ -747,6 +747,41 @@ void SEditorViewport::BindCommands()
 		FIsActionChecked::CreateSP( ClientRef, &FEditorViewportClient::IsSetGeometryInspectionShowWireframe )
 	);
 
+	CommandListRef.MapAction(
+		Commands.AllChannelsMask,
+		FExecuteAction::CreateSP( ClientRef, &FEditorViewportClient::SetChannelMask, EColorChannelMask::All ),
+		FCanExecuteAction(),
+		FIsActionChecked::CreateSP( ClientRef, &FEditorViewportClient::IsChannelBeingMasked, EColorChannelMask::All )
+	);
+
+	CommandListRef.MapAction(
+		Commands.RedChannelMask,
+		FExecuteAction::CreateSP( ClientRef, &FEditorViewportClient::SetChannelMask, EColorChannelMask::Red ),
+		FCanExecuteAction(),
+		FIsActionChecked::CreateSP( ClientRef, &FEditorViewportClient::IsChannelBeingMasked, EColorChannelMask::Red )
+	);
+
+	CommandListRef.MapAction(
+		Commands.GreenChannelMask,
+		FExecuteAction::CreateSP( ClientRef, &FEditorViewportClient::SetChannelMask, EColorChannelMask::Green ),
+		FCanExecuteAction(),
+		FIsActionChecked::CreateSP( ClientRef, &FEditorViewportClient::IsChannelBeingMasked, EColorChannelMask::Green )
+	);
+
+	CommandListRef.MapAction(
+		Commands.BlueChannelMask,
+		FExecuteAction::CreateSP( ClientRef, &FEditorViewportClient::SetChannelMask, EColorChannelMask::Blue ),
+		FCanExecuteAction(),
+		FIsActionChecked::CreateSP( ClientRef, &FEditorViewportClient::IsChannelBeingMasked, EColorChannelMask::Blue )
+	);
+
+	CommandListRef.MapAction(
+		Commands.AlphaChannelMask,
+		FExecuteAction::CreateSP( ClientRef, &FEditorViewportClient::SetChannelMask, EColorChannelMask::Alpha ),
+		FCanExecuteAction(),
+		FIsActionChecked::CreateSP( ClientRef, &FEditorViewportClient::IsChannelBeingMasked, EColorChannelMask::Alpha )
+	);
+
 	BindShowCommands( CommandListRef );
 }
 
diff --git a/Engine/Source/Editor/UnrealEd/Public/EditorViewportClient.h b/Engine/Source/Editor/UnrealEd/Public/EditorViewportClient.h
--- a/Engine/Source/Editor/UnrealEd/Public/EditorViewportClient.h
+++ b/Engine/Source/Editor/UnrealEd/Public/EditorViewportClient.h
@@ -40,6 +40,7 @@ class SEditorViewport;
 class UActorFactory;
 class UTypedElementViewportInteraction;
 enum class EViewStatusForScreenPercentage;
+enum class EColorChannelMask : int8;
 struct FGizmoState;
 struct FViewportClientNavigationHelper;
 namespace UE::RenderViewportFeedback
@@ -492,6 +493,12 @@ public:
 	/** Callback for checking the geometry inspection show wireframe flag. */
 	UNREALED_API bool IsSetGeometryInspectionShowWireframe() const;
 
+	/** Callback for setting a specific channel mask flag */
+	UNREALED_API void SetChannelMask(EColorChannelMask InColorChannelMask);
+
+	/** Callback for checking a specific channel mask flag */
+	UNREALED_API bool IsChannelBeingMasked(EColorChannelMask InColorChannelMask) const;
+
 	/** Gets ViewportCameraTransform object for the current viewport type */
 	FViewportCameraTransform& GetViewTransform()
 	{
@@ -2010,6 +2017,21 @@ public:
 	// Override the LOD of landscape in this viewport
 	int8 LandscapeLODOverride;
 
+	/** Parameters for masking specific color channels for visualization */
+	struct FColorChannelMaskParams
+	{
+		/** The color channel to mask to, whose value will be output as a grayscale color */
+		EColorChannelMask ColorChannelMask;
+
+		/** When true, the alpha channel is inverted when being visualized; does not effect the actual alpha value, just its visualization */
+		bool bInvertAlphaChannelMask;
+
+		/** Indicates that a checkerboard should be alpha blended underneath the viewport to help visualize the view's alpha */
+		bool bDrawAlphaBlendedCheckerboard;
+
+		FColorChannelMaskParams();
+	} ChannelMaskParams;
+
 	/** If true, draw vertices for selected BSP brushes and static meshes if the large vertices show flag is set. */
 	bool bDrawVertices;
 
diff --git a/Engine/Source/Editor/UnrealEd/Public/EditorViewportCommands.h b/Engine/Source/Editor/UnrealEd/Public/EditorViewportCommands.h
--- a/Engine/Source/Editor/UnrealEd/Public/EditorViewportCommands.h
+++ b/Engine/Source/Editor/UnrealEd/Public/EditorViewportCommands.h
@@ -166,6 +166,13 @@ public:
 	TSharedPtr< FUICommandInfo > GeometryInspectionRandomColorMode;
 	TSharedPtr< FUICommandInfo > GeometryInspectionShowWireframe;
 
+	/** Changes the viewport to show only specific color channels */
+	TSharedPtr< FUICommandInfo > AllChannelsMask;
+	TSharedPtr< FUICommandInfo > RedChannelMask;
+	TSharedPtr< FUICommandInfo > GreenChannelMask;
+	TSharedPtr< FUICommandInfo > BlueChannelMask;
+	TSharedPtr< FUICommandInfo > AlphaChannelMask;
+
 	/** Toggles realtime rendering in the viewport */
 	TSharedPtr< FUICommandInfo > ToggleRealTime;
 
diff --git a/Engine/Source/Runtime/Engine/Public/SceneView.h b/Engine/Source/Runtime/Engine/Public/SceneView.h
--- a/Engine/Source/Runtime/Engine/Public/SceneView.h
+++ b/Engine/Source/Runtime/Engine/Public/SceneView.h
@@ -1411,6 +1411,15 @@ enum class EViewWaterIntersection : uint8
 	InsideWater
 };
 
+enum class EColorChannelMask : int8
+{
+	All = INDEX_NONE,
+	Red = 0,
+	Green = 1,
+	Blue = 2,
+	Alpha = 3
+};
+
 // Hacky base class to avoid 8 bytes of padding after the vtable
 class FSceneViewFixLayout
 {
@@ -2408,6 +2417,19 @@ public:
 
 	/** True if scenecolor and depth should be multiview-allocated */
 	bool bRequireMultiView;
+
+	/** Parameters for masking specific color channels for visualization */
+	struct FColorChannelMaskParams
+	{
+		/** The color channel to mask to, whose value will be output as a grayscale color */
+		EColorChannelMask ColorChannelMask = EColorChannelMask::All;
+
+		/** When true, the alpha channel is inverted when being visualized; does not effect the actual alpha value, just its visualization */
+		bool bInvertAlphaChannelMask = false;
+
+		/** Indicates that a checkerboard should be alpha blended underneath the viewport to help visualize the view's alpha */
+		bool bDrawAlphaBlendedCheckerboard = false;
+	} ChannelMaskParams;
 	
 	/** DPI scale to be used for debugging font. */
 	float DebugDPIScale = 1.0f;
diff --git a/Engine/Source/Runtime/Renderer/Private/PostProcess/ChannelMask.cpp b/Engine/Source/Runtime/Renderer/Private/PostProcess/ChannelMask.cpp
new file mode 100644
--- /dev/null
+++ b/Engine/Source/Runtime/Renderer/Private/PostProcess/ChannelMask.cpp
@@ -0,0 +1,75 @@
+﻿// Copyright Epic Games, Inc. All Rights Reserved.
+
+#include "ChannelMask.h"
+
+#include "SceneRendering.h"
+#include "PostProcess/PostProcessMaterialInputs.h"
+
+namespace ChannelMask
+{
+	BEGIN_SHADER_PARAMETER_STRUCT(FChannelMaskParameters, )
+		SHADER_PARAMETER_STRUCT_INCLUDE(FViewShaderParameters, View)
+		SHADER_PARAMETER_RDG_TEXTURE(Texture2D, SceneColorTexture)
+		SHADER_PARAMETER_SAMPLER(SamplerState, SceneColorSampler)
+		SHADER_PARAMETER(FVector4f, ChannelMask)
+		SHADER_PARAMETER(uint32, bMaskChannels)
+		SHADER_PARAMETER(uint32, bInvertAlphaChannelMask)
+		SHADER_PARAMETER(uint32, bDisplayCheckerboard)
+		SHADER_PARAMETER(FScreenTransform, SvPositionToViewPos)
+		RENDER_TARGET_BINDING_SLOTS()
+	END_SHADER_PARAMETER_STRUCT()
+	
+	class FChannelMaskPS : public FGlobalShader
+	{
+	public:
+		DECLARE_SHADER_TYPE(FChannelMaskPS, Global);
+		SHADER_USE_PARAMETER_STRUCT(FChannelMaskPS, FGlobalShader);
+		using FParameters = FChannelMaskParameters;
+	};
+	IMPLEMENT_SHADER_TYPE(, FChannelMaskPS, TEXT("/Engine/Private/PostProcessChannelMask.usf"), TEXT("ChannelMask_MainPS"), SF_Pixel);
+}
+
+FScreenPassTexture ChannelMask::AddChannelMaskPass(FRDGBuilder& GraphBuilder, const FViewInfo& View, const FChannelMaskInputs& Inputs)
+{
+	FScreenPassRenderTarget Output = Inputs.OverrideOutput;
+	if (!Output.IsValid())
+	{
+		Output = FScreenPassRenderTarget::CreateFromInput(GraphBuilder, Inputs.SceneColor, ERenderTargetLoadAction::ELoad, TEXT("ChannelMask"));
+	}
+
+	const FScreenPassTextureViewport OutputViewport(Output);
+	const FScreenPassTextureViewport ColorViewport(Inputs.SceneColor);
+
+	const bool bMaskChannels = View.Family->ChannelMaskParams.ColorChannelMask != EColorChannelMask::All;
+	FVector4f ChannelMask = FVector4f(0.0f, 0.0f, 0.0f, 0.0f);
+	if (bMaskChannels)
+	{
+		ChannelMask[(int32)View.Family->ChannelMaskParams.ColorChannelMask] = 1.0f;
+	}
+	
+	FChannelMaskParameters* PassParameters = GraphBuilder.AllocParameters<FChannelMaskParameters>();
+	PassParameters->View = View.GetShaderParameters();
+	PassParameters->SceneColorTexture = Inputs.SceneColor.Texture;
+	PassParameters->SceneColorSampler = TStaticSamplerState<SF_Bilinear, AM_Border, AM_Border, AM_Border>::GetRHI();
+	PassParameters->ChannelMask = ChannelMask;
+	PassParameters->bMaskChannels = bMaskChannels;
+	PassParameters->bInvertAlphaChannelMask = View.Family->ChannelMaskParams.bInvertAlphaChannelMask;
+	PassParameters->bDisplayCheckerboard = View.Family->ChannelMaskParams.bDrawAlphaBlendedCheckerboard;
+	PassParameters->SvPositionToViewPos = (
+		FScreenTransform::ChangeTextureBasisFromTo(FScreenPassTextureViewport(Output), FScreenTransform::ETextureBasis::TexelPosition, FScreenTransform::ETextureBasis::ViewportUV) *
+		FScreenTransform::ChangeTextureBasisFromTo(View.ViewRect.Size(), View.ViewRect, FScreenTransform::ETextureBasis::ViewportUV, FScreenTransform::ETextureBasis::TexelPosition));
+	PassParameters->RenderTargets[0] = Output.GetRenderTargetBinding();
+	
+	TShaderMapRef<FChannelMaskPS> PixelShader(View.ShaderMap);
+
+	AddDrawScreenPass(
+		GraphBuilder,
+		RDG_EVENT_NAME("ChannelMask"),
+		View,
+		OutputViewport,
+		ColorViewport,
+		PixelShader,
+		PassParameters);
+	
+	return MoveTemp(Output);
+}
diff --git a/Engine/Source/Runtime/Renderer/Private/PostProcess/ChannelMask.h b/Engine/Source/Runtime/Renderer/Private/PostProcess/ChannelMask.h
new file mode 100644
--- /dev/null
+++ b/Engine/Source/Runtime/Renderer/Private/PostProcess/ChannelMask.h
@@ -0,0 +1,19 @@
+﻿// Copyright Epic Games, Inc. All Rights Reserved.
+
+#pragma once
+
+#include "ScreenPass.h"
+
+namespace ChannelMask
+{
+	struct FChannelMaskInputs
+	{
+		// Input scene color
+		FScreenPassTexture SceneColor;
+
+		// Optional output render target, if not provided a new render target is created and returned
+		FScreenPassRenderTarget OverrideOutput;
+	};
+	
+	FScreenPassTexture AddChannelMaskPass(class FRDGBuilder& GraphBuilder, const class FViewInfo& View, const struct FChannelMaskInputs& Inputs);
+}
diff --git a/Engine/Source/Runtime/Renderer/Private/PostProcess/PostProcessing.cpp b/Engine/Source/Runtime/Renderer/Private/PostProcess/PostProcessing.cpp
--- a/Engine/Source/Runtime/Renderer/Private/PostProcess/PostProcessing.cpp
+++ b/Engine/Source/Runtime/Renderer/Private/PostProcess/PostProcessing.cpp
@@ -75,6 +75,7 @@
 #include "IlluminanceMeter.h"
 #include "SparseVolumeTexture/SparseVolumeTextureStreamingVisualize.h"
 #include "CanvasItem.h"
+#include "ChannelMask.h"
 #include "MobileSSR.h"
 #include "Materials/MaterialRenderProxy.h"
 #include "CustomRenderPassSceneCapture.h"
@@ -460,6 +461,7 @@ void AddPostProcessingPasses(
 		VisualizeLightFunctionAtlas,
 		VisualizeLevelInstance,
 		VisualizeVirtualShadowMaps_PreEditorPrimitives,
+		ChannelMask,
 		SelectionOutline,
 		EditorPrimitive,
 		DebugPrimitive,
@@ -521,6 +523,7 @@ void AddPostProcessingPasses(
 		TEXT("VisualizeLightFunctionAtlas"),
 		TEXT("VisualizeLevelInstance"),
 		TEXT("VisualizeVirtualShadowMaps_PreEditorPrimitives"),
+		TEXT("ChannelMask"),
 		TEXT("SelectionOutline"),
 		TEXT("EditorPrimitive"),
 		TEXT("DebugPrimitive"),
@@ -565,6 +568,7 @@ void AddPostProcessingPasses(
 	PassSequence.SetEnabled(EPass::VisualizeSkyLightIlluminanceMeter, Scene&& Scene->SkyLight && View.Family && View.Family->EngineShowFlags.VisualizeSkyLightIlluminance);
 	PassSequence.SetEnabled(EPass::VisualizeLightFunctionAtlas, Scene && Scene->LightFunctionAtlasSceneData.GetLightFunctionAtlasEnabled() && View.Family && View.Family->EngineShowFlags.VisualizeLightFunctionAtlas);
 	PassSequence.SetEnabled(EPass::VisualizeLevelInstance, GIsEditor && EngineShowFlags.EditingLevelInstance && EngineShowFlags.VisualizeLevelInstanceEditing && !bVisualizeHDR);
+	PassSequence.SetEnabled(EPass::ChannelMask, View.Family->ChannelMaskParams.ColorChannelMask != EColorChannelMask::All || View.Family->ChannelMaskParams.bDrawAlphaBlendedCheckerboard);
 	PassSequence.SetEnabled(EPass::SelectionOutline, GIsEditor && EngineShowFlags.Selection && EngineShowFlags.SelectionOutline && !EngineShowFlags.Wireframe && !bVisualizeHDR);
 	PassSequence.SetEnabled(EPass::EditorPrimitive, FSceneRenderer::ShouldCompositeEditorPrimitives(View));
 #else
@@ -572,6 +576,7 @@ void AddPostProcessingPasses(
 	PassSequence.SetEnabled(EPass::VisualizeSkyLightIlluminanceMeter, false);
 	PassSequence.SetEnabled(EPass::VisualizeLightFunctionAtlas, false);
 	PassSequence.SetEnabled(EPass::VisualizeLevelInstance, false);
+	PassSequence.SetEnabled(EPass::ChannelMask, false);
 	PassSequence.SetEnabled(EPass::SelectionOutline, false);
 	PassSequence.SetEnabled(EPass::EditorPrimitive, false);
 #endif
@@ -1728,6 +1733,14 @@ void AddPostProcessingPasses(
 #endif
 
 #if WITH_EDITOR
+	if (PassSequence.IsEnabled(EPass::ChannelMask))
+	{
+		ChannelMask::FChannelMaskInputs PassInputs;
+		PassSequence.AcceptOverrideIfLastPass(EPass::ChannelMask, PassInputs.OverrideOutput);
+		PassInputs.SceneColor = SceneColor;
+		SceneColor = ChannelMask::AddChannelMaskPass(GraphBuilder, View, PassInputs);
+	}
+
 	if (PassSequence.IsEnabled(EPass::SelectionOutline))
 	{
 		FSelectionOutlineInputs PassInputs;
@@ -2496,6 +2509,7 @@ void AddMobilePostProcessingPasses(FRDGBuilder& GraphBuilder, FScene* Scene, con
 		PostProcessMaterialAfterTonemapping,
 		HighResolutionScreenshotMask,
 		VisualizeSkyLightIlluminanceMeter,
+		ChannelMask,
 		SelectionOutline,
 		EditorPrimitive,
 		DebugPrimitive,
@@ -2541,6 +2555,7 @@ void AddMobilePostProcessingPasses(FRDGBuilder& GraphBuilder, FScene* Scene, con
 		TEXT("FXAA"),
 		TEXT("HighResolutionScreenshotMask"),
 		TEXT("VisualizeSkyLightIlluminanceMeter"),
+		TEXT("ChannelMask"),
 		TEXT("SelectionOutline"),
 		TEXT("EditorPrimitive"),
 		TEXT("DebugPrimitive"),
@@ -2595,10 +2610,12 @@ void AddMobilePostProcessingPasses(FRDGBuilder& GraphBuilder, FScene* Scene, con
 	PassSequence.SetEnabled(EPass::HighResolutionScreenshotMask, bUseHighResolutionScreenshotMask);
 #if WITH_EDITOR
 	PassSequence.SetEnabled(EPass::VisualizeSkyLightIlluminanceMeter, Scene && Scene->SkyLight && View.Family->EngineShowFlags.VisualizeSkyLightIlluminance);
+	PassSequence.SetEnabled(EPass::ChannelMask, View.Family->ChannelMaskParams.ColorChannelMask != EColorChannelMask::All || View.Family->ChannelMaskParams.bDrawAlphaBlendedCheckerboard);
 	PassSequence.SetEnabled(EPass::SelectionOutline, GIsEditor && View.Family->EngineShowFlags.Selection && View.Family->EngineShowFlags.SelectionOutline && !View.Family->EngineShowFlags.Wireframe);
 	PassSequence.SetEnabled(EPass::EditorPrimitive, FSceneRenderer::ShouldCompositeEditorPrimitives(View));
 #else
 	PassSequence.SetEnabled(EPass::VisualizeSkyLightIlluminanceMeter, false);
+	PassSequence.SetEnabled(EPass::ChannelMask, false);
 	PassSequence.SetEnabled(EPass::SelectionOutline, false);
 	PassSequence.SetEnabled(EPass::EditorPrimitive, false);
 #endif
@@ -3309,6 +3326,14 @@ void AddMobilePostProcessingPasses(FRDGBuilder& GraphBuilder, FScene* Scene, con
 		SceneColor = ProcessAndRenderIlluminanceMeter(GraphBuilder, View, SceneColor);
 	}
 
+	if (PassSequence.IsEnabled(EPass::ChannelMask))
+	{
+		ChannelMask::FChannelMaskInputs PassInputs;
+		PassSequence.AcceptOverrideIfLastPass(EPass::ChannelMask, PassInputs.OverrideOutput);
+		PassInputs.SceneColor = SceneColor;
+		SceneColor = ChannelMask::AddChannelMaskPass(GraphBuilder, View, PassInputs);
+	}
+	
 	// Show the selection outline if it is in the editor and we aren't in wireframe 
 	// If the engine is in demo mode and game view is on we also do not show the selection outline
 	if (PassSequence.IsEnabled(EPass::SelectionOutline))
