Editor: Make SupportsTreeViewOverlay in ISCSEditorUICustomization take the whole subobject editor as input rather than UObject targets, so we don't rely on targets being UObjects. Expose SSubobjectEditor::GetDataContextHandle() so that customizations can use it to find the targets and allow them to handle cases of drag-drop operations when the subobject editor targets multiple items at the same time.

diff --git a/Engine/Source/Editor/Kismet/Public/ISCSEditorUICustomization.h b/Engine/Source/Editor/Kismet/Public/ISCSEditorUICustomization.h
--- a/Engine/Source/Editor/Kismet/Public/ISCSEditorUICustomization.h
+++ b/Engine/Source/Editor/Kismet/Public/ISCSEditorUICustomization.h
@@ -96,8 +96,10 @@ public:
 	virtual bool SelectDefaultSubobject(SSubobjectEditor& SubobjectEditor) { return false; }
 
 	/** Decides if this UI Customization provides an overlay for the tree view */
-	virtual bool SupportsTreeViewOverlay(TArrayView<UObject*> Context) const { return false; }
+	UE_EXPERIMENTAL(5.8, "This function is experimental for now.")
+	virtual bool SupportsTreeViewOverlay(const SSubobjectEditor& Editor) const { return false; }
 
 	/** Create an overlay for the tree view widget */
+	UE_EXPERIMENTAL(5.8, "This function is experimental for now.")
 	virtual TSharedRef<class SWidget> ConstructTreeViewOverlay(TWeakPtr<class SSubobjectEditor> Editor) const { return SNullWidget::NullWidget; }
 };
diff --git a/Engine/Source/Editor/LevelEditor/Private/ActorDetailsSCSEditorUICustomization.cpp b/Engine/Source/Editor/LevelEditor/Private/ActorDetailsSCSEditorUICustomization.cpp
--- a/Engine/Source/Editor/LevelEditor/Private/ActorDetailsSCSEditorUICustomization.cpp
+++ b/Engine/Source/Editor/LevelEditor/Private/ActorDetailsSCSEditorUICustomization.cpp
@@ -167,11 +167,11 @@ void FActorDetailsSCSEditorUICustomization::RemoveCustomization(TSharedPtr<ISCSE
 	Customizations.Remove(Customization);
 }
 
-bool FActorDetailsSCSEditorUICustomization::SupportsTreeViewOverlay(TArrayView<UObject*> Context) const
+bool FActorDetailsSCSEditorUICustomization::SupportsTreeViewOverlay(const SSubobjectEditor& Editor) const
 {
 	for (const TSharedPtr<ISCSEditorUICustomization>& Customization : Customizations)
 	{
-		if (Customization->SupportsTreeViewOverlay(Context))
+		if (Customization->SupportsTreeViewOverlay(Editor))
 		{
 			return true;
 		}
diff --git a/Engine/Source/Editor/LevelEditor/Private/ActorDetailsSCSEditorUICustomization.h b/Engine/Source/Editor/LevelEditor/Private/ActorDetailsSCSEditorUICustomization.h
--- a/Engine/Source/Editor/LevelEditor/Private/ActorDetailsSCSEditorUICustomization.h
+++ b/Engine/Source/Editor/LevelEditor/Private/ActorDetailsSCSEditorUICustomization.h
@@ -37,7 +37,7 @@ public:
 	virtual void PrePotentialTargetChange(SSubobjectEditor& SubobjectEditor) override;
 	virtual bool SelectDefaultSubobject(SSubobjectEditor& SubobjectEditor) override;
 
-	virtual bool SupportsTreeViewOverlay(TArrayView<UObject*> Context) const override;
+	virtual bool SupportsTreeViewOverlay(const SSubobjectEditor& Editor) const override;
 	virtual TSharedRef<SWidget> ConstructTreeViewOverlay(TWeakPtr<SSubobjectEditor> Editor) const override;
 
 	void AddCustomization(TSharedPtr<ISCSEditorUICustomization> Customization);
diff --git a/Engine/Source/Editor/LevelEditor/Private/SActorDetails.cpp b/Engine/Source/Editor/LevelEditor/Private/SActorDetails.cpp
--- a/Engine/Source/Editor/LevelEditor/Private/SActorDetails.cpp
+++ b/Engine/Source/Editor/LevelEditor/Private/SActorDetails.cpp
@@ -242,7 +242,7 @@ public:
 	virtual TSubclassOf<UActorComponent> GetComponentTypeFilter(TArrayView<UObject*> Context) const override;
 	virtual bool FilterSubobjectData(TArray<FSubobjectDataHandle>& SubobjectData) const override;
 	virtual bool SortSubobjectData(TArray<FSubobjectDataHandle>& SubobjectData) override;
-	virtual bool SupportsTreeViewOverlay(TArrayView<UObject*> Context) const override;
+	virtual bool SupportsTreeViewOverlay(const SSubobjectEditor& Editor) const override;
 	virtual TSharedRef<class SWidget> ConstructTreeViewOverlay(TWeakPtr<class SSubobjectEditor> Editor) const override;
 	virtual void PrePotentialTargetChange(SSubobjectEditor& SubobjectEditor) override;
 	virtual bool SelectDefaultSubobject(SSubobjectEditor& SubobjectEditor) override;
@@ -364,11 +364,11 @@ bool SActorDetails::FInternalSubobjectCustomization::SortSubobjectData(TArray<FS
 	return false;
 }
 
-bool SActorDetails::FInternalSubobjectCustomization::SupportsTreeViewOverlay(TArrayView<UObject*> Context) const
+bool SActorDetails::FInternalSubobjectCustomization::SupportsTreeViewOverlay(const SSubobjectEditor& Editor) const
 {
 	if (ParentCustomization)
 	{
-		return ParentCustomization->SupportsTreeViewOverlay(Context);
+		return ParentCustomization->SupportsTreeViewOverlay(Editor);
 	}
 	return false;
 }
diff --git a/Engine/Source/Editor/SubobjectEditor/Private/SSubobjectBlueprintEditor.cpp b/Engine/Source/Editor/SubobjectEditor/Private/SSubobjectBlueprintEditor.cpp
--- a/Engine/Source/Editor/SubobjectEditor/Private/SSubobjectBlueprintEditor.cpp
+++ b/Engine/Source/Editor/SubobjectEditor/Private/SSubobjectBlueprintEditor.cpp
@@ -162,8 +162,7 @@ void SSubobjectBlueprintEditor::Construct(const FArguments& InArgs)
 			.Visibility_Lambda([this]()
 			{
 				const bool bTreeIsVisible = GetComponentsTreeVisibility() == EVisibility::Visible;
-				UObject* ObjectContextPtr = GetObjectContext();
-				const bool bSupportsOverlay = UICustomization && UICustomization->SupportsTreeViewOverlay(MakeArrayView(&ObjectContextPtr, 1));
+				const bool bSupportsOverlay = UICustomization && UICustomization->SupportsTreeViewOverlay(*this);
 
 				if (bTreeIsVisible && bSupportsOverlay)
 				{
@@ -213,8 +212,7 @@ void SSubobjectBlueprintEditor::UpdateTree(bool bRegenerateTreeNodes)
 	SSubobjectEditor::UpdateTree(bRegenerateTreeNodes);
 	if (TreeViewOverlay.IsValid())
 	{
-		UObject* ObjectContextPtr = GetObjectContext();
-		if (UICustomization && UICustomization->SupportsTreeViewOverlay(MakeArrayView(&ObjectContextPtr, 1)))
+		if (UICustomization && UICustomization->SupportsTreeViewOverlay(*this))
 		{
 			TreeViewOverlay->SetContent(UICustomization->ConstructTreeViewOverlay(SharedThis(this)));
 		}
diff --git a/Engine/Source/Editor/SubobjectEditor/Private/SSubobjectInstanceEditor.cpp b/Engine/Source/Editor/SubobjectEditor/Private/SSubobjectInstanceEditor.cpp
--- a/Engine/Source/Editor/SubobjectEditor/Private/SSubobjectInstanceEditor.cpp
+++ b/Engine/Source/Editor/SubobjectEditor/Private/SSubobjectInstanceEditor.cpp
@@ -179,8 +179,7 @@ void SSubobjectInstanceEditor::Construct(const FArguments& InArgs)
 			.Visibility_Lambda([this]()
 			{
 				const bool bTreeIsVisible = GetComponentsTreeVisibility() == EVisibility::Visible;
-				UObject* ObjectContextPtr = GetObjectContext();
-				const bool bSupportsOverlay = UICustomization && UICustomization->SupportsTreeViewOverlay(MakeArrayView(&ObjectContextPtr, 1));
+				const bool bSupportsOverlay = UICustomization && UICustomization->SupportsTreeViewOverlay(*this);
 				
 				if (bTreeIsVisible && bSupportsOverlay)
 				{
@@ -234,8 +233,7 @@ void SSubobjectInstanceEditor::UpdateTree(bool bRegenerateTreeNodes)
 
 	if (TreeViewOverlay.IsValid())
 	{
-		UObject* ObjectContextPtr = GetObjectContext();
-		if (UICustomization && UICustomization->SupportsTreeViewOverlay(MakeArrayView(&ObjectContextPtr, 1)))
+		if (UICustomization && UICustomization->SupportsTreeViewOverlay(*this))
 		{
 			TreeViewOverlay->SetContent(UICustomization->ConstructTreeViewOverlay(SharedThis(this)));
 		}
diff --git a/Engine/Source/Editor/SubobjectEditor/Public/SSubobjectEditor.h b/Engine/Source/Editor/SubobjectEditor/Public/SSubobjectEditor.h
--- a/Engine/Source/Editor/SubobjectEditor/Public/SSubobjectEditor.h
+++ b/Engine/Source/Editor/SubobjectEditor/Public/SSubobjectEditor.h
@@ -601,7 +601,8 @@ public:
 	 * access to non-UObject based data and for cases where SubobjectData is a proxy for multiple objects
 	 * in bulk-edit cases.
 	 */
-	FSubobjectData* GetDataContextHandle() const;
+	UE_EXPERIMENTAL(5.8, "This function is experimental for now.")
+	UE_API FSubobjectData* GetDataContextHandle() const;
 
 	/** Refresh the type list presented by the add component button when clicked */
 	UE_API void RefreshComponentTypesList();
