[MegaLights] Remove the check for r.MegaLights.EnableForProject to show or hide MegaLights from the view mode dropdown menu. When r.MegaLights.EnableForProject is 0, MegaLights can still be enabled selectively via PPVs so it doesn't make sense to hide the UI

#tests tested in editor

diff --git a/Engine/Source/Editor/LevelEditor/Private/ViewportToolbar/LevelEditorSubmenus.cpp b/Engine/Source/Editor/LevelEditor/Private/ViewportToolbar/LevelEditorSubmenus.cpp
--- a/Engine/Source/Editor/LevelEditor/Private/ViewportToolbar/LevelEditorSubmenus.cpp
+++ b/Engine/Source/Editor/LevelEditor/Private/ViewportToolbar/LevelEditorSubmenus.cpp
@@ -1170,11 +1170,7 @@ void PopulateViewModesMenu(UToolMenu* InMenu)
 		);
 	}
 
-	static const auto CVarMegaLightsEnableForProject = IConsoleManager::Get().FindTConsoleVariableDataInt(TEXT("r.MegaLights.EnableForProject"));
-
-	if (CVarMegaLightsEnableForProject
-		&& CVarMegaLightsEnableForProject->GetValueOnGameThread() != 0
-		&& ShouldCompileMegaLightsShaders(GMaxRHIShaderPlatform))
+	if (ShouldCompileMegaLightsShaders(GMaxRHIShaderPlatform))
 	{
 		FToolMenuSection& Section = InMenu->FindOrAddSection("ViewMode");
 		Section.AddSubMenu(
