[LevelViewport][EditorTests] Editor tests now query the level viewport via tag instead of type.

#Synced-CL 48271216

diff --git a/Engine/Source/Editor/LevelEditor/Private/LevelViewportLayoutEntity.cpp b/Engine/Source/Editor/LevelEditor/Private/LevelViewportLayoutEntity.cpp
--- a/Engine/Source/Editor/LevelEditor/Private/LevelViewportLayoutEntity.cpp
+++ b/Engine/Source/Editor/LevelEditor/Private/LevelViewportLayoutEntity.cpp
@@ -10,6 +10,7 @@
 FLevelViewportLayoutEntity::FLevelViewportLayoutEntity(TSharedPtr<SAssetEditorViewport> InLevelViewport)
 	: LevelViewport(StaticCastSharedRef<SLevelViewport>(InLevelViewport.ToSharedRef()))
 {
+	InLevelViewport->SetTag("LevelViewport");
 }
 
 TSharedRef<SWidget> FLevelViewportLayoutEntity::AsWidget() const
