EditorUsability : HighResScreenshotUI
- Fixed HighResScreenshot OpenFolder action does not navigate to folder if it does not exists

#Synced-CL 48903505

diff --git a/Engine/Source/Editor/LevelEditor/Private/HighResScreenshotUI.cpp b/Engine/Source/Editor/LevelEditor/Private/HighResScreenshotUI.cpp
--- a/Engine/Source/Editor/LevelEditor/Private/HighResScreenshotUI.cpp
+++ b/Engine/Source/Editor/LevelEditor/Private/HighResScreenshotUI.cpp
@@ -5,6 +5,7 @@
 #include "Engine/World.h"
 #include "Framework/Application/SlateApplication.h"
 #include "Framework/Docking/TabManager.h"
+#include "HAL/PlatformFileManager.h"
 #include "Styling/AppStyle.h"
 #include "Styling/StyleColors.h"
 #include "SWarningOrErrorBox.h"
@@ -434,7 +435,18 @@ EVisibility SHighResScreenshotDialog::GetCaptureWarningVisibility() const
 
 FReply SHighResScreenshotDialog::OnOpenScreenshotFolder()
 {
-	FPlatformProcess::ExploreFolder(*FPaths::ConvertRelativePathToFull(FPaths::ScreenShotDir()));
+	const FString PathToScreenshotFolder = FPaths::ConvertRelativePathToFull(FPaths::ScreenShotDir());
+
+	// Before navigating to folder, check if it exists and create it if not
+	IPlatformFile& PlatformFile = FPlatformFileManager::Get().GetPlatformFile();
+	if (!PlatformFile.DirectoryExists(*PathToScreenshotFolder))
+	{
+		PlatformFile.CreateDirectoryTree(*PathToScreenshotFolder);
+	}
+
+	// Navigate to folder in explorer
+	FPlatformProcess::ExploreFolder(*PathToScreenshotFolder);
+
 	return FReply::Handled();
 }
 
