Fix crashes that can occur with object reinstancing
* SubobjectEditor::UpdateTree is called from the OnObjectReplaced callback attached in SActorDetails
* OnObjectReplaced is a callback made in the middle of reinstancing.  This means that handlers cannot safely follow object references at this time since there is no guarantee the objects have been replaced.  This poses a problem since SubobjectEditor::UpdateTree may need to walk the component list of an object (thus needing to naviagate through object references).  This function is also called once for every object that is replaced due to the invocation of NotifyToolsOfObjectReplacement in CopyPropertiesForUnrelatedObjects. (See loop in FBlueprintCompileReinstancer::ReplaceObjectHelper).
* In circumstances where the component type has changed but the object references haven't yet been replaced, this can cause the subobject editor to traverse to a dead object.
* This change causes the SubobjectEditor to only be updated _after_ object references have been replaced. An added benefit is this avoids duplicating a large amount of work on compilation where more than one object is replaced. (verse is a common case for this)

diff --git a/Engine/Source/Editor/LevelEditor/Private/SActorDetails.cpp b/Engine/Source/Editor/LevelEditor/Private/SActorDetails.cpp
--- a/Engine/Source/Editor/LevelEditor/Private/SActorDetails.cpp
+++ b/Engine/Source/Editor/LevelEditor/Private/SActorDetails.cpp
@@ -437,7 +437,7 @@ void SActorDetails::Construct(
 	SelectionSet = InSelectionSet;
 	checkf(SelectionSet, TEXT("SActorDetails must be constructed with a valid selection set!"));
 
-	FCoreUObjectDelegates::OnObjectsReplaced.AddRaw(this, &SActorDetails::OnObjectsReplaced);
+	FCoreUObjectDelegates::OnObjectsReinstanced.AddRaw(this, &SActorDetails::OnObjectsReinstanced);
 	FCoreUObjectDelegates::OnObjectPropertyChanged.AddRaw(this, &SActorDetails::OnObjectPropertyChanged);
 
 	FLevelEditorModule& LevelEditor = FModuleManager::GetModuleChecked<FLevelEditorModule>("LevelEditor");
@@ -639,7 +639,7 @@ SActorDetails::~SActorDetails()
 		GEditor->UnregisterForUndo(this);
 	}
 
-	FCoreUObjectDelegates::OnObjectsReplaced.RemoveAll(this);
+	FCoreUObjectDelegates::OnObjectsReinstanced.RemoveAll(this);
 	FCoreUObjectDelegates::OnObjectPropertyChanged.RemoveAll(this);
 
 	RemoveBPComponentCompileEventDelegate();
@@ -1310,7 +1310,7 @@ void SActorDetails::OnBlueprintComponentCompiled(UBlueprint* ComponentBlueprint)
 	UpdateComponentTreeFromEditorSelection();
 }
 
-void SActorDetails::OnObjectsReplaced(const TMap<UObject*, UObject*>& InReplacementObjects)
+void SActorDetails::OnObjectsReinstanced(const TMap<UObject*, UObject*>& InReplacementObjects)
 {
 	if (bHasSelectionOverride && SelectionOverrideActors.Num() > 0)
 	{
diff --git a/Engine/Source/Editor/LevelEditor/Private/SActorDetails.h b/Engine/Source/Editor/LevelEditor/Private/SActorDetails.h
--- a/Engine/Source/Editor/LevelEditor/Private/SActorDetails.h
+++ b/Engine/Source/Editor/LevelEditor/Private/SActorDetails.h
@@ -117,7 +117,7 @@ private:
 	void AddBPComponentCompileEventDelegate(UBlueprint* ComponentBlueprint);
 	void RemoveBPComponentCompileEventDelegate();
 	void OnBlueprintComponentCompiled(UBlueprint* ComponentBlueprint);
-	void OnObjectsReplaced(const TMap<UObject*, UObject*>& InReplacementObjects);
+	void OnObjectsReinstanced(const TMap<UObject*, UObject*>& InReinstancedObjects);
 	void OnObjectPropertyChanged(UObject* ObjectBeingModified, FPropertyChangedEvent& PropertyChangedEvent);
 
 private:
