Adds a user option to set the default value of the copy options property when using the replace actors context menu

#Synced-CL 48498703

diff --git a/Engine/Config/BaseEditorPerProjectUserSettings.ini b/Engine/Config/BaseEditorPerProjectUserSettings.ini
--- a/Engine/Config/BaseEditorPerProjectUserSettings.ini
+++ b/Engine/Config/BaseEditorPerProjectUserSettings.ini
@@ -208,6 +208,8 @@ bBSPAutoUpdate=True
 bAutoMoveBSPPivotOffset=False
 ; If true, Navigation will auto-update
 bNavigationAutoUpdate=True
+; The default value of the copy properties option when using the replace actors context menu
+bReplaceCopyPropertiesDefault=False
 ; If true, replacing actors allows copying properties when related by the native parent of the replaced class
 bReplaceCopiesPropertiesWhenRelatedByNativeParent=True
 ; If true, Replaces respects the scale of the original actor. If false, Replace sets the scale to 1.0
diff --git a/Engine/Source/Editor/LevelEditor/Private/LevelEditorCreateActorMenu.cpp b/Engine/Source/Editor/LevelEditor/Private/LevelEditorCreateActorMenu.cpp
--- a/Engine/Source/Editor/LevelEditor/Private/LevelEditorCreateActorMenu.cpp
+++ b/Engine/Source/Editor/LevelEditor/Private/LevelEditorCreateActorMenu.cpp
@@ -47,6 +47,7 @@
 #include "ThumbnailRendering/ThumbnailManager.h"
 #include "Engine/Selection.h"
 #include "ActorEditorUtils.h"
+#include "Settings/LevelEditorMiscSettings.h"
 #include "Widgets/Layout/SBorder.h"
 
 class SMenuThumbnail : public SCompoundWidget
@@ -448,7 +449,7 @@ void LevelEditorCreateActorMenu::FillAddReplaceActorMenu(UToolMenu* Menu, EActor
 	{
 		FToolMenuSection& Section = Menu->AddSection("Options", NSLOCTEXT("LevelViewportContextMenu", "Options", "Options"));
 
-		GReplaceSelectedActorsWithSelectedClassCopyProperties = false;
+		GReplaceSelectedActorsWithSelectedClassCopyProperties = GetDefault<ULevelEditorMiscSettings>()->bReplaceCopyPropertiesDefault;
 
 		FToolMenuEntry ToolMenuEntry = FToolMenuEntry::InitMenuEntry(
 			"CopyProperties",
diff --git a/Engine/Source/Editor/UnrealEd/Classes/Settings/LevelEditorMiscSettings.h b/Engine/Source/Editor/UnrealEd/Classes/Settings/LevelEditorMiscSettings.h
--- a/Engine/Source/Editor/UnrealEd/Classes/Settings/LevelEditorMiscSettings.h
+++ b/Engine/Source/Editor/UnrealEd/Classes/Settings/LevelEditorMiscSettings.h
@@ -39,6 +39,10 @@ public:
 	UPROPERTY(EditAnywhere, config, Category=Editing, meta=( DisplayName = "Update Navigation Automatically" ))
 	uint32 bNavigationAutoUpdate:1;
 
+	/** If enabled, the copy properties option will be enabled by default when using the replace actors context menu. */
+	UPROPERTY(EditAnywhere, config, Category=Editing, meta=( DisplayName = "Copy Properties on Replace" ))
+	uint32 bReplaceCopyPropertiesDefault:1;
+	
 	/** If enabled, replacing actors allows copying properties when related by the native parent of the replaced class. */
 	UPROPERTY(EditAnywhere, config, Category=Editing, meta=( DisplayName = "Copy Properties when Replacing with Different Blueprint" ))
 	uint32 bReplaceCopiesPropertiesWhenRelatedByNativeParent:1;
