Fix "Mirror [X/Y/Z] Axis" strings not being gathered

Move LOCTEXT out of the AXIS_UI_COMMAND macro scope as it prevents normal strings gathering.

#tests Run Localize command locally with Editor as a localization target, check Editor.po files, observe "MirrorAxis*" strings appeared

#Synced-CL 49076302

diff --git a/Engine/Source/Editor/LevelEditor/Private/LevelEditorActions.cpp b/Engine/Source/Editor/LevelEditor/Private/LevelEditorActions.cpp
--- a/Engine/Source/Editor/LevelEditor/Private/LevelEditorActions.cpp
+++ b/Engine/Source/Editor/LevelEditor/Private/LevelEditorActions.cpp
@@ -4328,18 +4328,23 @@ void FLevelEditorCommands::RegisterCommands()
 	UI_COMMAND( AlignBottomCenterBoundsToActor, "Align Bottom Center Bounds to Actor", "Aligns the actor with another actor at its bottom center bounds", EUserInterfaceActionType::Button, FInputChord() );
 	UI_COMMAND( DeltaTransformToActors, "Delta Transform", "Apply Delta Transform to selected elements", EUserInterfaceActionType::Button, FInputChord() );
 
-	#define AXIS_UI_COMMAND(CommandName, Axis) \
-	FUICommandInfo::MakeCommandInfo( \
-		AsShared(), \
-		CommandName, \
-		TEXT(#CommandName), \
-		FText::Format(LOCTEXT("MirrorAxisLabel", "Mirror {0} Axis"), AxisDisplayInfo::GetAxisDisplayName(Axis)), \
-		FText::Format(LOCTEXT("MirrorAxisTooltip", "Mirrors the element along the {0} axis"), AxisDisplayInfo::GetAxisDisplayName(Axis)), \
-		FSlateIcon(), \
-		EUserInterfaceActionType::Button, \
-		FInputChord() \
-		); \
+	auto MakeAxisUICommand = [this](TSharedPtr<FUICommandInfo>& Command, const TCHAR* CommandName, EAxisList::Type Axis)
+	{
+		const FText AxisDisplayName = AxisDisplayInfo::GetAxisDisplayName(Axis);
+		FUICommandInfo::MakeCommandInfo(
+			AsShared(),
+			Command,
+			CommandName,
+			FText::Format(LOCTEXT("MirrorAxisLabel", "Mirror {0} Axis"), AxisDisplayName),
+			FText::Format(LOCTEXT("MirrorAxisTooltip", "Mirrors the element along the {0} axis"), AxisDisplayName),
+			FSlateIcon(),
+			EUserInterfaceActionType::Button,
+			FInputChord()
+		);
+	};
 
+	#define AXIS_UI_COMMAND(Command, Axis) MakeAxisUICommand(Command, TEXT(#Command), Axis);
+	
 	AXIS_UI_COMMAND(MirrorActorY, EAxisList::Left)
 	AXIS_UI_COMMAND(MirrorActorZ, EAxisList::Up)
 	AXIS_UI_COMMAND(MirrorActorX, EAxisList::Forward)
