Fix crash when loading using world travel in PIE.
Added null check in SLevelViewport::RefreshPIEViewport.

#Synced-CL 49204220

diff --git a/Engine/Source/Editor/LevelEditor/Private/SLevelViewport.cpp b/Engine/Source/Editor/LevelEditor/Private/SLevelViewport.cpp
--- a/Engine/Source/Editor/LevelEditor/Private/SLevelViewport.cpp
+++ b/Engine/Source/Editor/LevelEditor/Private/SLevelViewport.cpp
@@ -2050,13 +2050,16 @@ void SLevelViewport::RefreshPIEViewport()
 {
 	if (UGameViewportClient* GameClient = GetPlayClient())
 	{
-		FSlateThrottleManager::Get().DisableThrottle(true);
-		
-		FTimerHandle TimerHandle;
-		GameClient->GetWorld()->GetTimerManager().SetTimer(TimerHandle, []
+		if (UWorld* World = GameClient->GetWorld())
 		{
-			FSlateThrottleManager::Get().DisableThrottle(false);
-		}, 0.1f, false);
+			FSlateThrottleManager::Get().DisableThrottle(true);
+		
+			FTimerHandle TimerHandle;
+			World->GetTimerManager().SetTimer(TimerHandle, []
+			{
+				FSlateThrottleManager::Get().DisableThrottle(false);
+			}, 0.1f, false);
+		}
 	}
 }
 
