Pass SubobjectEditor to ISCSEditorUICustomization for context to implementer

#Synced-CL 49299980

diff --git a/Engine/Source/Editor/Kismet/Private/BlueprintEditorSCSEditorUICustomization.cpp b/Engine/Source/Editor/Kismet/Private/BlueprintEditorSCSEditorUICustomization.cpp
--- a/Engine/Source/Editor/Kismet/Private/BlueprintEditorSCSEditorUICustomization.cpp
+++ b/Engine/Source/Editor/Kismet/Private/BlueprintEditorSCSEditorUICustomization.cpp
@@ -81,17 +81,17 @@ const FSlateBrush* FFBlueprintEditorSCSEditorUICustomization::GetIconBrush(const
 	return ISCSEditorUICustomization::GetIconBrush(SubobjectData);
 }
 
-TSharedPtr<SWidget> FFBlueprintEditorSCSEditorUICustomization::GetControlsWidget(const FSubobjectData& SubobjectData) const
+TSharedPtr<SWidget> FFBlueprintEditorSCSEditorUICustomization::GetControlsWidget(TSharedRef<SSubobjectEditor>& SubobjectEditor, const FSubobjectData& SubobjectData) const
 {
 	for (const TSharedPtr<ISCSEditorUICustomization>& Customization : Customizations)
 	{
-		if (TSharedPtr<SWidget> ControlsWidget = Customization->GetControlsWidget(SubobjectData))
+		if (TSharedPtr<SWidget> ControlsWidget = Customization->GetControlsWidget(SubobjectEditor, SubobjectData))
 		{
 			return ControlsWidget;
 		}
 	}
 
-	return ISCSEditorUICustomization::GetControlsWidget(SubobjectData);
+	return ISCSEditorUICustomization::GetControlsWidget(SubobjectEditor, SubobjectData);
 }
 
 EChildActorComponentTreeViewVisualizationMode FFBlueprintEditorSCSEditorUICustomization::GetChildActorVisualizationMode() const
diff --git a/Engine/Source/Editor/Kismet/Private/BlueprintEditorSCSEditorUICustomization.h b/Engine/Source/Editor/Kismet/Private/BlueprintEditorSCSEditorUICustomization.h
--- a/Engine/Source/Editor/Kismet/Private/BlueprintEditorSCSEditorUICustomization.h
+++ b/Engine/Source/Editor/Kismet/Private/BlueprintEditorSCSEditorUICustomization.h
@@ -22,7 +22,7 @@ public:
 	virtual bool HideBlueprintButtons(TArrayView<UObject*> Context) const override;
 
 	virtual const FSlateBrush* GetIconBrush(const FSubobjectData&) const override;
-	virtual TSharedPtr<SWidget> GetControlsWidget(const FSubobjectData&) const override;
+	virtual TSharedPtr<SWidget> GetControlsWidget(TSharedRef<SSubobjectEditor>& SubobjectEditor, const FSubobjectData&) const override;
 
 	virtual EChildActorComponentTreeViewVisualizationMode GetChildActorVisualizationMode() const override;
 
diff --git a/Engine/Source/Editor/Kismet/Public/ISCSEditorUICustomization.h b/Engine/Source/Editor/Kismet/Public/ISCSEditorUICustomization.h
--- a/Engine/Source/Editor/Kismet/Public/ISCSEditorUICustomization.h
+++ b/Engine/Source/Editor/Kismet/Public/ISCSEditorUICustomization.h
@@ -51,7 +51,23 @@ public:
 	 * Returning nullptr will cause built-in behaviour to be used, indicating that the UICustomization does not customize this
 	 * UI element. To override existing behaviour and actually show nothing, return a single null widget
 	 */
+	UE_DEPRECATED(5.8, "Use GetControlsWidget(TSharedRef<SSubobjectEditor>&, const FSubobjectData&")
 	virtual TSharedPtr<SWidget> GetControlsWidget(const FSubobjectData&) const { return {}; }
+	
+	/** Allows customizer to inject a right-side justified widget on the subobject editor
+	 * Typically these widgets would convey actions that can be taken and may optionally provide
+	 * status information through their visualization.
+	 * As an example, a subobject that is a non-native ActorComponent may provide a button via these widgets
+	 * to open a code editor to the C++ version of that type.
+	 * Returning nullptr will cause built-in behaviour to be used, indicating that the UICustomization does not customize this
+	 * UI element. To override existing behaviour and actually show nothing, return a single null widget
+	 */
+	virtual TSharedPtr<SWidget> GetControlsWidget(TSharedRef<SSubobjectEditor>& SubobjectEditor, const FSubobjectData& Data) const
+	{
+		PRAGMA_DISABLE_DEPRECATION_WARNINGS
+		return GetControlsWidget(Data);
+		PRAGMA_ENABLE_DEPRECATION_WARNINGS
+	}
 
 	/**
 	 * When it returns true, overrides the value returned by SSubobjectEditor::GetDesiredVisibility().
@@ -102,4 +118,4 @@ public:
 	/** Create an overlay for the tree view widget */
 	UE_EXPERIMENTAL(5.8, "This function is experimental for now.")
 	virtual TSharedRef<class SWidget> ConstructTreeViewOverlay(TWeakPtr<class SSubobjectEditor> Editor) const { return SNullWidget::NullWidget; }
-};
+};
\ No newline at end of file
diff --git a/Engine/Source/Editor/LevelEditor/Private/ActorDetailsSCSEditorUICustomization.cpp b/Engine/Source/Editor/LevelEditor/Private/ActorDetailsSCSEditorUICustomization.cpp
--- a/Engine/Source/Editor/LevelEditor/Private/ActorDetailsSCSEditorUICustomization.cpp
+++ b/Engine/Source/Editor/LevelEditor/Private/ActorDetailsSCSEditorUICustomization.cpp
@@ -76,16 +76,16 @@ const FSlateBrush* FActorDetailsSCSEditorUICustomization::GetIconBrush(const FSu
 	return ISCSEditorUICustomization::GetIconBrush(SubobjectData);
 }
 
-TSharedPtr<SWidget> FActorDetailsSCSEditorUICustomization::GetControlsWidget(const FSubobjectData& SubobjectData) const
+TSharedPtr<SWidget> FActorDetailsSCSEditorUICustomization::GetControlsWidget(TSharedRef<SSubobjectEditor>& SubobjectEditor, const FSubobjectData& SubobjectData) const
 {
 	for (const TSharedPtr<ISCSEditorUICustomization>& Customization : Customizations)
 	{
-		if (TSharedPtr<SWidget> ControlsWidget = Customization->GetControlsWidget(SubobjectData))
+		if (TSharedPtr<SWidget> ControlsWidget = Customization->GetControlsWidget(SubobjectEditor, SubobjectData))
 		{
 			return ControlsWidget;
 		}
 	}
-	return ISCSEditorUICustomization::GetControlsWidget(SubobjectData);
+	return ISCSEditorUICustomization::GetControlsWidget(SubobjectEditor, SubobjectData);
 }
 
 bool FActorDetailsSCSEditorUICustomization::OverrideDesiredVisibility(const UE::Editor::FSubobjectEditorContext& Context, EVisibility& VisibilityOut) const
diff --git a/Engine/Source/Editor/LevelEditor/Private/ActorDetailsSCSEditorUICustomization.h b/Engine/Source/Editor/LevelEditor/Private/ActorDetailsSCSEditorUICustomization.h
--- a/Engine/Source/Editor/LevelEditor/Private/ActorDetailsSCSEditorUICustomization.h
+++ b/Engine/Source/Editor/LevelEditor/Private/ActorDetailsSCSEditorUICustomization.h
@@ -25,7 +25,7 @@ public:
 	virtual bool HideBlueprintButtons(TArrayView<UObject*> Context) const override;
 
 	virtual const FSlateBrush* GetIconBrush(const FSubobjectData&) const override;
-	virtual TSharedPtr<SWidget> GetControlsWidget(const FSubobjectData&) const override;
+	virtual TSharedPtr<SWidget> GetControlsWidget(TSharedRef<SSubobjectEditor>&, const FSubobjectData&) const override;
 	virtual bool OverrideDesiredVisibility(const UE::Editor::FSubobjectEditorContext& Context, EVisibility& VisibilityOut) const;
 
 	virtual EChildActorComponentTreeViewVisualizationMode GetChildActorVisualizationMode() const override;
diff --git a/Engine/Source/Editor/LevelEditor/Private/SActorDetails.cpp b/Engine/Source/Editor/LevelEditor/Private/SActorDetails.cpp
--- a/Engine/Source/Editor/LevelEditor/Private/SActorDetails.cpp
+++ b/Engine/Source/Editor/LevelEditor/Private/SActorDetails.cpp
@@ -236,7 +236,7 @@ public:
 	virtual bool HideComponentsTree(TArrayView<UObject*> Context) const override;
 	virtual bool HideComponentsFilterBox(TArrayView<UObject*> Context) const override;
 	virtual const FSlateBrush* GetIconBrush(const FSubobjectData&) const override;
-	virtual TSharedPtr<SWidget> GetControlsWidget(const FSubobjectData&) const override;
+	virtual TSharedPtr<SWidget> GetControlsWidget(TSharedRef<SSubobjectEditor>& InSubobjectEditor, const FSubobjectData&) const override;
 	virtual bool OverrideDesiredVisibility(const UE::Editor::FSubobjectEditorContext& Context, EVisibility& VisibilityOut) const override;
 	virtual EChildActorComponentTreeViewVisualizationMode GetChildActorVisualizationMode() const override;
 	virtual TSubclassOf<UActorComponent> GetComponentTypeFilter(TArrayView<UObject*> Context) const override;
@@ -304,14 +304,14 @@ const FSlateBrush* SActorDetails::FInternalSubobjectCustomization::GetIconBrush(
 	return ISCSEditorUICustomization::GetIconBrush(SubobjectData);
 }
 
-TSharedPtr<SWidget> SActorDetails::FInternalSubobjectCustomization::GetControlsWidget(const FSubobjectData& SubobjectData) const
+TSharedPtr<SWidget> SActorDetails::FInternalSubobjectCustomization::GetControlsWidget(TSharedRef<SSubobjectEditor>& InSubobjectEditor, const FSubobjectData& SubobjectData) const
 {
 	if (ParentCustomization)
 	{
-		return ParentCustomization->GetControlsWidget(SubobjectData);
+		return ParentCustomization->GetControlsWidget(InSubobjectEditor, SubobjectData);
 	}
 
-	return ISCSEditorUICustomization::GetControlsWidget(SubobjectData);
+	return ISCSEditorUICustomization::GetControlsWidget(InSubobjectEditor, SubobjectData);
 }
 
 bool SActorDetails::FInternalSubobjectCustomization::OverrideDesiredVisibility(const UE::Editor::FSubobjectEditorContext& Context, EVisibility& VisibilityOut) const
diff --git a/Engine/Source/Editor/SubobjectEditor/Private/SSubobjectEditor.cpp b/Engine/Source/Editor/SubobjectEditor/Private/SSubobjectEditor.cpp
--- a/Engine/Source/Editor/SubobjectEditor/Private/SSubobjectEditor.cpp
+++ b/Engine/Source/Editor/SubobjectEditor/Private/SSubobjectEditor.cpp
@@ -1769,7 +1769,8 @@ TSharedRef<SWidget> SSubobject_RowWidget::GetInheritedLinkWidget()
 	{
 		if (const ISCSEditorUICustomization* UICustomization = PinnedSubobjectEditor->GetUICustomization())
 		{
-			if (TSharedPtr<SWidget> ControlsWidget = UICustomization->GetControlsWidget(*Data))
+			TSharedRef<SSubobjectEditor> SubobjectEditorRef = PinnedSubobjectEditor.ToSharedRef();
+			if (TSharedPtr<SWidget> ControlsWidget = UICustomization->GetControlsWidget(SubobjectEditorRef, *Data))
 			{
 				return ControlsWidget.ToSharedRef();
 			}
