[Details Panel] Component selection is incorrect after modifying a Spline Point in a component of the same Actor

#Synced-CL 49370055

diff --git a/Engine/Source/Editor/LevelEditor/Private/SActorDetails.cpp b/Engine/Source/Editor/LevelEditor/Private/SActorDetails.cpp
--- a/Engine/Source/Editor/LevelEditor/Private/SActorDetails.cpp
+++ b/Engine/Source/Editor/LevelEditor/Private/SActorDetails.cpp
@@ -1023,6 +1023,13 @@ void SActorDetails::OnSubobjectEditorTreeViewSelectionChanged(const TArray<FSubo
 
 			// Enable the selection guard to prevent OnEditorSelectionChanged() from altering the contents of the SubobjectTreeWidget
 			TGuardValue<bool> SelectionGuard(bSelectionGuard, true);
+
+			// Make sure we clear component visualizers if active. This ensures things like the viewport gizmos reflect the actual selection.
+			if (GUnrealEd && GUnrealEd->ComponentVisManager.IsActive())
+			{
+				GUnrealEd->ComponentVisManager.ClearActiveComponentVis();
+			}
+
 			SelectionSet->SetSelection(NewEditorSelection, FTypedElementSelectionOptions());
 			SelectionSet->NotifyPendingChanges(); // Fire while still under the selection guard
 		}
