Bump the minimum size of the right column for the details panel
- The width of the wrap button is larger than the minimum colum width of the level editor's detail panel right column width.  When the number of extension buttons exceed the width of the column, it would replace with a wrap button.  However, if the wrap button exceeds that size, then nothing is displayed and makes it appear as if there are no extension buttons.
- Bumping the size of this right column for now to provide space for at least two buttons before wrapping.  Will also encompass wrap button width.

diff --git a/Engine/Source/Editor/LevelEditor/Private/SActorDetails.cpp b/Engine/Source/Editor/LevelEditor/Private/SActorDetails.cpp
--- a/Engine/Source/Editor/LevelEditor/Private/SActorDetails.cpp
+++ b/Engine/Source/Editor/LevelEditor/Private/SActorDetails.cpp
@@ -456,6 +456,7 @@ void SActorDetails::Construct(
 	DetailsViewArgs.HostCommandList = InCommandList;
 	DetailsViewArgs.HostTabManager = InTabManager;
 	DetailsViewArgs.bShowSectionSelector = true;
+	DetailsViewArgs.RightColumnMinWidth = 44.0f;
 
 	FPropertyEditorModule& PropPlugin = FModuleManager::LoadModuleChecked<FPropertyEditorModule>("PropertyEditor");
 	DetailsView = PropPlugin.CreateDetailView(DetailsViewArgs);
