EditorUsability : LevelEditor
- Added project badge context menu with the ability to copy infos to clipboard (project name, version, changelist, branch, build configuration, build date, RHI)
- Added double click to project badge will copy changelist to clipboard (for QA)

#Synced-CL 49549735

diff --git a/Engine/Source/Editor/LevelEditor/Private/LevelEditor.cpp b/Engine/Source/Editor/LevelEditor/Private/LevelEditor.cpp
--- a/Engine/Source/Editor/LevelEditor/Private/LevelEditor.cpp
+++ b/Engine/Source/Editor/LevelEditor/Private/LevelEditor.cpp
@@ -11,6 +11,7 @@
 #include "EditorModes.h"
 #include "EditorViewportCommands.h"
 #include "Framework/Application/SlateApplication.h"
+#include "HAL/PlatformApplicationMisc.h"
 #include "IContentBrowserSingleton.h"
 #include "Interfaces/IProjectManager.h"
 #include "ISlateReflectorModule.h"
@@ -38,6 +39,7 @@
 #include "Widgets/Docking/SDockTab.h"
 #include "Widgets/Layout/SBorder.h"
 #include "Widgets/Layout/SBox.h"
+#include "Widgets/Notifications/SNotificationList.h"
 #include "Widgets/Text/STextBlock.h"
 #include "WorkspaceMenuStructure.h"
 #include "WorkspaceMenuStructureModule.h"
@@ -148,35 +150,34 @@ public:
 		}
 
 		const FString EngineVersionString = FEngineVersion::Current().ToString(FEngineVersion::Current().HasChangelist() ? EVersionComponent::Changelist : EVersionComponent::Patch);
-
 		const EBuildConfiguration BuildConfig = FApp::GetBuildConfiguration();
 		const bool bAppendBuildConfig = (BuildConfig != EBuildConfiguration::Shipping && BuildConfig != EBuildConfiguration::Development && BuildConfig != EBuildConfiguration::Unknown);
 
-		FFormatNamedArguments Args;
-		Args.Add(TEXT("ProjectNameWatermarkPrefix"), FText::FromString(ProjectNameWatermarkPrefix));
-		Args.Add(TEXT("Branch"), FEngineBuildSettings::IsPerforceBuild() ? FText::FromString(TEXT(" ") + FEngineVersion::Current().GetBranch()) : FText::GetEmpty());
-		Args.Add(TEXT("GameName"), FText::FromString(TEXT(" ") + FString(FApp::GetProjectName())));
-		Args.Add(TEXT("EngineVersion"), (GetDefault<UEditorPerProjectUserSettings>()->bDisplayEngineVersionInBadge) ? FText::FromString(TEXT(" (") + EngineVersionString + TEXT(")")) : FText::GetEmpty());
-		Args.Add(TEXT("Config"), bAppendBuildConfig ? FText::Format(NSLOCTEXT("UnrealEditor", "TitleBarConfig", " [{0}]"), EBuildConfigurations::ToText(BuildConfig)) : FText::GetEmpty());
-
-		FText RightContentText = FText::Format(NSLOCTEXT("UnrealEditor", "TitleBarRightContentAndConfig", "{ProjectNameWatermarkPrefix}{GameName}{Config}{Branch}{EngineVersion}"), Args);
+		FFormatNamedArguments LabelArgs;
+		LabelArgs.Add(TEXT("ProjectNameWatermarkPrefix"), FText::FromString(ProjectNameWatermarkPrefix));
+		LabelArgs.Add(TEXT("Branch"), FEngineBuildSettings::IsPerforceBuild() ? FText::FromString(TEXT(" ") + FEngineVersion::Current().GetBranch()) : FText::GetEmpty());
+		LabelArgs.Add(TEXT("GameName"), FText::FromString(TEXT(" ") + FString(FApp::GetProjectName())));
+		LabelArgs.Add(TEXT("EngineVersion"), (GetDefault<UEditorPerProjectUserSettings>()->bDisplayEngineVersionInBadge) ? FText::FromString(TEXT(" (") + EngineVersionString + TEXT(")")) : FText::GetEmpty());
+		LabelArgs.Add(TEXT("Config"), bAppendBuildConfig ? FText::Format(NSLOCTEXT("UnrealEditor", "TitleBarConfig", " [{0}]"), EBuildConfigurations::ToText(BuildConfig)) : FText::GetEmpty());
 
+		const FText RightContentText = FText::Format(NSLOCTEXT("UnrealEditor", "TitleBarRightContentAndConfig", "{ProjectNameWatermarkPrefix}{GameName}{Config}{Branch}{EngineVersion}"), LabelArgs);
+	
 		// Create the tooltip showing more detailed information
-		FFormatNamedArguments TooltipArgs;
-		FString TooltipVersionStr = EngineVersionString;
-		TooltipArgs.Add(TEXT("Version"), FText::FromString(TooltipVersionStr));
-		TooltipArgs.Add(TEXT("Branch"), FText::FromString(FEngineVersion::Current().GetBranch()));
-		TooltipArgs.Add(TEXT("BuildConfiguration"), EBuildConfigurations::ToText(BuildConfig));
-		TooltipArgs.Add(TEXT("BuildDate"), FText::FromString(FApp::GetBuildDate()));
-		TooltipArgs.Add(TEXT("GraphicsRHI"), FText::FromString(FApp::GetGraphicsRHI()));
-
-		FText RightContentTooltip = FText::Format(NSLOCTEXT("UnrealEditor", "TitleBarRightContentTooltip", "Version: {Version}\nBranch: {Branch}\nBuild Configuration: {BuildConfiguration}\nBuild Date: {BuildDate}\nGraphics RHI: {GraphicsRHI}"), TooltipArgs);
-
-		SetToolTipText(RightContentTooltip);
+		{
+			FFormatNamedArguments TooltipArgs;
+			TooltipArgs.Add(TEXT("Version"), GetVersionText());
+			TooltipArgs.Add(TEXT("Changelist"), GetChangelistText());
+			TooltipArgs.Add(TEXT("Branch"), GetBranchText());
+			TooltipArgs.Add(TEXT("BuildConfiguration"), GetBuildConfigurationText());
+			TooltipArgs.Add(TEXT("BuildDate"), GetBuildDateText());
+			TooltipArgs.Add(TEXT("GraphicsRHI"), GetGraphicsRHIText());
+
+			const FText RightContentTooltip = FText::Format(NSLOCTEXT("UnrealEditor", "TitleBarRightContentTooltip", "Version: {Version}\nBranch: {Branch}\nBuild Configuration: {BuildConfiguration}\nBuild Date: {BuildDate}\nGraphics RHI: {GraphicsRHI}\n\nRight Click to open context menu with various clipboard options\nDouble Click to quickly copy changelist to clipboard"), TooltipArgs);
+			SetToolTipText(RightContentTooltip);
+		}
 
 		TSharedRef<SWidget> DefaultNamePlate = SNew(STextBlock)
 			.Text(RightContentText)
-			.Visibility(EVisibility::HitTestInvisible)
 			.TextStyle(FAppStyle::Get(), "SProjectBadge.Text")
 			.Margin(FAppStyle::Get().GetMargin("SProjectBadge.BadgePadding"))
 			.ColorAndOpacity(BadgeTextColor);
@@ -188,6 +189,8 @@ public:
 			.Padding(FMargin(0.0f, 0.0f, 2.0f, 0.0f))
 			[
 				SNew(SBorder)
+				.OnMouseDoubleClick(this, &SProjectBadge::OnBadgeDoubleClick)
+				.OnMouseButtonDown(this, &SProjectBadge::OnOpenContextMenu)
 				.BorderBackgroundColor(BadgeBackgroundColor)
 				.BorderImage(FAppStyle::GetBrush("SProjectBadge.BadgeShape"))
 				.Padding(FMargin(0))
@@ -215,6 +218,183 @@ public:
 	// End of SWidget interface
 
 private:
+	static void CopyToClipboard(const FText& InKey, const FText& InValue)
+	{
+		FPlatformApplicationMisc::ClipboardCopy(*InValue.ToString());
+		
+		// Notify of the action
+		FNotificationInfo NotificationInfo(FText::Format(LOCTEXT("ProjectBadgeCopyNotification", "{0} value copied to clipboard\n{1}"), InKey, InValue));
+		NotificationInfo.ExpireDuration = 1.f;
+		NotificationInfo.bFireAndForget = true;
+		NotificationInfo.Image = FAppStyle::GetBrush("Icons.SuccessWithColor");
+		FSlateNotificationManager::Get().AddNotification(NotificationInfo);
+	}
+	
+	static FText GetProjectNameText()
+	{
+		return FText::FromString(FApp::GetProjectName());
+	}
+
+	static FText GetVersionText()
+	{
+		const FString EngineVersionString = FEngineVersion::Current().ToString(FEngineVersion::Current().HasChangelist() ? EVersionComponent::Changelist : EVersionComponent::Patch);
+		return FText::FromString(EngineVersionString);
+	}
+	
+	static FText GetChangelistText()
+	{
+		return FText::FromString(LexToString(FEngineVersion::Current().GetChangelist()));
+	}
+	
+	static FText GetBranchText()
+	{
+		return FText::FromString(FEngineVersion::Current().GetBranch());
+	}
+	
+	static FText GetBuildConfigurationText()
+	{
+		return EBuildConfigurations::ToText(FApp::GetBuildConfiguration());
+	}
+
+	static FText GetBuildDateText()
+	{
+		return FText::FromString(FApp::GetBuildDate());
+	}
+	
+	static FText GetGraphicsRHIText()
+	{
+		return FText::FromString(FApp::GetGraphicsRHI());
+	}
+
+	static TSharedRef<SWidget> CreateContextMenu()
+	{
+		FMenuBuilder MenuBuilder(/** CloseAfterSelection */true, /** CommandList */nullptr);
+		
+		const FText LabelFormat = LOCTEXT("ProjectBadgeContextLabel", "Copy {0} ({1})");
+		const FText TooltipFormat = LOCTEXT("ProjectBadgeContextTooltip", "Copy {0} value ({1}) to the clipboard");
+
+		// ProjectName
+		{
+			const FText ProjectNameLabel = LOCTEXT("ProjectNameLabel", "Project Name");
+			const FText ProjectNameValue = GetProjectNameText();
+
+			MenuBuilder.AddMenuEntry(
+				FText::Format(LabelFormat, ProjectNameLabel, ProjectNameValue),
+				FText::Format(TooltipFormat, ProjectNameLabel, ProjectNameValue),
+				FSlateIcon(FCoreStyle::Get().GetStyleSetName(), "GenericCommands.Copy"),
+				FUIAction(FExecuteAction::CreateLambda(&SProjectBadge::CopyToClipboard, ProjectNameLabel, ProjectNameValue))
+			);
+		}
+		
+		// Version
+		{
+			const FText VersionLabel = LOCTEXT("VersionLabel", "Version");
+			const FText VersionValue = GetVersionText();
+
+			MenuBuilder.AddMenuEntry(
+				FText::Format(LabelFormat, VersionLabel, VersionValue),
+				FText::Format(TooltipFormat, VersionLabel, VersionValue),
+				FSlateIcon(FCoreStyle::Get().GetStyleSetName(), "GenericCommands.Copy"),
+				FUIAction(FExecuteAction::CreateLambda(&SProjectBadge::CopyToClipboard, VersionLabel, VersionValue))
+			);
+		}
+
+		// Changelist
+		{
+			const FText ChangelistLabel = LOCTEXT("ChangelistLabel", "Changelist");
+			const FText ChangelistValue = GetChangelistText();
+
+			MenuBuilder.AddMenuEntry(
+				FText::Format(LabelFormat, ChangelistLabel, ChangelistValue),
+				FText::Format(TooltipFormat, ChangelistLabel, ChangelistValue),
+				FSlateIcon(FCoreStyle::Get().GetStyleSetName(), "GenericCommands.Copy"),
+				FUIAction(FExecuteAction::CreateLambda(&SProjectBadge::CopyToClipboard, ChangelistLabel, ChangelistValue))
+			);
+		}
+		
+		// Branch
+		{
+			const FText BranchLabel = LOCTEXT("BranchLabel", "Branch");
+			const FText BranchValue = GetBranchText();
+
+			MenuBuilder.AddMenuEntry(
+				FText::Format(LabelFormat, BranchLabel, BranchValue),
+				FText::Format(TooltipFormat, BranchLabel, BranchValue),
+				FSlateIcon(FCoreStyle::Get().GetStyleSetName(), "GenericCommands.Copy"),
+				FUIAction(FExecuteAction::CreateLambda(&SProjectBadge::CopyToClipboard, BranchLabel, BranchValue))
+			);
+		}
+		
+		// BuildConfiguration
+		{
+			const FText BuildConfigurationLabel = LOCTEXT("BuildConfigurationLabel", "Build Configuration");
+			const FText BuildConfigurationValue = GetBuildConfigurationText();
+
+			MenuBuilder.AddMenuEntry(
+				FText::Format(LabelFormat, BuildConfigurationLabel, BuildConfigurationValue),
+				FText::Format(TooltipFormat, BuildConfigurationLabel, BuildConfigurationValue),
+				FSlateIcon(FCoreStyle::Get().GetStyleSetName(), "GenericCommands.Copy"),
+				FUIAction(FExecuteAction::CreateLambda(&SProjectBadge::CopyToClipboard, BuildConfigurationLabel, BuildConfigurationValue))
+			);
+		}
+		
+		// BuildDate
+		{
+			const FText BuildDateLabel = LOCTEXT("BuildDateLabel", "Build Date");
+			const FText BuildDateValue = GetBuildDateText();
+
+			MenuBuilder.AddMenuEntry(
+				FText::Format(LabelFormat, BuildDateLabel, BuildDateValue),
+				FText::Format(TooltipFormat, BuildDateLabel, BuildDateValue),
+				FSlateIcon(FCoreStyle::Get().GetStyleSetName(), "GenericCommands.Copy"),
+				FUIAction(FExecuteAction::CreateLambda(&SProjectBadge::CopyToClipboard, BuildDateLabel, BuildDateValue))
+			);
+		}
+		
+		// GraphicsRHI
+		{
+			const FText GraphicsRHILabel = LOCTEXT("GraphicsRHILabel", "Graphics RHI");
+			const FText GraphicsRHIValue = GetGraphicsRHIText();
+
+			MenuBuilder.AddMenuEntry(
+				FText::Format(LabelFormat, GraphicsRHILabel, GraphicsRHIValue),
+				FText::Format(TooltipFormat, GraphicsRHILabel, GraphicsRHIValue),
+				FSlateIcon(FCoreStyle::Get().GetStyleSetName(), "GenericCommands.Copy"),
+				FUIAction(FExecuteAction::CreateLambda(&SProjectBadge::CopyToClipboard, GraphicsRHILabel, GraphicsRHIValue))
+			);
+		}
+
+		return MenuBuilder.MakeWidget();
+	}
+
+	FReply OnBadgeDoubleClick(const FGeometry& InGeometry, const FPointerEvent& InEvent)
+	{
+		const FText ChangelistText = GetChangelistText();
+		if (!ChangelistText.IsEmpty())
+		{
+			CopyToClipboard(LOCTEXT("ChangelistLabel", "Changelist"), ChangelistText);
+			return FReply::Handled();
+		}
+		return FReply::Unhandled();
+	}
+
+	FReply OnOpenContextMenu(const FGeometry& InGeometry, const FPointerEvent& InEvent)
+	{
+		if (InEvent.GetEffectingButton() == EKeys::RightMouseButton)
+		{
+			FSlateApplication::Get().PushMenu(
+				AsShared(),
+				FWidgetPath(),
+				CreateContextMenu(),
+				FSlateApplication::Get().GetCursorPos(),
+				FPopupTransitionEffect(FPopupTransitionEffect::ContextMenu)
+			);
+			return FReply::Handled();
+		}
+
+		return FReply::Unhandled();
+	}
+
 	FGeometry CachedGeometry;
 };
 
