Fix crash in LevelViewport while trying to access a nullptr world

#tests Executed the code while testing in engine test, behave normally.
#Synced-CL 49624560

diff --git a/Engine/Source/Editor/LevelEditor/Private/SLevelViewport.cpp b/Engine/Source/Editor/LevelEditor/Private/SLevelViewport.cpp
--- a/Engine/Source/Editor/LevelEditor/Private/SLevelViewport.cpp
+++ b/Engine/Source/Editor/LevelEditor/Private/SLevelViewport.cpp
@@ -555,7 +555,7 @@ bool SLevelViewport::IsActorEditorContextVisible() const
 		GetWorld()->GetCurrentLevel() &&
 		(&GetLevelViewportClient() == GCurrentLevelEditingViewportClient) &&
 		ActiveViewport.IsValid() &&
-		(ActiveViewport->GetPlayInEditorIsSimulate() || !ActiveViewport->GetClient()->GetWorld()->IsGameWorld()) &&
+		(ActiveViewport->GetPlayInEditorIsSimulate() || (ActiveViewport->GetClient()->GetWorld() && !ActiveViewport->GetClient()->GetWorld()->IsGameWorld())) &&
 		SActorEditorContext::IsVisible(GetWorld());
 }
 
