[ITF Interactions] Add the "Alt+Middle+Click" set pivot gesture as an explicit Level Viewport click handler.

#rnx
#Synced-CL 49683958

diff --git a/Engine/Source/Editor/LevelEditor/Private/SLevelEditor.cpp b/Engine/Source/Editor/LevelEditor/Private/SLevelEditor.cpp
--- a/Engine/Source/Editor/LevelEditor/Private/SLevelEditor.cpp
+++ b/Engine/Source/Editor/LevelEditor/Private/SLevelEditor.cpp
@@ -770,6 +770,9 @@ void SLevelEditor::RegisterViewportInteractions(UViewportInteractionsBehaviorSou
 
 	// Right Mouse Click - summons context menu
 	Source->AddInteraction<ULevelViewportClickContextMenu>();
+	
+	// Middle Mouse Click
+	Source->AddInteraction<ULevelViewportSetPivot>();
 
 	// Double click
 	Source->AddInteraction<ULevelViewportDoubleClick>();
diff --git a/Engine/Source/Editor/LevelEditor/Private/ViewportInteractions/LevelViewportClickSelection.cpp b/Engine/Source/Editor/LevelEditor/Private/ViewportInteractions/LevelViewportClickSelection.cpp
--- a/Engine/Source/Editor/LevelEditor/Private/ViewportInteractions/LevelViewportClickSelection.cpp
+++ b/Engine/Source/Editor/LevelEditor/Private/ViewportInteractions/LevelViewportClickSelection.cpp
@@ -117,6 +117,34 @@ ULevelViewportClickContextMenu::ULevelViewportClickContextMenu()
 	}
 }
 
+ULevelViewportSetPivot::ULevelViewportSetPivot()
+{
+	InteractionName = TEXT("Level Viewport Middle Click");
+	Groups = { UE::Editor::ViewportInteractions::ViewportClick };
+
+	MouseButton = EKeys::MiddleMouseButton;
+
+	if (TStrongObjectPtr<USingleClickInputBehavior> SingleClickInputBehavior = SingleClickInputBehaviorWeak.Pin())
+	{
+		SingleClickInputBehavior->SetUseMiddleMouseButton();
+		SingleClickInputBehavior->ModifierCheckFunc = [](const FInputDeviceState& InputDeviceState)
+		{
+			return InputDeviceState.bAltKeyDown;
+		};
+	}
+}
+
+void ULevelViewportSetPivot::ProcessClick_Internal(const FViewportClick& InViewportClick, HHitProxy* InHitProxy, FSceneView& View)
+{
+	if (GEditor)
+	{
+		// This only sets the pivot for the _next_ transform interaction.
+		// The position of the pivot will be reset as soon as the next transform interaction completes.
+		// This is consistent with legacy behavior
+		GEditor->SetPivot(GEditor->ClickLocation, true, false, true);
+	}
+}
+
 ULevelViewportDoubleClick::ULevelViewportDoubleClick()
 {
 	UDoubleClickInputBehavior* DoubleClickInputBehavior = NewObject<UDoubleClickInputBehavior>();
diff --git a/Engine/Source/Editor/LevelEditor/Private/ViewportInteractions/LevelViewportClickSelection.h b/Engine/Source/Editor/LevelEditor/Private/ViewportInteractions/LevelViewportClickSelection.h
--- a/Engine/Source/Editor/LevelEditor/Private/ViewportInteractions/LevelViewportClickSelection.h
+++ b/Engine/Source/Editor/LevelEditor/Private/ViewportInteractions/LevelViewportClickSelection.h
@@ -35,6 +35,21 @@ public:
 	UE_API ULevelViewportClickContextMenu();
 };
 
+/**
+ * Set the pivot based on click, for various meta-interactions
+ */
+UCLASS(MinimalAPI)
+class ULevelViewportSetPivot : public ULevelViewportClickSelection
+{
+	GENERATED_BODY()
+
+public:
+	UE_API ULevelViewportSetPivot();
+	
+protected:
+	UE_API virtual void ProcessClick_Internal(const FViewportClick& InViewportClick, HHitProxy* InHitProxy, FSceneView& View) override;
+};
+
 /**
  * Left-mouse double click
  */
