Fix crash when an editor mode wants to be visible in the toolbar without a command.

#rnx

#Synced-CL 49727259

diff --git a/Engine/Source/Editor/LevelEditor/Private/LevelEditorToolBar.cpp b/Engine/Source/Editor/LevelEditor/Private/LevelEditorToolBar.cpp
--- a/Engine/Source/Editor/LevelEditor/Private/LevelEditorToolBar.cpp
+++ b/Engine/Source/Editor/LevelEditor/Private/LevelEditorToolBar.cpp
@@ -1381,7 +1381,7 @@ void FLevelEditorToolBar::RegisterLevelEditorToolBar( const TSharedRef<FUIComman
 								TSharedPtr<FUICommandInfo> EditorModeCommand =
 									FInputBindingManager::Get().FindCommandInContext(Commands.GetContextName(), EditorModeCommandName);
 
-								if (Mode.IsVisible())
+								if (Mode.IsVisible() && EditorModeCommand.IsValid())
 								{
 									CommandInfos.Add(EditorModeCommand);
 								}
