Level Editor: It's now possible to enumerate the available viewports by tab, not just as a flat list.

#Synced-CL 49784486

diff --git a/Engine/Source/Editor/LevelEditor/Private/SLevelEditor.cpp b/Engine/Source/Editor/LevelEditor/Private/SLevelEditor.cpp
--- a/Engine/Source/Editor/LevelEditor/Private/SLevelEditor.cpp
+++ b/Engine/Source/Editor/LevelEditor/Private/SLevelEditor.cpp
@@ -1989,6 +1989,60 @@ TArray< TSharedPtr< SLevelViewport > > SLevelEditor::GetViewports() const
 
 	return OutViewports;
 }
+
+TMap<int32, TArray<TSharedPtr<SLevelViewport>>> SLevelEditor::GetViewportsByTab() const
+{
+	constexpr int32 ExternalViewportIndex = 4;
+
+	TMap< int32, TArray< TSharedPtr< SLevelViewport > > > OutViewports;
+	OutViewports.Reserve(5); // 4 tabs and 1 standalone
+
+	for (int32 TabIndex = 0; TabIndex < ViewportTabs.Num(); ++TabIndex)
+	{
+		TArray< TSharedPtr< SLevelViewport > > TabViewportList;
+		TabViewportList.Reserve(4);
+
+		TSharedPtr<FLevelViewportTabContent> ViewportTab = ViewportTabs[TabIndex].Pin();
+
+		if (ViewportTab.IsValid())
+		{
+			const TMap< FName, TSharedPtr< IEditorViewportLayoutEntity > >* LevelViewports = ViewportTab->GetViewports();
+
+			if (LevelViewports != NULL)
+			{
+				for (auto& Pair : *LevelViewports)
+				{
+					TSharedPtr<SLevelViewport> Viewport = StaticCastSharedPtr<ILevelViewportLayoutEntity>(Pair.Value)->AsLevelViewport();
+					if (Viewport.IsValid())
+					{
+						TabViewportList.Add(Viewport);
+					}
+				}
+			}
+		}
+
+		OutViewports.Add(TabIndex, TabViewportList);
+	}
+
+	// Also add any standalone viewports
+	{
+		TArray< TSharedPtr< SLevelViewport > > StandaloneViewportList;
+		StandaloneViewportList.Reserve(StandaloneViewports.Num());
+
+		for (const TWeakPtr< SLevelViewport >& StandaloneViewportWeakPtr : StandaloneViewports)
+		{
+			const TSharedPtr< SLevelViewport > Viewport = StandaloneViewportWeakPtr.Pin();
+			if (Viewport.IsValid())
+			{
+				StandaloneViewportList.Add(Viewport);
+			}
+		}
+
+		OutViewports.Add(ExternalViewportIndex, StandaloneViewportList);
+	}
+
+	return OutViewports;
+}
  
 TSharedPtr<SLevelViewport> SLevelEditor::GetActiveViewportInterface()
 {
diff --git a/Engine/Source/Editor/LevelEditor/Private/SLevelEditor.h b/Engine/Source/Editor/LevelEditor/Private/SLevelEditor.h
--- a/Engine/Source/Editor/LevelEditor/Private/SLevelEditor.h
+++ b/Engine/Source/Editor/LevelEditor/Private/SLevelEditor.h
@@ -76,6 +76,7 @@ public:
 	virtual void SummonLevelViewportViewOptionMenu(const ELevelViewportType ViewOption) override;
 	virtual const TArray< TSharedPtr< class IToolkit > >& GetHostedToolkits() const override;
 	virtual TArray< TSharedPtr< SLevelViewport > > GetViewports() const override;
+	virtual TMap< int32, TArray< TSharedPtr< SLevelViewport > > > GetViewportsByTab() const override;
 	virtual TSharedPtr<SLevelViewport> GetActiveViewportInterface() override;
 	virtual void AppendCommands( const TSharedRef<FUICommandList>& InCommandsToAppend ) override;
 	virtual void AddStandaloneLevelViewport( const TSharedRef<SLevelViewport>& LevelViewport ) override;
diff --git a/Engine/Source/Editor/LevelEditor/Public/ILevelEditor.h b/Engine/Source/Editor/LevelEditor/Public/ILevelEditor.h
--- a/Engine/Source/Editor/LevelEditor/Public/ILevelEditor.h
+++ b/Engine/Source/Editor/LevelEditor/Public/ILevelEditor.h
@@ -48,6 +48,12 @@ public:
 
 	/** Gets an array of all viewports in this level editor */
 	virtual TArray< TSharedPtr< SLevelViewport > > GetViewports() const = 0;
+
+	/** 
+	 * Gets a list of all the viewport tabs in this level editor mapped to arrays of all viewports in that tab. 
+	 * There are 4 viewport tabs, mapped to 0 to 3. Standalone viewports are mapped to index 4. 
+	 */
+	virtual TMap< int32, TArray< TSharedPtr< SLevelViewport > > > GetViewportsByTab() const = 0;
 	
 	/** Gets the active level viewport for this level editor */
 	virtual TSharedPtr<SLevelViewport> GetActiveViewportInterface() = 0;
