[SLevelEditor/ALT]
- Fix minor race when actor changes during loading

#Synced-CL 50160399

diff --git a/Engine/Source/Editor/LevelEditor/Private/SLevelEditor.h b/Engine/Source/Editor/LevelEditor/Private/SLevelEditor.h
--- a/Engine/Source/Editor/LevelEditor/Private/SLevelEditor.h
+++ b/Engine/Source/Editor/LevelEditor/Private/SLevelEditor.h
@@ -370,7 +370,7 @@ private:
 	FOnActiveViewportChanged OnActiveViewportChangedDelegate;
 
 	/** If this flag is raised we will force refresh on next selection update. */
-	bool bNeedsRefresh : 1;
+	std::atomic<bool> bNeedsRefresh;
 
 	TMap<FName, TSharedPtr<FLevelEditorModeUILayer>> ModeUILayers;
 
