Fixed typo in CVar text

#13864

#Synced-CL 50045910
#Synced-CL 50045910

diff --git a/Engine/Source/Editor/LevelEditor/Private/LevelEditorToolBar.cpp b/Engine/Source/Editor/LevelEditor/Private/LevelEditorToolBar.cpp
--- a/Engine/Source/Editor/LevelEditor/Private/LevelEditorToolBar.cpp
+++ b/Engine/Source/Editor/LevelEditor/Private/LevelEditorToolBar.cpp
@@ -75,8 +75,8 @@ int32 CVarLevelEditorToolbarSettingsValue = 0;
 static FAutoConsoleVariableRef CVarToolMenusViewportToolbars(
 	TEXT("LevelEditorToolbarSettings"),
 	UE::Editor::Private::CVarLevelEditorToolbarSettingsValue,
-	TEXT("Control whether the level editor Settings menu is visible in the toolbar. Set to 1 (default) to show it and "
-		 "0 to hide it."),
+	TEXT("Control whether the level editor Settings menu is visible in the toolbar. Set to 1 to show it and "
+		 "0 (default) to hide it."),
 	ECVF_Default
 );
 
