Add exposure to Level Editor Subsystem to Get and Set a level editor viewport camera's location, rotation and FOV
Clamps FOV between 5 and 170 degrees, in line with the Editor UI spinner box UI

#Synced-CL 50635208

diff --git a/Engine/Source/Editor/LevelEditor/Private/LevelEditorSubsystem.cpp b/Engine/Source/Editor/LevelEditor/Private/LevelEditorSubsystem.cpp
--- a/Engine/Source/Editor/LevelEditor/Private/LevelEditorSubsystem.cpp
+++ b/Engine/Source/Editor/LevelEditor/Private/LevelEditorSubsystem.cpp
@@ -341,6 +341,70 @@ bool ULevelEditorSubsystem::GetAllowsCinematicControl(FName ViewportConfigKey)
 	return false;
 }
 
+bool ULevelEditorSubsystem::GetLevelViewportCameraInfo(FVector& CameraLocation, FRotator& CameraRotation, FName ViewportConfigKey)
+{
+	CameraLocation = FVector::ZeroVector;
+	CameraRotation = FRotator::ZeroRotator;
+	
+	TSharedPtr<SLevelViewport> LevelViewport = InternalEditorLevelLibrary::GetLevelViewport(ViewportConfigKey);
+	if (LevelViewport.IsValid())
+	{
+		FLevelEditorViewportClient* LevelVC = &LevelViewport->GetLevelViewportClient();
+		if (LevelVC)
+		{
+			CameraLocation = LevelVC->GetViewLocation();
+			CameraRotation = LevelVC->GetViewRotation();
+			return true;
+		}
+	}
+	return false;
+}
+
+void ULevelEditorSubsystem::SetLevelViewportCameraInfo(FVector CameraLocation, FRotator CameraRotation, FName ViewportConfigKey)
+{
+	TSharedPtr<SLevelViewport> LevelViewport = InternalEditorLevelLibrary::GetLevelViewport(ViewportConfigKey);
+	if (LevelViewport.IsValid())
+	{
+		FLevelEditorViewportClient* LevelVC = &LevelViewport->GetLevelViewportClient();
+		if (LevelVC)
+		{
+			LevelVC->SetViewLocationForOrbiting(CameraLocation);
+			LevelVC->SetViewLocation(CameraLocation);
+			LevelVC->SetViewRotation(CameraRotation);
+		}
+	}
+}
+
+bool ULevelEditorSubsystem::GetLevelViewportFOV(float& FOV, FName ViewportConfigKey)
+{
+	FOV = 0.f;
+	TSharedPtr<SLevelViewport> LevelViewport = InternalEditorLevelLibrary::GetLevelViewport(ViewportConfigKey);
+	if (LevelViewport.IsValid())
+	{
+		FLevelEditorViewportClient* LevelVC = &LevelViewport->GetLevelViewportClient();
+		if (LevelVC)
+		{
+			FOV = LevelVC->ViewFOV;
+			return true;
+		}
+	}
+	
+	return false;
+}
+
+void ULevelEditorSubsystem::SetLevelViewportFOV(float FOV, FName ViewportConfigKey)
+{
+	TSharedPtr<SLevelViewport> LevelViewport = InternalEditorLevelLibrary::GetLevelViewport(ViewportConfigKey);
+	if (LevelViewport.IsValid())
+	{
+		FLevelEditorViewportClient* LevelVC = &LevelViewport->GetLevelViewportClient();
+		if (LevelVC)
+		{
+			LevelVC->ViewFOV = FMath::Clamp(FOV, 5.f, 170.f);
+		}
+	}
+}
+
 /**
  *
  * Editor Scripting | Level
diff --git a/Engine/Source/Editor/LevelEditor/Public/LevelEditorSubsystem.h b/Engine/Source/Editor/LevelEditor/Public/LevelEditorSubsystem.h
--- a/Engine/Source/Editor/LevelEditor/Public/LevelEditorSubsystem.h
+++ b/Engine/Source/Editor/LevelEditor/Public/LevelEditorSubsystem.h
@@ -95,6 +95,47 @@ public:
 	UFUNCTION(BlueprintCallable, Category = "Editor Scripting | Level Utility", meta = (DevelopmentOnly))
 	UE_API bool GetAllowsCinematicControl(FName ViewportConfigKey = NAME_None);
 
+	/**
+	 * Gets information about the camera position for the given level editor viewport.
+	 *
+	 * @param	CameraLocation	(out) Current location of the level editing viewport camera, or zero if none found
+	 * @param	CameraRotation	(out) Current rotation of the level editing viewport camera, or zero if none found
+	 * @param	ViewportConfigKey the viewport you wish to get from
+	 * @return	Whether we were able to get the camera position for a level editing viewport
+	 */
+	UFUNCTION(BlueprintPure, Category = "Editor Scripting | Level Utility", meta = (DevelopmentOnly))
+	UE_API bool GetLevelViewportCameraInfo(FVector& CameraLocation, FRotator& CameraRotation, FName ViewportConfigKey);
+
+	
+	/**
+	* Sets information about the camera position for the given level editor viewport.
+	*
+	* @param	CameraLocation	Location the camera will be moved to
+	* @param	CameraRotation	Rotation the camera will be set to
+	* @param	ViewportConfigKey	the viewport you wish to target
+	*/
+	UFUNCTION(BlueprintCallable, Category = "Editor Scripting | Level Utility", meta = (DevelopmentOnly))
+	UE_API void SetLevelViewportCameraInfo(FVector CameraLocation, FRotator CameraRotation, FName ViewportConfigKey);
+	
+	/**
+	 * Gets the horizontal field of view for the given level editor viewport. 
+	 *
+	 * @param	FOV				Field of View of the Level editing viewport camera
+	 * @param	ViewportConfigKey	the viewport you wish to get from
+	 * @return	Whether we were able to get the FOV for a level editing viewport
+	 */
+	UFUNCTION(BlueprintPure, Category = "Editor Scripting | Level Utility", meta=(DisplayName="Get Level Viewport Field Of View", DevelopmentOnly))
+	UE_API bool GetLevelViewportFOV(float& FOV, FName ViewportConfigKey);
+
+	/**
+	* Sets the horizontal field of view for the given level editor viewport.
+	*
+	* @param	FOV					Field of View to set the camera to. Value will be clamped between 5 and 170 deg
+	* @param	ViewportConfigKey	the viewport you wish to target
+	*/
+	UFUNCTION(BlueprintCallable, Category = "Editor Scripting | Level Utility", meta=(DisplayName="Set Level Viewport Field Of View", DevelopmentOnly))
+	UE_API  void SetLevelViewportFOV(float FOV, FName ViewportConfigKey);
+	
 	/**
 	 * Close the current Persistent Level (without saving it). Create a new blank Level and save it. Load the new created level.
 	 * @param	AssetPath		Asset Path of where the level will be saved.
@@ -201,6 +242,7 @@ public:
 	/** Expose MapChanged to blueprints. Note: This executes too early for some editor scripting, consider using OnMapOpened if this doesn't work for you. */
 	UPROPERTY(BlueprintAssignable, Category = "Editor Scripting | Level Utility")
 	FOnLevelEditorMapChanged OnMapChanged;
+	
 	/** Expose MapOpened to blueprints */
 	UPROPERTY(BlueprintAssignable, Category = "Editor Scripting | Level Utility")
 	FOnLevelEditorMapOpened OnMapOpened;
