[Selection] Fix a couple of crashes with ISM selection, and improve the traversal of actor, component, ISM with double clicks

#Synced-CL 50542684

diff --git a/Engine/Source/Editor/LevelEditor/Private/Elements/SMInstance/SMInstanceElementLevelEditorSelectionCustomization.cpp b/Engine/Source/Editor/LevelEditor/Private/Elements/SMInstance/SMInstanceElementLevelEditorSelectionCustomization.cpp
--- a/Engine/Source/Editor/LevelEditor/Private/Elements/SMInstance/SMInstanceElementLevelEditorSelectionCustomization.cpp
+++ b/Engine/Source/Editor/LevelEditor/Private/Elements/SMInstance/SMInstanceElementLevelEditorSelectionCustomization.cpp
@@ -98,30 +98,70 @@ FTypedElementHandle FSMInstanceElementLevelEditorSelectionCustomization::GetSele
 		const bool bComponentAlreadySelected = InCurrentSelection->Contains(OwningComponentHandle);
 		const bool bIsISMAlreadySelected = InCurrentSelection->Contains(InElementSelectionHandle);
 
-		bool bIsSiblingSelected = false;
-		if (InCurrentSelection->HasElementsOfType(InElementSelectionHandle.GetId().GetTypeId()))
+		// Do we need to go down or up the hierarchy from a double click?
+		if (bWasDoubleClick && !bCameFromSecondary)
 		{
-			const FSMInstanceManager SelectedSMInstance = SMInstanceElementDataUtil::GetSMInstanceFromHandleChecked(InCurrentSelection->GetTopElement<ITypedElementHierarchyInterface>());
-			if (SelectedSMInstance && (SelectedSMInstance.GetISMComponent() == Component))
+			// down
+			if (bComponentAlreadySelected)
 			{
-				bIsSiblingSelected = true;
+				return InElementSelectionHandle;
+			}
+			// up
+			else if (InCurrentSelection->HasElementsOfType(InElementSelectionHandle.GetId().GetTypeId()))
+			{
+				bCameFromSecondary = true;
+				return FComponentElementLevelEditorSelectionCustomization::GetSelectionElementStatic(OwningComponentHandle, InCurrentSelection, ETypedElementSelectionMethod::FromSecondary);
 			}
 		}
 
-		// If a single click, and the ISM was already selected, defer to the component selection, so that we traverse the hierarchy one level instead of picking the top level selection again if we're hit on the same selection.
-		if (!bWasDoubleClick && (bIsSiblingSelected || bIsISMAlreadySelected))
-		{
-			return FComponentElementLevelEditorSelectionCustomization::GetSelectionElementStatic(OwningComponentHandle, InCurrentSelection, ETypedElementSelectionMethod::FromSecondary);
-		}
-		
-		// consider the component already selected if it is in the selection, or if the ISM handle that was hit is a sibling to the current handle
-		if (bWasDoubleClick && (bIsSiblingSelected || bComponentAlreadySelected))
+		// If the active selection is within the same component
+		else if (HasSMInstanceSelectionInSameComponent(InElementSelectionHandle, InCurrentSelection))
 		{
 			return InElementSelectionHandle;
 		}
 
+		// Otherwise forward selection to the component to component handle
+		if (InSelectionMethod == ETypedElementSelectionMethod::Secondary)
+		{
+			bCameFromSecondary = false;
+		}
 		return FComponentElementLevelEditorSelectionCustomization::GetSelectionElementStatic(OwningComponentHandle, InCurrentSelection, InSelectionMethod);
 	}
 
 	return InElementSelectionHandle;
 }
+
+bool FSMInstanceElementLevelEditorSelectionCustomization::AllowSelectionModifiers(const TTypedElement<ITypedElementSelectionInterface>& InElementSelectionHandle, FTypedElementListConstRef InSelectionSet)
+{
+	return HasSMInstanceSelectionInSameComponent(InElementSelectionHandle, InSelectionSet);
+}
+
+bool FSMInstanceElementLevelEditorSelectionCustomization::HasSMInstanceSelectionInSameComponent(const TTypedElement<ITypedElementSelectionInterface>& InElementSelectionHandle, FTypedElementListConstRef InSelectionSet)
+{
+	const FSMInstanceManager SMInstance = SMInstanceElementDataUtil::GetSMInstanceFromHandle(InElementSelectionHandle);
+	if (!SMInstance)
+	{
+		return false;
+	}
+
+	bool bHasSMInstanceSelectedWithSameComponent = false;
+	if (InSelectionSet->HasElementsOfType(FSMInstanceElementData::StaticTypeId()))
+	{
+		InSelectionSet->ForEachElementHandle([InElementSelectionHandle, SMInstance, &bHasSMInstanceSelectedWithSameComponent](const FTypedElementHandle& InHandle)
+		{
+			if (InHandle.GetId().GetTypeId() == FSMInstanceElementData::StaticTypeId())
+			{
+				const FSMInstanceManager SelectedSMInstance = SMInstanceElementDataUtil::GetSMInstanceFromHandle(InHandle);
+				if (SelectedSMInstance && (SelectedSMInstance.GetISMComponent() == SMInstance.GetISMComponent()))
+				{
+					bHasSMInstanceSelectedWithSameComponent = true;
+				}
+				return false;
+			}
+
+			return true;
+		});
+	}
+
+	return bHasSMInstanceSelectedWithSameComponent;
+}
diff --git a/Engine/Source/Editor/LevelEditor/Private/Elements/SMInstance/SMInstanceElementLevelEditorSelectionCustomization.h b/Engine/Source/Editor/LevelEditor/Private/Elements/SMInstance/SMInstanceElementLevelEditorSelectionCustomization.h
--- a/Engine/Source/Editor/LevelEditor/Private/Elements/SMInstance/SMInstanceElementLevelEditorSelectionCustomization.h
+++ b/Engine/Source/Editor/LevelEditor/Private/Elements/SMInstance/SMInstanceElementLevelEditorSelectionCustomization.h
@@ -13,4 +13,8 @@ public:
 	virtual bool SelectElement(const TTypedElement<ITypedElementSelectionInterface>& InElementSelectionHandle, FTypedElementListRef InSelectionSet, const FTypedElementSelectionOptions& InSelectionOptions) override;
 	virtual bool DeselectElement(const TTypedElement<ITypedElementSelectionInterface>& InElementSelectionHandle, FTypedElementListRef InSelectionSet, const FTypedElementSelectionOptions& InSelectionOptions) override;
 	virtual FTypedElementHandle GetSelectionElement(const TTypedElement<ITypedElementSelectionInterface>& InElementSelectionHandle, FTypedElementListConstRef InCurrentSelection, const ETypedElementSelectionMethod InSelectionMethod) override;
+	virtual bool AllowSelectionModifiers(const TTypedElement<ITypedElementSelectionInterface>& InElementSelectionHandle, FTypedElementListConstRef InSelectionSet) override;
+private:
+	static bool HasSMInstanceSelectionInSameComponent(const TTypedElement<ITypedElementSelectionInterface>& InElementSelectionHandle, FTypedElementListConstRef InSelectionSet);
+	bool bCameFromSecondary = false;
 };
