Add Flip Aspect Ratio for Preview Json

#Synced-CL 50766215

diff --git a/Engine/Source/Editor/EditorStyle/Private/StarshipStyle.cpp b/Engine/Source/Editor/EditorStyle/Private/StarshipStyle.cpp
--- a/Engine/Source/Editor/EditorStyle/Private/StarshipStyle.cpp
+++ b/Engine/Source/Editor/EditorStyle/Private/StarshipStyle.cpp
@@ -5043,6 +5043,7 @@ void FStarshipEditorStyle::FStyle::SetupLevelEditorStyle()
 		Set("LevelViewport.HighResScreenshot", new IMAGE_BRUSH_SVG("Starship/Common/HighResolutionScreenshot", Icon16x16));
 		Set("LevelViewport.ToggleCinematicPreview", new IMAGE_BRUSH_SVG("Starship/Common/Cinematics", Icon16x16));
 		Set("LevelViewport.ToggleAllowConstrainedAspectRatioInPreview", new IMAGE_BRUSH_SVG("Starship/Common/ConstrainedAspectRatio_16", Icon16x16));
+		Set("LevelViewport.ToggleFlipConstrainedAspectRatioInPreview", new IMAGE_BRUSH_SVG("Starship/Common/RotationScreenspace_16", Icon16x16));
 		Set("LevelViewport.ToggleCameraShakePreview", new IMAGE_BRUSH_SVG("Starship/Common/CameraShake", Icon16x16));
 		Set("LevelViewport.AdvancedSettings", new CORE_IMAGE_BRUSH_SVG("Starship/Common/settings", Icon16x16));
 		Set("LevelViewport.PlaySettings", new CORE_IMAGE_BRUSH_SVG("Starship/Common/settings", Icon16x16));
diff --git a/Engine/Source/Editor/LevelEditor/Private/LevelViewportActions.cpp b/Engine/Source/Editor/LevelEditor/Private/LevelViewportActions.cpp
--- a/Engine/Source/Editor/LevelEditor/Private/LevelViewportActions.cpp
+++ b/Engine/Source/Editor/LevelEditor/Private/LevelViewportActions.cpp
@@ -89,6 +89,7 @@ void FLevelViewportCommands::RegisterCommands()
 
 	UI_COMMAND( ToggleCinematicPreview, "Allow Cinematic Control", "If enabled, allows cinematic (Sequencer) previews to play in this viewport", EUserInterfaceActionType::ToggleButton, FInputChord() );
 	UI_COMMAND( ToggleAllowConstrainedAspectRatioInPreview, "Constrained Aspect Ratio Preview", "Allows this viewport to have a Constrained Aspect Ratio in Preview", EUserInterfaceActionType::ToggleButton, FInputChord() );
+	UI_COMMAND( ToggleFlipConstrainedAspectRatioInPreview, "Flip Constrained Aspect Ratio Preview", "Flip the Constrained Aspect Ratio in Preview", EUserInterfaceActionType::ToggleButton, FInputChord() );
 
 	UI_COMMAND( FindInLevelScriptBlueprint, "Find In Level Script", "Finds references of a selected actor in the level script blueprint", EUserInterfaceActionType::Button, FInputChord(EModifierKey::Control, EKeys::K) );
 	UI_COMMAND(
diff --git a/Engine/Source/Editor/LevelEditor/Private/SLevelViewport.cpp b/Engine/Source/Editor/LevelEditor/Private/SLevelViewport.cpp
--- a/Engine/Source/Editor/LevelEditor/Private/SLevelViewport.cpp
+++ b/Engine/Source/Editor/LevelEditor/Private/SLevelViewport.cpp
@@ -1639,6 +1639,14 @@ void SLevelViewport::BindOptionCommands( FUICommandList& OutCommandList )
 		FIsActionButtonVisible::CreateLambda([]() { return GEditor->IsFeatureLevelPreviewActive() && (GEditor->PreviewPlatform.GetConstrainedAspectRatio() != 0.0f); })
 	);
 
+	OutCommandList.MapAction(
+		ViewportActions.ToggleFlipConstrainedAspectRatioInPreview,
+		FExecuteAction::CreateLambda([]() { GEditor->ToggleFlipConstrainedAspectRatioInPreview(); }),
+		FCanExecuteAction::CreateLambda([]() { return GEditor->GetAllowConstrainedAspectRatioInPreview(); }),
+		FIsActionChecked::CreateLambda([]() { return GEditor->GetFlipConstrainedAspectRatioInPreview(); }),
+		FIsActionButtonVisible::CreateLambda([]() { return GEditor->IsFeatureLevelPreviewActive() && (GEditor->PreviewPlatform.GetConstrainedAspectRatio() != 0.0f); })
+	);
+
 	IAssetRegistry & AssetRegistry = FModuleManager::LoadModuleChecked<FAssetRegistryModule>("AssetRegistry").Get();
 	TArray<FTopLevelAssetPath> ClassNames;
 	TSet<FTopLevelAssetPath> DerivedClassNames;
diff --git a/Engine/Source/Editor/LevelEditor/Private/ViewportToolbar/LevelEditorSubmenus.cpp b/Engine/Source/Editor/LevelEditor/Private/ViewportToolbar/LevelEditorSubmenus.cpp
--- a/Engine/Source/Editor/LevelEditor/Private/ViewportToolbar/LevelEditorSubmenus.cpp
+++ b/Engine/Source/Editor/LevelEditor/Private/ViewportToolbar/LevelEditorSubmenus.cpp
@@ -2433,8 +2433,10 @@ FToolMenuEntry CreatePerformanceAndScalabilitySubmenu()
 						UnnamedSection.AddEntry(UE::UnrealEd::CreateRemoveRealtimeOverrideEntry(Context->Viewport));
 					}
 
-					FToolMenuEntry& MenuEntry = UnnamedSection.AddMenuEntry(FLevelViewportCommands::Get().ToggleAllowConstrainedAspectRatioInPreview);
-					MenuEntry.SetShowInToolbarTopLevel(true);
+					FToolMenuEntry& MenuEntryAllowConstrainedAspectRatio = UnnamedSection.AddMenuEntry(FLevelViewportCommands::Get().ToggleAllowConstrainedAspectRatioInPreview);
+					MenuEntryAllowConstrainedAspectRatio.SetShowInToolbarTopLevel(true);
+					FToolMenuEntry& MenuEntryFlipConstrainedAspectRatio = UnnamedSection.AddMenuEntry(FLevelViewportCommands::Get().ToggleFlipConstrainedAspectRatioInPreview);
+					MenuEntryFlipConstrainedAspectRatio.SetShowInToolbarTopLevel(true);
 				}
 
 				{
@@ -2474,9 +2476,11 @@ FToolMenuEntry CreatePIEPerformanceAndScalabilitySubmenu()
 				{
 					FToolMenuSection& UnnamedSection = Submenu->FindOrAddSection(NAME_None);
 
-					FToolMenuEntry& MenuEntry = UnnamedSection.AddMenuEntry(FLevelViewportCommands::Get().ToggleAllowConstrainedAspectRatioInPreview);
-					MenuEntry.SetShowInToolbarTopLevel(true);
-					
+					FToolMenuEntry& MenuEntryAllowConstrainedAspectRatio = UnnamedSection.AddMenuEntry(FLevelViewportCommands::Get().ToggleAllowConstrainedAspectRatioInPreview);
+					MenuEntryAllowConstrainedAspectRatio.SetShowInToolbarTopLevel(true);
+					FToolMenuEntry& MenuEntryAllowFlipConstrainedAspectRatio = UnnamedSection.AddMenuEntry(FLevelViewportCommands::Get().ToggleFlipConstrainedAspectRatioInPreview);
+					MenuEntryAllowFlipConstrainedAspectRatio.SetShowInToolbarTopLevel(true);
+
 					UnnamedSection.AddEntry(UE::UnrealEd::CreateResetScalabilitySubmenu());
 				}
 
diff --git a/Engine/Source/Editor/LevelEditor/Public/LevelViewportActions.h b/Engine/Source/Editor/LevelEditor/Public/LevelViewportActions.h
--- a/Engine/Source/Editor/LevelEditor/Public/LevelViewportActions.h
+++ b/Engine/Source/Editor/LevelEditor/Public/LevelViewportActions.h
@@ -115,6 +115,9 @@ PRAGMA_ENABLE_DEPRECATION_WARNINGS
 	/** Allows this viewport to have a Constrained Aspect Ratio in Preview */
 	TSharedPtr< FUICommandInfo > ToggleAllowConstrainedAspectRatioInPreview;
 
+	/** Allows this viewport to flip the Constrained Aspect Ratio in Preview */
+	TSharedPtr< FUICommandInfo > ToggleFlipConstrainedAspectRatioInPreview;
+
 	/** Finds instances of selected object in level script. */ 
 	TSharedPtr< FUICommandInfo > FindInLevelScriptBlueprint;
 
diff --git a/Engine/Source/Editor/UnrealEd/Classes/Editor/EditorEngine.h b/Engine/Source/Editor/UnrealEd/Classes/Editor/EditorEngine.h
--- a/Engine/Source/Editor/UnrealEd/Classes/Editor/EditorEngine.h
+++ b/Engine/Source/Editor/UnrealEd/Classes/Editor/EditorEngine.h
@@ -2337,6 +2337,16 @@ public:
 	 */
 	UNREALED_API bool GetAllowConstrainedAspectRatioInPreview();
 
+	/**
+	 * Toggle on/off the Editor Flip of the Constrained Aspect Ratio in the preview.
+	 */
+	UNREALED_API void ToggleFlipConstrainedAspectRatioInPreview();
+
+	/**
+	 * @return  returns if the Editor Flips the Constrained Aspect Ratio in the preview
+	 */
+	UNREALED_API bool GetFlipConstrainedAspectRatioInPreview();
+
 	/** 
 	 * Accesses the array of snap grid position options
 	 */
@@ -2886,6 +2896,9 @@ private:
 	/** True if we allow Constrained Aspect Ratio in the Preview */
 	bool bAllowConstrainedAspectRatioInPreview;
 
+	/** True if we Flip the Constrained Aspect Ratio in the Preview */
+	bool bFlipConstrainedAspectRatioInPreview;
+
 	/** List of editors who want to receive undo/redo events */
 	TSet< class FEditorUndoClient* > UndoClients;
 
diff --git a/Engine/Source/Editor/UnrealEd/Private/EditorConstraints.cpp b/Engine/Source/Editor/UnrealEd/Private/EditorConstraints.cpp
--- a/Engine/Source/Editor/UnrealEd/Private/EditorConstraints.cpp
+++ b/Engine/Source/Editor/UnrealEd/Private/EditorConstraints.cpp
@@ -75,6 +75,17 @@ bool UEditorEngine::GetAllowConstrainedAspectRatioInPreview()
 	return bAllowConstrainedAspectRatioInPreview;
 }
 
+void UEditorEngine::ToggleFlipConstrainedAspectRatioInPreview()
+{
+	bFlipConstrainedAspectRatioInPreview ^= true;
+	FEditorSupportDelegates::RedrawAllViewports.Broadcast();
+}
+
+bool UEditorEngine::GetFlipConstrainedAspectRatioInPreview()
+{
+	return bFlipConstrainedAspectRatioInPreview;
+}
+
 const TArray<float>& UEditorEngine::GetCurrentPositionGridArray() const
 {
 	const ULevelEditorViewportSettings* ViewportSettings = GetDefault<ULevelEditorViewportSettings>();
diff --git a/Engine/Source/Editor/UnrealEd/Private/EditorViewportClient.cpp b/Engine/Source/Editor/UnrealEd/Private/EditorViewportClient.cpp
--- a/Engine/Source/Editor/UnrealEd/Private/EditorViewportClient.cpp
+++ b/Engine/Source/Editor/UnrealEd/Private/EditorViewportClient.cpp
@@ -1203,7 +1203,12 @@ FSceneView* FEditorViewportClient::CalcSceneView(FSceneViewFamily* ViewFamily, c
 	bool bConstrainAspectRatio = bUseControllingActorViewInfo && ControllingActorViewInfo.bConstrainAspectRatio;
 	if (GEditor != nullptr && GEditor->IsFeatureLevelPreviewActive() && GEditor->GetAllowConstrainedAspectRatioInPreview())
 	{
-		FSlateApplication::Get().OnConstrainedAspectRatioChanged.Broadcast(GEditor->PreviewPlatform.GetConstrainedAspectRatio());
+		float PlatformAspectRatio = GEditor->PreviewPlatform.GetConstrainedAspectRatio();
+		if (GEditor->GetFlipConstrainedAspectRatioInPreview() && PlatformAspectRatio != 0.f)
+		{
+			PlatformAspectRatio = 1 / PlatformAspectRatio;
+		}
+		FSlateApplication::Get().OnConstrainedAspectRatioChanged.Broadcast(PlatformAspectRatio);
 	}
 	else
 	{
diff --git a/Engine/Source/Runtime/Engine/Private/LocalPlayer.cpp b/Engine/Source/Runtime/Engine/Private/LocalPlayer.cpp
--- a/Engine/Source/Runtime/Engine/Private/LocalPlayer.cpp
+++ b/Engine/Source/Runtime/Engine/Private/LocalPlayer.cpp
@@ -1247,7 +1247,13 @@ bool ULocalPlayer::GetProjectionData(FViewport* Viewport, FSceneViewProjectionDa
 	{
 		if (GEditor != nullptr && GEditor->IsFeatureLevelPreviewActive() && GEditor->GetAllowConstrainedAspectRatioInPreview())
 		{
-			FSlateApplication::Get().OnConstrainedAspectRatioChanged.Broadcast(GEditor->PreviewPlatform.GetConstrainedAspectRatio());
+			float PlatformAspectRatio = GEditor->PreviewPlatform.GetConstrainedAspectRatio();
+			if (GEditor->GetFlipConstrainedAspectRatioInPreview() && PlatformAspectRatio != 0.f)
+			{
+				PlatformAspectRatio = 1 / PlatformAspectRatio;
+			}
+
+			FSlateApplication::Get().OnConstrainedAspectRatioChanged.Broadcast(PlatformAspectRatio);
 
 			FMargin Safezones(GEditor->PreviewPlatform.GetSafeZones());
 			if (Safezones != FSlateApplication::Get().GetCustomSafeZone())
