Add Preview json Versioning

#Synced-CL 50957382

diff --git a/Engine/Plugins/Runtime/AndroidDeviceProfileSelector/Source/AndroidDeviceProfileSelector/Private/AndroidDeviceProfileSelectorModule.cpp b/Engine/Plugins/Runtime/AndroidDeviceProfileSelector/Source/AndroidDeviceProfileSelector/Private/AndroidDeviceProfileSelectorModule.cpp
--- a/Engine/Plugins/Runtime/AndroidDeviceProfileSelector/Source/AndroidDeviceProfileSelector/Private/AndroidDeviceProfileSelectorModule.cpp
+++ b/Engine/Plugins/Runtime/AndroidDeviceProfileSelector/Source/AndroidDeviceProfileSelector/Private/AndroidDeviceProfileSelectorModule.cpp
@@ -97,8 +97,15 @@ bool FAndroidDeviceProfileSelectorModule::CanExportDeviceParametersToJson()
 #endif
 }
 
-void FAndroidDeviceProfileSelectorModule::GetDeviceParametersFromJson(FString& JsonLocation, TMap<FString, FString>& OutDeviceParameters)
+static bool GetDeviceSpecsFromJson(const FString& JsonLocation, FPIEPreviewDeviceSpecifications& OutDeviceSpecs)
 {
+	static TMap<const FString, FPIEPreviewDeviceSpecifications> FileLocationToDeviceSpec;
+	if (FPIEPreviewDeviceSpecifications* CachedDeviceSpecs = FileLocationToDeviceSpec.Find(JsonLocation))
+	{
+		OutDeviceSpecs = *CachedDeviceSpecs;
+		return true;
+	}
+
 	TSharedPtr<FJsonObject> JsonRootObject;
 	FString Json;
 
@@ -108,9 +115,59 @@ void FAndroidDeviceProfileSelectorModule::GetDeviceParametersFromJson(FString& J
 		FJsonSerializer::Deserialize(JsonReader, JsonRootObject);
 	}
 
+	if (JsonRootObject.IsValid() && FJsonObjectConverter::JsonAttributesToUStruct(JsonRootObject->Values, FPIEPreviewDeviceSpecifications::StaticStruct(), &OutDeviceSpecs, 0, 0))
+	{
+		if (OutDeviceSpecs.AndroidProperties.Version == ANDROID_DEVICE_PROPERTIES_VERSION)
+		{
+			FileLocationToDeviceSpec.Add(JsonLocation, OutDeviceSpecs);
+		}
+		return true;
+	}
+
+	return false;
+}
+
+bool FAndroidDeviceProfileSelectorModule::IsJsonCompatible(const FString& JsonLocation)
+{
+	FPIEPreviewDeviceSpecifications DeviceSpecs;
+	if (GetDeviceSpecsFromJson(JsonLocation, DeviceSpecs))
+	{
+		if (DeviceSpecs.AndroidProperties.Version > ANDROID_DEVICE_PROPERTIES_VERSION)
+		{
+			UE_LOG(LogAndroidDPSelector, Log, TEXT("Json %s was saved with a newer version than this Engine Version supports(Found v%d, Expected v%d), please update your Json by connecting the same device that generated this Json and call Generate Preview Json."),
+				*JsonLocation, DeviceSpecs.AndroidProperties.Version, ANDROID_DEVICE_PROPERTIES_VERSION);
+			return false;
+		}
+		else if (DeviceSpecs.DevicePlatform != EPIEPreviewDeviceType::Android)
+		{
+			UE_LOG(LogAndroidDPSelector, Log, TEXT("Json %s was saved with a different DevicePlatform than Android. "), *JsonLocation);
+			return false;
+		}
+	}
+	else
+	{
+		return false;
+	}
+
+	return true;
+}
+
+static void MakeOutdatedDeviceSpecCompatible(FPIEPreviewDeviceSpecifications& OutDeviceSpecs)
+{
+
+}
+
+void FAndroidDeviceProfileSelectorModule::GetDeviceParametersFromJson(FString& JsonLocation, TMap<FString, FString>& OutDeviceParameters)
+{
 	FPIEPreviewDeviceSpecifications DeviceSpecs;
-	if (JsonRootObject.IsValid() && FJsonObjectConverter::JsonAttributesToUStruct(JsonRootObject->Values, FPIEPreviewDeviceSpecifications::StaticStruct(), &DeviceSpecs, 0, 0))
+	if (GetDeviceSpecsFromJson(JsonLocation, DeviceSpecs))
 	{
+		if (DeviceSpecs.AndroidProperties.Version < ANDROID_DEVICE_PROPERTIES_VERSION)
+		{
+			MakeOutdatedDeviceSpecCompatible(DeviceSpecs);
+			UE_LOG(LogAndroidDPSelector, Log, TEXT("Json %s is using an Outdated Json Version(Found v%d, Expected v%d), please update your Json by connecting the same device that generated this Json and call Generate Preview Json."),
+				*JsonLocation, DeviceSpecs.AndroidProperties.Version, ANDROID_DEVICE_PROPERTIES_VERSION);
+		}
 		FPIEAndroidDeviceProperties& AndroidProperties = DeviceSpecs.AndroidProperties;
 		OutDeviceParameters.Add(TEXT("SRC_GPUFamily"), AndroidProperties.GPUFamily);
 		OutDeviceParameters.Add(TEXT("SRC_GLVersion"), AndroidProperties.GLVersion);
@@ -134,6 +191,11 @@ void FAndroidDeviceProfileSelectorModule::GetDeviceParametersFromJson(FString& J
 		OutDeviceParameters.Add(FString(FAndroidProfileSelectorSourceProperties::SRC_InsetsRight), FString::Printf(TEXT("%f"), DeviceSpecs.InsetsRight));
 		OutDeviceParameters.Add(FString(FAndroidProfileSelectorSourceProperties::SRC_InsetsBottom), FString::Printf(TEXT("%f"), DeviceSpecs.InsetsBottom));
 	}
+	else
+	{
+		UE_LOG(LogAndroidDPSelector, Log, TEXT("Json %s failed to load"), *JsonLocation);
+		return;
+	}
 
 	// Initialize the ConfigRules for the Preview Device.
 	FString ConfigRulesFile = FPaths::ProjectDir() + TEXT("Build/Android/configrules.txt");
diff --git a/Engine/Plugins/Runtime/AndroidDeviceProfileSelector/Source/AndroidDeviceProfileSelector/Private/AndroidDeviceProfileSelectorModule.h b/Engine/Plugins/Runtime/AndroidDeviceProfileSelector/Source/AndroidDeviceProfileSelector/Private/AndroidDeviceProfileSelectorModule.h
--- a/Engine/Plugins/Runtime/AndroidDeviceProfileSelector/Source/AndroidDeviceProfileSelector/Private/AndroidDeviceProfileSelectorModule.h
+++ b/Engine/Plugins/Runtime/AndroidDeviceProfileSelector/Source/AndroidDeviceProfileSelector/Private/AndroidDeviceProfileSelectorModule.h
@@ -19,6 +19,7 @@ public:
 #if WITH_EDITOR
 	virtual void ExportDeviceParametersToJson(FString& FolderLocation, TArray<FString>& OutExportedFiles) override;
 	virtual bool CanExportDeviceParametersToJson() override;
+	virtual bool IsJsonCompatible(const FString& JsonLocation) override;
 	virtual void GetDeviceParametersFromJson(FString& JsonLocation, TMap<FString, FString>& OutDeviceParameters) override;
 	virtual bool CanGetDeviceParametersFromJson() override { return true; }
 	virtual const EShaderPlatform GetPreviewShaderPlatform() override;
diff --git a/Engine/Plugins/Runtime/IOSDeviceProfileSelector/Source/IOSPreviewDeviceProfileSelector/Private/IOSPreviewDeviceProfileSelectorModule.cpp b/Engine/Plugins/Runtime/IOSDeviceProfileSelector/Source/IOSPreviewDeviceProfileSelector/Private/IOSPreviewDeviceProfileSelectorModule.cpp
--- a/Engine/Plugins/Runtime/IOSDeviceProfileSelector/Source/IOSPreviewDeviceProfileSelector/Private/IOSPreviewDeviceProfileSelectorModule.cpp
+++ b/Engine/Plugins/Runtime/IOSDeviceProfileSelector/Source/IOSPreviewDeviceProfileSelector/Private/IOSPreviewDeviceProfileSelectorModule.cpp
@@ -14,6 +14,7 @@
 #include "IOSDeviceHelper.h"
 #endif
 
+DEFINE_LOG_CATEGORY_STATIC(LogIOSDPSelector, Log, All)
 #define IOS_DEFAULT_DEVICE_PROFILE_NAME TEXT("IOS")
 
 IMPLEMENT_MODULE(FIOSPreviewDeviceProfileSelectorModule, IOSPreviewDeviceProfileSelector);
@@ -97,12 +98,68 @@ bool FIOSPreviewDeviceProfileSelectorModule::CanExportDeviceParametersToJson()
 #endif
 };
 
+static bool GetDeviceSpecsFromJson(const FString& JsonLocation, FPIEPreviewDeviceSpecifications& OutDeviceSpecs)
+{
+	static TMap<const FString, FPIEPreviewDeviceSpecifications> FileLocationToDeviceSpec;
+	if (FPIEPreviewDeviceSpecifications* CachedDeviceSpecs = FileLocationToDeviceSpec.Find(JsonLocation))
+	{
+		OutDeviceSpecs = *CachedDeviceSpecs;
+		return true;
+	}
+
+	TSharedPtr<FJsonObject> JsonRootObject;
+	FString Json;
+
+	if (FFileHelper::LoadFileToString(Json, *JsonLocation))
+	{
+		TSharedRef<TJsonReader<> > JsonReader = TJsonReaderFactory<>::Create(Json);
+		FJsonSerializer::Deserialize(JsonReader, JsonRootObject);
+	}
+
+	if (JsonRootObject.IsValid() && FJsonObjectConverter::JsonAttributesToUStruct(JsonRootObject->Values, FPIEPreviewDeviceSpecifications::StaticStruct(), &OutDeviceSpecs, 0, 0))
+	{
+		if (OutDeviceSpecs.IOSProperties.Version == IOS_DEVICE_PROPERTIES_VERSION)
+		{
+			FileLocationToDeviceSpec.Add(JsonLocation, OutDeviceSpecs);
+		}
+		return true;
+	}
+
+	return false;
+}
+
+bool FIOSPreviewDeviceProfileSelectorModule::IsJsonCompatible(const FString& JsonLocation)
+{
+	FPIEPreviewDeviceSpecifications DeviceSpecs;
+	if (GetDeviceSpecsFromJson(JsonLocation, DeviceSpecs))
+	{
+		if (DeviceSpecs.IOSProperties.Version > IOS_DEVICE_PROPERTIES_VERSION)
+		{
+			UE_LOG(LogIOSDPSelector, Log, TEXT("Json %s was saved with a newer version than this Engine Version supports(Found v%d, Expected v%d), please update your Json by connecting the same device that generated this Json and call Generate Preview Json."),
+			*JsonLocation, DeviceSpecs.IOSProperties.Version, IOS_DEVICE_PROPERTIES_VERSION);
+			return false;
+		}
+		else if (DeviceSpecs.DevicePlatform != EPIEPreviewDeviceType::IOS)
+		{
+			UE_LOG(LogIOSDPSelector, Log, TEXT("Json %s was saved with a different DevicePlatform than IOS. "), *JsonLocation);
+			return false;
+		}
+	}
+	else
+	{
+		return false;
+	}
+
+	return true;
+}
+
 #if WITH_IOS_DEVICE_DETECTION
 void ExportDeviceProfile(const FString& OutPath, const FString& DeviceName)
 {
 	// instantiate an FPIEPreviewDeviceSpecifications instance and its values
 	FPIEPreviewDeviceSpecifications DeviceSpecs;
 	DeviceSpecs.DevicePlatform = EPIEPreviewDeviceType::IOS;
+	DeviceSpecs.IOSProperties.Version = IOS_DEVICE_PROPERTIES_VERSION;
 	DeviceSpecs.IOSProperties.DeviceModel = DeviceName;
 	// create a JSon object from the above structure
 	TSharedPtr<FJsonObject> JsonObject = FJsonObjectConverter::UStructToJsonObject<FPIEPreviewDeviceSpecifications>(DeviceSpecs);
@@ -140,23 +197,31 @@ void FIOSPreviewDeviceProfileSelectorModule::ExportDeviceParametersToJson(FStrin
 #endif
 }
 
-void FIOSPreviewDeviceProfileSelectorModule::GetDeviceParametersFromJson(FString& JsonLocation, TMap<FString, FString>& OutDeviceParameters)
+static void MakeOutdatedDeviceSpecCompatible(FPIEPreviewDeviceSpecifications& OutDeviceSpecs)
 {
-	TSharedPtr<FJsonObject> JsonRootObject;
-	FString Json;
 
-	if (FFileHelper::LoadFileToString(Json, *JsonLocation))
-	{
-		TSharedRef<TJsonReader<> > JsonReader = TJsonReaderFactory<>::Create(Json);
-		FJsonSerializer::Deserialize(JsonReader, JsonRootObject);
-	}
+}
 
+void FIOSPreviewDeviceProfileSelectorModule::GetDeviceParametersFromJson(FString& JsonLocation, TMap<FString, FString>& OutDeviceParameters)
+{
 	FPIEPreviewDeviceSpecifications DeviceSpecs;
-	if (JsonRootObject.IsValid() && FJsonObjectConverter::JsonAttributesToUStruct(JsonRootObject->Values, FPIEPreviewDeviceSpecifications::StaticStruct(), &DeviceSpecs, 0, 0))
+	if (GetDeviceSpecsFromJson(JsonLocation, DeviceSpecs))
 	{
+		if(DeviceSpecs.IOSProperties.Version < IOS_DEVICE_PROPERTIES_VERSION)
+		{
+			MakeOutdatedDeviceSpecCompatible(DeviceSpecs);
+			UE_LOG(LogIOSDPSelector, Log, TEXT("Json %s is using an Outdated Json Version(Found v%d, Expected v%d), please update your Json by connecting the same device that generated this Json and call Generate Preview Json."), 
+			*JsonLocation, DeviceSpecs.IOSProperties.Version, IOS_DEVICE_PROPERTIES_VERSION);
+		}
+
 		FPIEIOSDeviceProperties& IOSProperties = DeviceSpecs.IOSProperties;
 		OutDeviceParameters.Add((TEXT("SRC_DeviceModel")), IOSProperties.DeviceModel);
 	}
+	else
+	{
+		UE_LOG(LogIOSDPSelector, Log, TEXT("Json %s failed to load"), *JsonLocation);
+		return;
+	}
 }
 
 const EShaderPlatform FIOSPreviewDeviceProfileSelectorModule::GetPreviewShaderPlatform()
diff --git a/Engine/Plugins/Runtime/IOSDeviceProfileSelector/Source/IOSPreviewDeviceProfileSelector/Private/IOSPreviewDeviceProfileSelectorModule.h b/Engine/Plugins/Runtime/IOSDeviceProfileSelector/Source/IOSPreviewDeviceProfileSelector/Private/IOSPreviewDeviceProfileSelectorModule.h
--- a/Engine/Plugins/Runtime/IOSDeviceProfileSelector/Source/IOSPreviewDeviceProfileSelector/Private/IOSPreviewDeviceProfileSelectorModule.h
+++ b/Engine/Plugins/Runtime/IOSDeviceProfileSelector/Source/IOSPreviewDeviceProfileSelector/Private/IOSPreviewDeviceProfileSelectorModule.h
@@ -18,6 +18,7 @@ public:
 	virtual const FString GetRuntimeDeviceProfileName() override;
 #if WITH_EDITOR
 	virtual bool CanExportDeviceParametersToJson() override;
+	virtual bool IsJsonCompatible(const FString& JsonLocation) override;
 	virtual void ExportDeviceParametersToJson(FString& FolderLocation, TArray<FString>& OutExportedFiles) override;
 	virtual void GetDeviceParametersFromJson(FString& JsonLocation, TMap<FString, FString>& OutDeviceParameters) override;
 	virtual bool CanGetDeviceParametersFromJson() override { return true; }
diff --git a/Engine/Source/Developer/Android/AndroidDeviceDetection/Private/AndroidDeviceDetectionModule.cpp b/Engine/Source/Developer/Android/AndroidDeviceDetection/Private/AndroidDeviceDetectionModule.cpp
--- a/Engine/Source/Developer/Android/AndroidDeviceDetection/Private/AndroidDeviceDetectionModule.cpp
+++ b/Engine/Source/Developer/Android/AndroidDeviceDetection/Private/AndroidDeviceDetectionModule.cpp
@@ -1111,6 +1111,7 @@ public:
 			DeviceSpecs.ScaleFactors = { 0.25f, 0.5f, 0.75f, 1.0f };
 
 			// Android specific values
+			DeviceSpecs.AndroidProperties.Version = ANDROID_DEVICE_PROPERTIES_VERSION;
 			DeviceSpecs.AndroidProperties.AndroidVersion = DeviceInfo->HumanAndroidVersion;
 			DeviceSpecs.AndroidProperties.DeviceModel = DeviceInfo->Model;
 			DeviceSpecs.AndroidProperties.DeviceMake = DeviceInfo->DeviceBrand;
diff --git a/Engine/Source/Editor/LevelEditor/Private/LevelEditorActions.cpp b/Engine/Source/Editor/LevelEditor/Private/LevelEditorActions.cpp
--- a/Engine/Source/Editor/LevelEditor/Private/LevelEditorActions.cpp
+++ b/Engine/Source/Editor/LevelEditor/Private/LevelEditorActions.cpp
@@ -587,52 +587,59 @@ void FLevelEditorActionCallbacks::PreviewJson_Clicked(FName PlatformName, FStrin
 		{
 			if (IDeviceProfileSelectorModule* DPSelectorModule = FModuleManager::LoadModulePtr<IDeviceProfileSelectorModule>(*DeviceProfileSelectionModule))
 			{
-				FString DevileProfileName;
+				if(DPSelectorModule->IsJsonCompatible(JsonFile))
+				{
+					FString DevileProfileName;
 
-				TMap<FString, FString> DeviceParameters;
-				DPSelectorModule->GetDeviceParametersFromJson(JsonFile, DeviceParameters);
-				DPSelectorModule->SetSelectorProperties(DeviceParameters);
-				float ConstrainedAspectRatio = DPSelectorModule->GetConstrainedAspectRatio();
+					TMap<FString, FString> DeviceParameters;
+					DPSelectorModule->GetDeviceParametersFromJson(JsonFile, DeviceParameters);
+					DPSelectorModule->SetSelectorProperties(DeviceParameters);
+					float ConstrainedAspectRatio = DPSelectorModule->GetConstrainedAspectRatio();
 
-				DevileProfileName = DPSelectorModule->GetDeviceProfileName();
-				EShaderPlatform PreviewShaderPlatform = DPSelectorModule->GetPreviewShaderPlatform();
+					DevileProfileName = DPSelectorModule->GetDeviceProfileName();
+					EShaderPlatform PreviewShaderPlatform = DPSelectorModule->GetPreviewShaderPlatform();
 
-				// Apply the PreviewPlatform returned by the DPSelectorModule.
-				if (FDataDrivenShaderPlatformInfo::IsValid(PreviewShaderPlatform) && FDataDrivenShaderPlatformInfo::GetIsPreviewPlatform(PreviewShaderPlatform))
-				{
-					FName PreviewShaderPlatformName = FDataDrivenShaderPlatformInfo::GetName(PreviewShaderPlatform);
+					// Apply the PreviewPlatform returned by the DPSelectorModule.
+					if (FDataDrivenShaderPlatformInfo::IsValid(PreviewShaderPlatform) && FDataDrivenShaderPlatformInfo::GetIsPreviewPlatform(PreviewShaderPlatform))
+					{
+						FName PreviewShaderPlatformName = FDataDrivenShaderPlatformInfo::GetName(PreviewShaderPlatform);
 
-					FName FriendlyDeviceProfilePreviewName = FName(FDataDrivenShaderPlatformInfo::GetFriendlyName(PreviewShaderPlatform).ToString() + TEXT(" ") + DevileProfileName);
+						FName FriendlyDeviceProfilePreviewName = FName(FDataDrivenShaderPlatformInfo::GetFriendlyName(PreviewShaderPlatform).ToString() + TEXT(" ") + DevileProfileName);
 
-					UDeviceProfile* DevileProfile = UDeviceProfileManager::Get().FindProfile(DevileProfileName, false);
-					if (DevileProfile)
-					{
-						auto GetPreviewFeatureLevelInfo = [&]()
-							{
-								const ERHIFeatureLevel::Type FeatureLevel = GetMaxSupportedFeatureLevel(PreviewShaderPlatform);
-								return FPreviewPlatformInfo(
-									FeatureLevel,
-									PreviewShaderPlatform, PlatformName,
-									FDataDrivenShaderPlatformInfo::GetShaderFormat(PreviewShaderPlatform),
-									DevileProfile->GetFName(),
-									true,
-									PreviewShaderPlatformName,
-									FText::FromName(FriendlyDeviceProfilePreviewName),
-									ConstrainedAspectRatio,
-									DPSelectorModule->GetSafeZones());
-							};
-
-						FPreviewPlatformInfo PreviewFeatureLevelInfo = GetPreviewFeatureLevelInfo();
-						GEditor->SetPreviewPlatform(PreviewFeatureLevelInfo, false);
+						UDeviceProfile* DevileProfile = UDeviceProfileManager::Get().FindProfile(DevileProfileName, false);
+						if (DevileProfile)
+						{
+							auto GetPreviewFeatureLevelInfo = [&]()
+								{
+									const ERHIFeatureLevel::Type FeatureLevel = GetMaxSupportedFeatureLevel(PreviewShaderPlatform);
+									return FPreviewPlatformInfo(
+										FeatureLevel,
+										PreviewShaderPlatform, PlatformName,
+										FDataDrivenShaderPlatformInfo::GetShaderFormat(PreviewShaderPlatform),
+										DevileProfile->GetFName(),
+										true,
+										PreviewShaderPlatformName,
+										FText::FromName(FriendlyDeviceProfilePreviewName),
+										ConstrainedAspectRatio,
+										DPSelectorModule->GetSafeZones());
+								};
+
+							FPreviewPlatformInfo PreviewFeatureLevelInfo = GetPreviewFeatureLevelInfo();
+							GEditor->SetPreviewPlatform(PreviewFeatureLevelInfo, false);
+						}
+						else
+						{
+							UE_LOG(LevelEditorActions, Warning, TEXT("JsonFile %s did not match any MatchProfile"), *JsonFile);
+						}
 					}
 					else
 					{
-						UE_LOG(LevelEditorActions, Warning, TEXT("JsonFile %s did not match any MatchProfile"), *JsonFile);
+						UE_LOG(LevelEditorActions, Warning, TEXT("JsonFile %s did not return a valid Preview Platform"), *JsonFile);
 					}
 				}
 				else
 				{
-					UE_LOG(LevelEditorActions, Warning, TEXT("JsonFile %s did not return a valid Preview Platform"), *JsonFile);
+					FMessageDialog::Open(EAppMsgType::Ok, LOCTEXT("PreviewJsonError_NoFiles", "The Json is not compatible with this Engine Version. Please see log for details."), LOCTEXT("PreviewJsonError_Title", "Preview Json"));
 				}
 			}
 		}
@@ -4659,63 +4666,70 @@ void FLevelEditorCommands::RegisterCommands()
 
 		if (PlatformEngineIni && PlatformEngineIni->GetString(TEXT("DeviceProfileManager"), TEXT("PreviewDeviceProfileSelectionModule"), DeviceProfileSelectionModule, GEngineIni))
 		{
-			EShaderPlatform ShaderPlatform = FDataDrivenShaderPlatformInfo::GetShaderPlatformFromName(Item.ShaderPlatformToPreview);
-
-			TArray<PreviewPlatformCommand>* PlatformToPreviewJsonPlatformOverridesValue = PlatformToPreviewJsonPlatformOverrides.Find(Item.PlatformName);
-			if (!PlatformToPreviewJsonPlatformOverridesValue)
+			if (IDeviceProfileSelectorModule* DPSelectorModule = FModuleManager::LoadModulePtr<IDeviceProfileSelectorModule>(*DeviceProfileSelectionModule))
 			{
-				FLevelEditorCommands::PreviewPlatformCommand GenerateJsonPlatform;
-				GenerateJsonPlatform.CommandInfo = FUICommandInfoDecl(
-					this->AsShared(),
-					FName(*FString::Printf(TEXT("Generate Platform Json for %s"), *Item.PlatformName.ToString())),
-					NSLOCTEXT("GeneratePlatformJson", "Generate Platform Json", "Generate Platform Json"),
-					NSLOCTEXT("GeneratePlatformJsonDesc", "Generate Platform Json From Connected Devices", "Generate Platform Json From Connected Devices"))
-					.UserInterfaceType(EUserInterfaceActionType::Button)
-					.DefaultChord(FInputChord());
-				GenerateJsonPlatform.bIsGeneratingJsonCommand = true;
-
-				PlatformToPreviewJsonPlatformOverridesValue = &PlatformToPreviewJsonPlatformOverrides.Add(Item.PlatformName);
-				PlatformToPreviewJsonPlatformOverridesValue->Add(GenerateJsonPlatform);
-			
-				FLevelEditorCommands::PreviewPlatformCommand PreviewJsonPlatform;
-				PreviewJsonPlatform.CommandInfo = FUICommandInfoDecl(
-					this->AsShared(),
-					FName(*FString::Printf(TEXT("Preview %s via Json"), *Item.PlatformName.ToString())),
-					NSLOCTEXT("PreviewviaJson", "Preview via Json", "Preview via Json..."),
-					NSLOCTEXT("PreviewviaJsonDesc", "Preview via Json", "Preview via Json"))
-					.UserInterfaceType(EUserInterfaceActionType::Button)
-					.DefaultChord(FInputChord());
-				PreviewJsonPlatform.bIsGeneratingJsonCommand = false;
-				PreviewJsonPlatform.SectionName = SectionName;
-				PlatformToPreviewJsonPlatformOverridesValue->Add(PreviewJsonPlatform);
-
-				TMap<FString, TArray<FString>> DirectoryToJsonFiles;
-				FString AbsoluteDebugInfoDirectory = IFileManager::Get().ConvertToAbsolutePathForExternalAppForWrite(*(FPaths::ProjectSavedDir() / TEXT("PreviewJsonDevices") / Item.PlatformName.ToString()));
-				IFileManager::Get().FindFiles(DirectoryToJsonFiles.Add(AbsoluteDebugInfoDirectory), *AbsoluteDebugInfoDirectory, TEXT(".json"));
-				FString ProjectEditorJsonDir = FPaths::ProjectContentDir() / TEXT("Editor") / TEXT("PreviewJsonDevices") / Item.PlatformName.ToString();
-				IFileManager::Get().FindFiles(DirectoryToJsonFiles.Add(ProjectEditorJsonDir), *ProjectEditorJsonDir, TEXT(".json"));
-
-				TSet<FString> UniqueJsons;
-				for (auto Iter = DirectoryToJsonFiles.CreateConstIterator(); Iter; ++Iter)
+				EShaderPlatform ShaderPlatform = FDataDrivenShaderPlatformInfo::GetShaderPlatformFromName(Item.ShaderPlatformToPreview);
+
+				TArray<PreviewPlatformCommand>* PlatformToPreviewJsonPlatformOverridesValue = PlatformToPreviewJsonPlatformOverrides.Find(Item.PlatformName);
+				if (!PlatformToPreviewJsonPlatformOverridesValue)
 				{
-					FString DirectoryName = Iter.Key();
-					for (const FString& JsonFile : Iter.Value())
+					FLevelEditorCommands::PreviewPlatformCommand GenerateJsonPlatform;
+					GenerateJsonPlatform.CommandInfo = FUICommandInfoDecl(
+						this->AsShared(),
+						FName(*FString::Printf(TEXT("Generate Platform Json for %s"), *Item.PlatformName.ToString())),
+						NSLOCTEXT("GeneratePlatformJson", "Generate Platform Json", "Generate Platform Json"),
+						NSLOCTEXT("GeneratePlatformJsonDesc", "Generate Platform Json From Connected Devices", "Generate Platform Json From Connected Devices"))
+						.UserInterfaceType(EUserInterfaceActionType::Button)
+						.DefaultChord(FInputChord());
+					GenerateJsonPlatform.bIsGeneratingJsonCommand = true;
+
+					PlatformToPreviewJsonPlatformOverridesValue = &PlatformToPreviewJsonPlatformOverrides.Add(Item.PlatformName);
+					PlatformToPreviewJsonPlatformOverridesValue->Add(GenerateJsonPlatform);
+
+					FLevelEditorCommands::PreviewPlatformCommand PreviewJsonPlatform;
+					PreviewJsonPlatform.CommandInfo = FUICommandInfoDecl(
+						this->AsShared(),
+						FName(*FString::Printf(TEXT("Preview %s via Json"), *Item.PlatformName.ToString())),
+						NSLOCTEXT("PreviewviaJson", "Preview via Json", "Preview via Json..."),
+						NSLOCTEXT("PreviewviaJsonDesc", "Preview via Json", "Preview via Json"))
+						.UserInterfaceType(EUserInterfaceActionType::Button)
+						.DefaultChord(FInputChord());
+					PreviewJsonPlatform.bIsGeneratingJsonCommand = false;
+					PreviewJsonPlatform.SectionName = SectionName;
+					PlatformToPreviewJsonPlatformOverridesValue->Add(PreviewJsonPlatform);
+
+					TMap<FString, TArray<FString>> DirectoryToJsonFiles;
+					FString AbsoluteDebugInfoDirectory = IFileManager::Get().ConvertToAbsolutePathForExternalAppForWrite(*(FPaths::ProjectSavedDir() / TEXT("PreviewJsonDevices") / Item.PlatformName.ToString()));
+					IFileManager::Get().FindFiles(DirectoryToJsonFiles.Add(AbsoluteDebugInfoDirectory), *AbsoluteDebugInfoDirectory, TEXT(".json"));
+					FString ProjectEditorJsonDir = FPaths::ProjectContentDir() / TEXT("Editor") / TEXT("PreviewJsonDevices") / Item.PlatformName.ToString();
+					IFileManager::Get().FindFiles(DirectoryToJsonFiles.Add(ProjectEditorJsonDir), *ProjectEditorJsonDir, TEXT(".json"));
+
+					TSet<FString> UniqueJsons;
+					for (auto Iter = DirectoryToJsonFiles.CreateConstIterator(); Iter; ++Iter)
 					{
-						if (!UniqueJsons.Find(JsonFile))
+						FString DirectoryName = Iter.Key();
+						for (const FString& JsonFile : Iter.Value())
 						{
-							FLevelEditorCommands::PreviewPlatformCommand PreviewJsonFilePlatform;
-							PreviewJsonFilePlatform.CommandInfo = FUICommandInfoDecl(
-								this->AsShared(),
-								FName(*FString::Printf(TEXT("Preview %s with Json"), *JsonFile)),
-								FText::Format(NSLOCTEXT("PreviewJson", "Preview Json", "Preview {0}"), FText::FromString(FPaths::GetBaseFilename(JsonFile))),
-								FText::Format(NSLOCTEXT("PreviewJsonDesc", "Preview using Platform Json", "Preview {0}"), FText::FromString(FPaths::GetBaseFilename(JsonFile))))
-								.UserInterfaceType(EUserInterfaceActionType::Check)
-								.DefaultChord(FInputChord());
-							PreviewJsonFilePlatform.bIsGeneratingJsonCommand = false;
-							PreviewJsonFilePlatform.FilePath = DirectoryName / JsonFile;
-							PreviewJsonFilePlatform.SectionName = SectionName;
-							PlatformToPreviewJsonPlatformOverridesValue->Add(PreviewJsonFilePlatform);
-							UniqueJsons.Add(JsonFile);
+							if (!UniqueJsons.Find(JsonFile))
+							{
+								FString FullFilePath = DirectoryName / JsonFile;
+								if(DPSelectorModule->IsJsonCompatible(FullFilePath))
+								{
+									FLevelEditorCommands::PreviewPlatformCommand PreviewJsonFilePlatform;
+									PreviewJsonFilePlatform.CommandInfo = FUICommandInfoDecl(
+										this->AsShared(),
+										FName(*FString::Printf(TEXT("Preview %s with Json"), *JsonFile)),
+										FText::Format(NSLOCTEXT("PreviewJson", "Preview Json", "Preview {0}"), FText::FromString(FPaths::GetBaseFilename(JsonFile))),
+										FText::Format(NSLOCTEXT("PreviewJsonDesc", "Preview using Platform Json", "Preview {0}"), FText::FromString(FPaths::GetBaseFilename(JsonFile))))
+										.UserInterfaceType(EUserInterfaceActionType::Check)
+										.DefaultChord(FInputChord());
+									PreviewJsonFilePlatform.bIsGeneratingJsonCommand = false;
+									PreviewJsonFilePlatform.FilePath = DirectoryName / JsonFile;
+									PreviewJsonFilePlatform.SectionName = SectionName;
+									PlatformToPreviewJsonPlatformOverridesValue->Add(PreviewJsonFilePlatform);
+								}
+								UniqueJsons.Add(JsonFile);
+							}
 						}
 					}
 				}
diff --git a/Engine/Source/Editor/PIEPreviewDeviceSpecification/Public/PIEPreviewDeviceSpecification.h b/Engine/Source/Editor/PIEPreviewDeviceSpecification/Public/PIEPreviewDeviceSpecification.h
--- a/Engine/Source/Editor/PIEPreviewDeviceSpecification/Public/PIEPreviewDeviceSpecification.h
+++ b/Engine/Source/Editor/PIEPreviewDeviceSpecification/Public/PIEPreviewDeviceSpecification.h
@@ -12,6 +12,11 @@
 
 #include "PIEPreviewDeviceSpecification.generated.h"
 
+// If you modify the logic or structure of the FPIEPreviewDeviceSpecifications
+// Please bump the version for the platform affected and add backwards compatibility to the MakeOutdatedDeviceSpecCompatible function of the platform's Preview Device Profile Selector.
+#define IOS_DEVICE_PROPERTIES_VERSION 0
+#define ANDROID_DEVICE_PROPERTIES_VERSION 0
+
 UENUM()
 enum class EPIEPreviewDeviceType : uint8
 {
@@ -79,7 +84,8 @@ struct FPIEAndroidDeviceProperties
 {
 public:
 	GENERATED_USTRUCT_BODY()
-
+	UPROPERTY()
+	uint32 Version = 0;
 	UPROPERTY()
 	FString GPUFamily;
 	UPROPERTY()
@@ -120,6 +126,9 @@ struct FPIEIOSDeviceProperties
 {
 public:
 	GENERATED_USTRUCT_BODY()
+	
+	UPROPERTY()
+	uint32 Version = 0;
 
 	UPROPERTY()
 	FString DeviceModel;
diff --git a/Engine/Source/Runtime/Engine/Public/IDeviceProfileSelectorModule.h b/Engine/Source/Runtime/Engine/Public/IDeviceProfileSelectorModule.h
--- a/Engine/Source/Runtime/Engine/Public/IDeviceProfileSelectorModule.h
+++ b/Engine/Source/Runtime/Engine/Public/IDeviceProfileSelectorModule.h
@@ -47,6 +47,11 @@ public:
 	 */
 	virtual bool CanExportDeviceParametersToJson() { return false; }
 
+	/**
+	 *  @return If the Json is compatible with this version of the DeviceProfileSelector
+	 */
+	virtual bool IsJsonCompatible(const FString& JsonLocation) { return true; }
+
 	/**
 	 * Get the Device Parameters from the Json
 	* @Param JsonLocation Json location
@@ -56,7 +61,7 @@ public:
 	
 	/**
 	 * Get the ConstrainedAspectRatio of the Device
-	 *  @return ConstrinedAspectRatio of the Device
+	 *  @return ConstrainedAspectRatio of the Device
 	 */
 	virtual float GetConstrainedAspectRatio() { return ConstrainedAspectRatio; }
 	
