Editor: Make some standard actions available in the non-actor element context menu.

#Synced-CL 50868977

diff --git a/Engine/Source/Editor/LevelEditor/Private/LevelEditorContextMenu.cpp b/Engine/Source/Editor/LevelEditor/Private/LevelEditorContextMenu.cpp
--- a/Engine/Source/Editor/LevelEditor/Private/LevelEditorContextMenu.cpp
+++ b/Engine/Source/Editor/LevelEditor/Private/LevelEditorContextMenu.cpp
@@ -637,19 +637,144 @@ void FLevelEditorContextMenu::RegisterElementContextMenu()
 	UToolMenu* Menu = ToolMenus->RegisterMenu("LevelEditor.ElementContextMenu");
 	Menu->AddDynamicSection("ElementContextMenuDynamic", FNewToolMenuDelegate::CreateLambda([](UToolMenu* InMenu)
 	{
+		ULevelEditorContextMenuContext* LevelEditorContext = InMenu->FindContext<ULevelEditorContextMenuContext>();
+		if (!LevelEditorContext || !LevelEditorContext->LevelEditor.IsValid())
 		{
-			FToolMenuSection& Section = InMenu->AddSection("ElementEditActions", LOCTEXT("ElementEditActions", "Edit"));
-			FLevelEditorContextMenuImpl::FillEditMenu(InMenu, &Section);
+			return;
 		}
 
+		// This menu will mostly follow FLevelEditorContextMenu::RegisterActorContextMenu as more support for
+		//  non-actors is added.
+
 		{
-			FToolMenuSection& Section = InMenu->AddSection("ElementLevelActions", LOCTEXT("ElementLevelActions", "Level"));
+			// Options that affect the current viewport.
+			FToolMenuSection& Section = InMenu->AddSection("ElementViewOptions", LOCTEXT("ViewOptionsHeading", "View Options"));
+			const FVector* ClickLocation = &GEditor->ClickLocation;
+
+			Section.AddMenuEntry(
+				FEditorViewportCommands::Get().FocusViewportToSelection,
+				TAttribute<FText>(),
+				TAttribute<FText>(),
+				FSlateIcon(FAppStyle::GetAppStyleSetName(), "Icons.FrameActor")
+			);
+
+			// This keys off the mouse position so can only appear in the viewport
+			if (LevelEditorContext->ContextType == ELevelEditorMenuContext::Viewport)
+			{
+				FUIAction GoHereAction;
+				GoHereAction.ExecuteAction = FExecuteAction::CreateStatic(&FLevelEditorActionCallbacks::GoHere_Clicked, ClickLocation);
+				Section.AddMenuEntry(
+					FLevelEditorCommands::Get().GoHere,
+					TAttribute<FText>(),
+					TAttribute<FText>(),
+					FSlateIcon(FAppStyle::GetAppStyleSetName(), "Actors.GoHere")
+				);
+			}
+
+			Section.AddMenuEntry(
+				FLevelEditorCommands::Get().SnapCameraToObject,
+				TAttribute<FText>(),
+				TAttribute<FText>(),
+				FSlateIcon(FAppStyle::GetAppStyleSetName(), "Actors.SnapViewToObject")
+			);
+
+			Section.AddMenuEntry(
+				FLevelEditorCommands::Get().SnapObjectToCamera,
+				TAttribute<FText>(),
+				TAttribute<FText>(),
+				FSlateIcon(FAppStyle::GetAppStyleSetName(), "Actors.SnapObjectToView")
+			);
+
+			// TODO: Support piloting like we do for actors. This requires changes to viewport code.
+		}
+
+		{
+			// Options for editing, transforming, and manipulating this element
+			FToolMenuSection& Section = InMenu->AddSection("ElementOptions", LOCTEXT("ElementOptionsHeading", "Element Options"));
+
+			// TODO: Duplicate and Rename in this submenu still need implementing. Duplicate requires implementing
+			//  UEntityElementWorldInterface::DuplicateElements, whereas Rename requires support from the outliner,
+			//  as the outliner requires something to be selected in it.
+			Section.AddSubMenu(
+				"EditSubMenu",
+				LOCTEXT("EditSubMenu", "Edit"),
+				FText::GetEmpty(),
+				FNewToolMenuDelegate::CreateStatic(&FLevelEditorContextMenuImpl::FillEditMenu),
+				false, // default value
+				FSlateIcon(FAppStyle::GetAppStyleSetName(), "Icons.Edit"));
+
+			// TODO: The visibility actions don't yet properly support elements, but should
+			//Section.AddSubMenu(
+			//	"VisibilitySubMenu",
+			//	LOCTEXT("VisibilitySubMenu", "Visibility"),
+			//	LOCTEXT("VisibilitySubMenu_ToolTip", "Selected element visibility options"),
+			//	FNewToolMenuDelegate::CreateStatic(&FLevelEditorContextMenuImpl::FillActorVisibilityMenu),
+			//	false, // default value
+			//	FSlateIcon(FAppStyle::GetAppStyleSetName(), "Icons.Visibility"));
+
 			Section.AddSubMenu(
 				"TransformSubMenu",
 				LOCTEXT("TransformSubMenu", "Transform"),
-				LOCTEXT("TransformSubMenu_ToolTip_Element", "Element transform utils"),
-				FNewToolMenuDelegate::CreateStatic(&FLevelEditorContextMenuImpl::FillTransformMenu));
+				LOCTEXT("TransformSubMenu_ToolTip", "Element transform utils"),
+				FNewToolMenuDelegate::CreateStatic(&FLevelEditorContextMenuImpl::FillTransformMenu),
+				false, // default value
+				FSlateIcon(FAppStyle::GetAppStyleSetName(), "Icons.Transform"));
+
+			Section.AddSubMenu(
+				"SnapAlignSubMenu",
+				LOCTEXT("SnapAlignSubMenu", "Snapping"),
+				LOCTEXT("SnapAlignSubMenu_ToolTip", "Element snap/align utils"),
+				FNewToolMenuDelegate::CreateStatic(&FLevelEditorContextMenuImpl::FillSnapAlignMenu),
+				false, // default value
+				FSlateIcon(FAppStyle::GetAppStyleSetName(), "Icons.Snap"));
+
+			// TODO: Some pivot options work, but others still require additional support for non-actors
+			//Section.AddSubMenu(
+			//	"PivotSubMenu",
+			//	LOCTEXT("PivotSubMenu", "Pivot"),
+			//	LOCTEXT("PivotSubMenu_ToolTip", "Element pivoting utils"),
+			//	FNewToolMenuDelegate::CreateStatic(&FLevelEditorContextMenuImpl::FillPivotMenu),
+			//	false, // default value
+			//	FSlateIcon(FAppStyle::GetAppStyleSetName(), "Icons.SetShowPivot"));
+
+			// TODO: Attach/detach requires additional support. Attach, in particular, requires ability to show
+			//  elements to which attachment is possible.
+			//Section.AddSubMenu(
+			//	"ElementAttachToSubMenu",
+			//	LOCTEXT("ElementAttachToSubMenu", "Attach To"),
+			//	LOCTEXT("ElementAttachToSubMenu_ToolTip", "Attach Element as child"),
+			//	FNewToolMenuDelegate::CreateStatic(&FLevelEditorContextMenuImpl::FillActorMenu),
+			//	false, // default value
+			//	FSlateIcon(FAppStyle::GetAppStyleSetName(), "Actors.Attach"));
+			//Section.AddMenuEntry(
+			//	FLevelEditorCommands::Get().DetachFromParent,
+			//	TAttribute<FText>(), // Use command title
+			//	TAttribute<FText>(), // Use command tooltip
+			//	FSlateIcon(FAppStyle::GetAppStyleSetName(), "Actors.Detach"));
+		}
+			
+		// General-purpose extension point for tools that apply to many types of elements, analogous to ActorUETools
+		if (LevelEditorContext->ContextType == ELevelEditorMenuContext::MainMenu)
+		{
+			FToolMenuSection& Section = InMenu->AddSection("ElementUETools", LOCTEXT("UEToolsHeading", "UE Tools"));
 		}
+
+		// General-purpose extension point for tools that only appear for certain types of elements, analogous to ActorTypeTools
+		FToolMenuSection& ElementTypeToolsSection = InMenu->AddSection("ElementTypeTools");
+
+		{
+			// Play from here
+			// TODO: May someday support custom play-from-here behavior the way that some actors do, and
+			//  may someday support simulation changes
+			FToolMenuSection& Section = InMenu->AddSection("ElementPreview", LOCTEXT("PreviewHeading", "Preview"));
+
+			// Note: not using a command for play from here since it requires a mouse click
+			FUIAction PlayFromHereAction(
+				FExecuteAction::CreateStatic(&FPlayWorldCommandCallbacks::StartPlayFromHere));
+
+			const FText PlayFromHereLabel = GEditor->OnlyLoadEditorVisibleLevelsInPIE() ? LOCTEXT("PlayFromHereVisible", "Play From Here (visible levels)") : LOCTEXT("PlayFromHere", "Play From Here");
+			Section.AddMenuEntry(NAME_None, PlayFromHereLabel, LOCTEXT("PlayFromHere_ToolTip", "Starts a game preview from the clicked location"), FSlateIcon(FAppStyle::GetAppStyleSetName(), "Icons.Play"), PlayFromHereAction);
+}
 	}));
 }
 
@@ -1171,23 +1296,24 @@ void FLevelEditorContextMenuImpl::FillTransformMenu(UToolMenu* Menu)
 			Section.AddMenuEntry(FLevelEditorCommands::Get().DeltaTransformToActors);
 		}
 
-		if (LevelEditorContext->CurrentSelection->HasSelectedObjects<AActor>())
-		{
-			FToolMenuSection& Section = Menu->AddSection("MirrorLock");
+		FToolMenuSection& Section = Menu->AddSection("MirrorLock");
 
-			if (AxisDisplayInfo::GetAxisDisplayCoordinateSystem() == EAxisList::LeftUpForward)
-			{
-				Section.AddMenuEntry(FLevelEditorCommands::Get().MirrorActorY);
-				Section.AddMenuEntry(FLevelEditorCommands::Get().MirrorActorZ);
-				Section.AddMenuEntry(FLevelEditorCommands::Get().MirrorActorX);
-			}
-			else
-			{
-				Section.AddMenuEntry(FLevelEditorCommands::Get().MirrorActorX);
-				Section.AddMenuEntry(FLevelEditorCommands::Get().MirrorActorY);
-				Section.AddMenuEntry(FLevelEditorCommands::Get().MirrorActorZ);
-			}
+		if (AxisDisplayInfo::GetAxisDisplayCoordinateSystem() == EAxisList::LeftUpForward)
+		{
+			Section.AddMenuEntry(FLevelEditorCommands::Get().MirrorActorY);
+			Section.AddMenuEntry(FLevelEditorCommands::Get().MirrorActorZ);
+			Section.AddMenuEntry(FLevelEditorCommands::Get().MirrorActorX);
+		}
+		else
+		{
+			Section.AddMenuEntry(FLevelEditorCommands::Get().MirrorActorX);
+			Section.AddMenuEntry(FLevelEditorCommands::Get().MirrorActorY);
+			Section.AddMenuEntry(FLevelEditorCommands::Get().MirrorActorZ);
+		}
 
+		// TODO: Should support this for non-actor elements somehow
+		if (LevelEditorContext->CurrentSelection->HasSelectedObjects<AActor>())
+		{
 			Section.AddMenuEntry(FLevelEditorCommands::Get().LockActorMovement);
 		}
 	}
