Editor: Fix a couple duplicate LOCTEXT entries in LevelEditorContextMenu.cpp

#Synced-CL 51218464

diff --git a/Engine/Source/Editor/LevelEditor/Private/LevelEditorContextMenu.cpp b/Engine/Source/Editor/LevelEditor/Private/LevelEditorContextMenu.cpp
--- a/Engine/Source/Editor/LevelEditor/Private/LevelEditorContextMenu.cpp
+++ b/Engine/Source/Editor/LevelEditor/Private/LevelEditorContextMenu.cpp
@@ -706,24 +706,24 @@ void FLevelEditorContextMenu::RegisterElementContextMenu()
 			// TODO: The visibility actions don't yet properly support elements, but should
 			//Section.AddSubMenu(
 			//	"VisibilitySubMenu",
-			//	LOCTEXT("VisibilitySubMenu", "Visibility"),
-			//	LOCTEXT("VisibilitySubMenu_ToolTip", "Selected element visibility options"),
+			//	LOCTEXT("VisibilityElementSubMenu", "Visibility"),
+			//	LOCTEXT("VisibilityElementSubMenu_ToolTip", "Selected element visibility options"),
 			//	FNewToolMenuDelegate::CreateStatic(&FLevelEditorContextMenuImpl::FillActorVisibilityMenu),
 			//	false, // default value
 			//	FSlateIcon(FAppStyle::GetAppStyleSetName(), "Icons.Visibility"));
 
 			Section.AddSubMenu(
 				"TransformSubMenu",
-				LOCTEXT("TransformSubMenu", "Transform"),
-				LOCTEXT("TransformSubMenu_ToolTip", "Element transform utils"),
+				LOCTEXT("TransformElementSubMenu", "Transform"),
+				LOCTEXT("TransformElementSubMenu_ToolTip", "Element transform utils"),
 				FNewToolMenuDelegate::CreateStatic(&FLevelEditorContextMenuImpl::FillTransformMenu),
 				false, // default value
 				FSlateIcon(FAppStyle::GetAppStyleSetName(), "Icons.Transform"));
 
 			Section.AddSubMenu(
 				"SnapAlignSubMenu",
-				LOCTEXT("SnapAlignSubMenu", "Snapping"),
-				LOCTEXT("SnapAlignSubMenu_ToolTip", "Element snap/align utils"),
+				LOCTEXT("SnapAlignElementSubMenu", "Snapping"),
+				LOCTEXT("SnapAlignElementSubMenu_ToolTip", "Element snap/align utils"),
 				FNewToolMenuDelegate::CreateStatic(&FLevelEditorContextMenuImpl::FillSnapAlignMenu),
 				false, // default value
 				FSlateIcon(FAppStyle::GetAppStyleSetName(), "Icons.Snap"));
@@ -731,8 +731,8 @@ void FLevelEditorContextMenu::RegisterElementContextMenu()
 			// TODO: Some pivot options work, but others still require additional support for non-actors
 			//Section.AddSubMenu(
 			//	"PivotSubMenu",
-			//	LOCTEXT("PivotSubMenu", "Pivot"),
-			//	LOCTEXT("PivotSubMenu_ToolTip", "Element pivoting utils"),
+			//	LOCTEXT("PivotElementSubMenu", "Pivot"),
+			//	LOCTEXT("PivotElementSubMenu_ToolTip", "Element pivoting utils"),
 			//	FNewToolMenuDelegate::CreateStatic(&FLevelEditorContextMenuImpl::FillPivotMenu),
 			//	false, // default value
 			//	FSlateIcon(FAppStyle::GetAppStyleSetName(), "Icons.SetShowPivot"));
