[Level viewport] Add button to enable/disable editor shortcuts when PIE is active.

In UE 5.6, PIE was changed to handle editor shortcuts. I added a new preference to let users disable editor shortcuts when PIE is active.
Note: this has no impact on the game's handling of keyboard shortcuts.

#Synced-CL 51330672

diff --git a/Engine/Source/Editor/LevelEditor/Private/LevelViewportActions.cpp b/Engine/Source/Editor/LevelEditor/Private/LevelViewportActions.cpp
--- a/Engine/Source/Editor/LevelEditor/Private/LevelViewportActions.cpp
+++ b/Engine/Source/Editor/LevelEditor/Private/LevelViewportActions.cpp
@@ -107,6 +107,8 @@ void FLevelViewportCommands::RegisterCommands()
 		FInputChord()
 	);
 
+	UI_COMMAND( EnableEditorShortcutsInPIE, "Enable editor keyboard shortcuts", "Enable editor keyboard shortcuts when PIE is active", EUserInterfaceActionType::ToggleButton, FInputChord() );
+
 	// Generate a command for each volume class
 	{
 		UI_COMMAND( ShowAllVolumes, "Show All Volumes", "Shows all volumes", EUserInterfaceActionType::Button, FInputChord() );
diff --git a/Engine/Source/Editor/LevelEditor/Private/SLevelViewport.cpp b/Engine/Source/Editor/LevelEditor/Private/SLevelViewport.cpp
--- a/Engine/Source/Editor/LevelEditor/Private/SLevelViewport.cpp
+++ b/Engine/Source/Editor/LevelEditor/Private/SLevelViewport.cpp
@@ -709,6 +709,17 @@ void SLevelViewport::OnPerformanceSettingsChanged(UObject* Obj, FPropertyChanged
 	}
 }
 
+void SLevelViewport::OnToggleEditorShortcutsInPIE()
+{
+	ULevelEditorPlaySettings* const PlaySettings = GetMutableDefault<ULevelEditorPlaySettings>();
+	PlaySettings->EnableEditorShortcutsInPIE = !PlaySettings->EnableEditorShortcutsInPIE;
+}
+
+bool SLevelViewport::IsEditorShortcutsEnabledInPIE()
+{
+	return GetDefault<ULevelEditorPlaySettings>()->EnableEditorShortcutsInPIE;
+}
+
 FSceneViewport* SLevelViewport::GetGameSceneViewport() const
 {
 	return ActiveViewport.Get();
@@ -794,12 +805,18 @@ FReply SLevelViewport::OnKeyDown( const FGeometry& MyGeometry, const FKeyEvent&
 
 	FReply Reply = FReply::Unhandled();
 
+	bool bEnableEditorShortcuts = true;
+	if (ULevelEditorPlaySettings* const PlaySettings = GetMutableDefault<ULevelEditorPlaySettings>())
+	{
+		bEnableEditorShortcuts = PlaySettings->EnableEditorShortcutsInPIE;
+	}
+
 	if( HasPlayInEditorViewport() || LevelViewportClient->IsSimulateInEditorViewport() )
 	{
 		// Only process commands for pie when a play world is active
 		bool bHandled = FPlayWorldCommands::GlobalPlayWorldActions->ProcessCommandBindings( InKeyEvent );
 		
-		if (HasPlayInEditorViewport())
+		if (bEnableEditorShortcuts && HasPlayInEditorViewport())
 		{
 			bHandled = bHandled || PIECommands->ProcessCommandBindings(InKeyEvent);
 		}
@@ -1748,6 +1765,13 @@ void SLevelViewport::BindOptionCommands( FUICommandList& OutCommandList )
 		FCanExecuteAction(),
 		FIsActionChecked::CreateSP( this, &SLevelViewport::IsViewportToolbarVisible )
 		);
+
+	OutCommandList.MapAction(
+		ViewportActions.EnableEditorShortcutsInPIE,
+		FExecuteAction::CreateSP( this, &SLevelViewport::OnToggleEditorShortcutsInPIE ),
+		FCanExecuteAction(),
+		FIsActionChecked::CreateSP( this, &SLevelViewport::IsEditorShortcutsEnabledInPIE )
+		);
 }
 
 void SLevelViewport::BindViewCommands( FUICommandList& OutCommandList )
diff --git a/Engine/Source/Editor/LevelEditor/Private/ViewportToolbar/LevelEditorSubmenus.cpp b/Engine/Source/Editor/LevelEditor/Private/ViewportToolbar/LevelEditorSubmenus.cpp
--- a/Engine/Source/Editor/LevelEditor/Private/ViewportToolbar/LevelEditorSubmenus.cpp
+++ b/Engine/Source/Editor/LevelEditor/Private/ViewportToolbar/LevelEditorSubmenus.cpp
@@ -2567,6 +2567,11 @@ FToolMenuEntry CreatePIESettingsSubmenu()
 							BuildPIEVolumeControlCustomWidget(Context->LevelViewport),
 							LOCTEXT("PIEVolumeControlLabel", "Volume")
 						));
+						FToolMenuEntry EditorShortcutsEntry = FToolMenuEntry::InitMenuEntry(FLevelViewportCommands::Get().EnableEditorShortcutsInPIE,
+							TAttribute<FText>(), TAttribute<FText>(),
+							FSlateIcon(FAppStyle::GetAppStyleSetName(), "LevelEditor.BrowseViewportControls"));
+						EditorShortcutsEntry.SetShowInToolbarTopLevel(true);
+						SettingsSection.AddEntry(EditorShortcutsEntry);
 					}
 					
 					// Viewport advanced settings
diff --git a/Engine/Source/Editor/LevelEditor/Public/LevelViewportActions.h b/Engine/Source/Editor/LevelEditor/Public/LevelViewportActions.h
--- a/Engine/Source/Editor/LevelEditor/Public/LevelViewportActions.h
+++ b/Engine/Source/Editor/LevelEditor/Public/LevelViewportActions.h
@@ -162,6 +162,9 @@ PRAGMA_ENABLE_DEPRECATION_WARNINGS
 	
 	TSharedPtr< FUICommandInfo > FocusViewportToSelectedActors;
 
+	/** Enable editor shortcuts when PIE is active */
+	TSharedPtr< FUICommandInfo > EnableEditorShortcutsInPIE;
+
 	/**
 	 * Bookmarks                   
 	 */
diff --git a/Engine/Source/Editor/LevelEditor/Public/SLevelViewport.h b/Engine/Source/Editor/LevelEditor/Public/SLevelViewport.h
--- a/Engine/Source/Editor/LevelEditor/Public/SLevelViewport.h
+++ b/Engine/Source/Editor/LevelEditor/Public/SLevelViewport.h
@@ -897,6 +897,12 @@ private:
 	/** Updates the real-time overrride applied to the viewport */
 	UE_API void OnPerformanceSettingsChanged(UObject* Obj, struct FPropertyChangedEvent& ChangeEvent);
 
+	/** Called when the user wants to enable or disable editor shortcuts when PIE is active. */
+	UE_API void OnToggleEditorShortcutsInPIE();
+
+	/** Called to determine when editor shortcuts in PIE is active. */
+	UE_API bool IsEditorShortcutsEnabledInPIE();
+
 private:
 	/** Tab which this viewport is located in */
 	TWeakPtr<class FLevelViewportLayout> ParentLayout;
diff --git a/Engine/Source/Editor/UnrealEd/Classes/Settings/LevelEditorPlaySettings.h b/Engine/Source/Editor/UnrealEd/Classes/Settings/LevelEditorPlaySettings.h
--- a/Engine/Source/Editor/UnrealEd/Classes/Settings/LevelEditorPlaySettings.h
+++ b/Engine/Source/Editor/UnrealEd/Classes/Settings/LevelEditorPlaySettings.h
@@ -291,6 +291,10 @@ public:
 	/** Whether to automatically solo audio in first PIE client. */
 	UPROPERTY(config, EditAnywhere, Category = PlayInEditor, meta = (ToolTip="Whether to automatically solo audio in first PIE client", EditCondition="EnableGameSound", DisplayName = "Solo Audio in First PIE Client"))
 	bool SoloAudioInFirstPIEClient;
+	
+	/** Whether to enable the editor's keyboard shortcuts in PIE. */
+	UPROPERTY(config, EditAnywhere, Category = PlayInEditor, meta = (ToolTip="Enable editor keyboard shortcuts when PIE is active", DisplayName = "Enable editor keyboard shortcuts"))
+	bool EnableEditorShortcutsInPIE;
 
 	/** Whether to play a sound when entering and exiting PIE */
 	UPROPERTY(config, EditAnywhere, AdvancedDisplay, Category = PlayInEditor, meta = (DisplayName = "Enable PIE Enter and Exit Sounds"))
