Fixed crash when resetting the editor layout while in the Animation Mode.

#rnx
#Synced-CL 51362910

diff --git a/Engine/Source/Editor/LevelEditor/Private/LevelEditor.cpp b/Engine/Source/Editor/LevelEditor/Private/LevelEditor.cpp
--- a/Engine/Source/Editor/LevelEditor/Private/LevelEditor.cpp
+++ b/Engine/Source/Editor/LevelEditor/Private/LevelEditor.cpp
@@ -426,12 +426,13 @@ TSharedRef<SDockTab> FLevelEditorModule::SpawnLevelEditor( const FSpawnTabArgs&
 	TSharedPtr<SLevelEditor> LevelEditorTmp;
 	if (OwnerWindow.IsValid())
 	{
+		// Make sure that any mode that could hold on to the old SLevelEditor releases it before setting the new one below
+		GLevelEditorModeTools().DeactivateAllModes();
+
 		LevelEditorTab->SetContent(SAssignNew(LevelEditorTmp, SLevelEditor));
 		SetLevelEditorInstance(LevelEditorTmp);
 		LevelEditorTmp->Initialize(LevelEditorTab, OwnerWindow.ToSharedRef());
 
-		GLevelEditorModeTools().DeactivateAllModes();
-
 		LevelEditorCreatedEvent.Broadcast(LevelEditorTmp);
 
 		TSharedRef<SProjectBadge> ProjectBadge = SNew(SProjectBadge);
