Fix stale element handle crashes in TypedElement customizations

Add null guards after GetActorFromHandleChecked/GetComponentFromHandleChecked
calls in selection and viewport interaction customizations. Element handles
can outlive their underlying UObjects (stored as TWeakObjectPtr), causing
crashes when the weak pointer resolves to nullptr.

#Synced-CL 51731166

diff --git a/Engine/Source/Editor/LevelEditor/Private/Elements/Actor/ActorElementLevelEditorSelectionCustomization.cpp b/Engine/Source/Editor/LevelEditor/Private/Elements/Actor/ActorElementLevelEditorSelectionCustomization.cpp
--- a/Engine/Source/Editor/LevelEditor/Private/Elements/Actor/ActorElementLevelEditorSelectionCustomization.cpp
+++ b/Engine/Source/Editor/LevelEditor/Private/Elements/Actor/ActorElementLevelEditorSelectionCustomization.cpp
@@ -64,7 +64,12 @@ FTypedElementHandle FActorElementLevelEditorSelectionCustomization::GetSelection
 
 void FActorElementLevelEditorSelectionCustomization::GetNormalizedElements(const TTypedElement<ITypedElementSelectionInterface>& InElementSelectionHandle, FTypedElementListConstRef InSelectionSet, const FTypedElementSelectionNormalizationOptions& InNormalizationOptions, FTypedElementListRef OutNormalizedElements)
 {
-	AActor* Actor = ActorElementDataUtil::GetActorFromHandleChecked(InElementSelectionHandle);
+	AActor* Actor = ActorElementDataUtil::GetActorFromHandle(InElementSelectionHandle);
+	if (!Actor)
+	{
+		UE_LOGF(LogActorLevelEditorSelection, Warning, "GetNormalizedElements: Failed to resolve actor from element handle");
+		return;
+	}
 
 	if (InSelectionSet->HasElementsOfType(NAME_Components))
 	{
@@ -113,7 +118,12 @@ bool IsActorReachable(const AActor* Actor)
 
 bool FActorElementLevelEditorSelectionCustomization::CanSelectActorElement(const TTypedElement<ITypedElementSelectionInterface>& InActorSelectionHandle, const FTypedElementSelectionOptions& InSelectionOptions) const
 {
-	AActor* Actor = ActorElementDataUtil::GetActorFromHandleChecked(InActorSelectionHandle);
+	AActor* Actor = ActorElementDataUtil::GetActorFromHandle(InActorSelectionHandle);
+	if (!Actor)
+	{
+		UE_LOGF(LogActorLevelEditorSelection, Warning, "CanSelectActorElement: Failed to resolve actor from element handle");
+		return false;
+	}
 
 	// Bail if global selection is locked, or this actor cannot be edited or selected
 	if (GEdSelectionLock || !Actor->IsEditable() || !Actor->IsSelectable())
@@ -165,7 +175,12 @@ bool FActorElementLevelEditorSelectionCustomization::CanSelectActorElement(const
 
 bool FActorElementLevelEditorSelectionCustomization::CanDeselectActorElement(const TTypedElement<ITypedElementSelectionInterface>& InActorSelectionHandle, const FTypedElementSelectionOptions& InSelectionOptions) const
 {
-	AActor* Actor = ActorElementDataUtil::GetActorFromHandleChecked(InActorSelectionHandle);
+	AActor* Actor = ActorElementDataUtil::GetActorFromHandle(InActorSelectionHandle);
+	if (!Actor)
+	{
+		UE_LOGF(LogActorLevelEditorSelection, Warning, "CanDeselectActorElement: Failed to resolve actor from element handle");
+		return true;
+	}
 
 	// Bail if global selection is locked
 	if (GEdSelectionLock)
@@ -189,7 +204,12 @@ bool FActorElementLevelEditorSelectionCustomization::CanDeselectActorElement(con
 
 bool FActorElementLevelEditorSelectionCustomization::SelectActorElement(const TTypedElement<ITypedElementSelectionInterface>& InActorSelectionHandle, FTypedElementListRef InSelectionSet, const FTypedElementSelectionOptions& InSelectionOptions)
 {
-	AActor* Actor = ActorElementDataUtil::GetActorFromHandleChecked(InActorSelectionHandle);
+	AActor* Actor = ActorElementDataUtil::GetActorFromHandle(InActorSelectionHandle);
+	if (!Actor)
+	{
+		UE_LOGF(LogActorLevelEditorSelection, Warning, "SelectActorElement: Failed to resolve actor from element handle");
+		return false;
+	}
 
 	if (const IToolkitHost* ToolkitHostPtr = GetToolkitHost())
 	{
@@ -262,7 +282,12 @@ bool FActorElementLevelEditorSelectionCustomization::SelectActorElement(const TT
 
 bool FActorElementLevelEditorSelectionCustomization::DeselectActorElement(const TTypedElement<ITypedElementSelectionInterface>& InActorSelectionHandle, FTypedElementListRef InSelectionSet, const FTypedElementSelectionOptions& InSelectionOptions)
 {
-	AActor* Actor = ActorElementDataUtil::GetActorFromHandleChecked(InActorSelectionHandle);
+	AActor* Actor = ActorElementDataUtil::GetActorFromHandle(InActorSelectionHandle);
+	if (!Actor)
+	{
+		UE_LOGF(LogActorLevelEditorSelection, Warning, "DeselectActorElement: Failed to resolve actor from element handle");
+		return false;
+	}
 
 	if (const IToolkitHost* ToolkitHostPtr = GetToolkitHost())
 	{
diff --git a/Engine/Source/Editor/LevelEditor/Private/Elements/Actor/ActorElementLevelEditorViewportInteractionCustomization.cpp b/Engine/Source/Editor/LevelEditor/Private/Elements/Actor/ActorElementLevelEditorViewportInteractionCustomization.cpp
--- a/Engine/Source/Editor/LevelEditor/Private/Elements/Actor/ActorElementLevelEditorViewportInteractionCustomization.cpp
+++ b/Engine/Source/Editor/LevelEditor/Private/Elements/Actor/ActorElementLevelEditorViewportInteractionCustomization.cpp
@@ -16,7 +16,12 @@
 
 void FActorElementLevelEditorViewportInteractionCustomization::GizmoManipulationStarted(const TTypedElement<ITypedElementWorldInterface>& InElementWorldHandle, const UE::Widget::EWidgetMode InWidgetMode)
 {
-	AActor* Actor = ActorElementDataUtil::GetActorFromHandleChecked(InElementWorldHandle);
+	AActor* Actor = ActorElementDataUtil::GetActorFromHandle(InElementWorldHandle);
+	if (!Actor)
+	{
+		UE_LOGF(LogTemp, Warning, "FActorElementLevelEditorViewportInteractionCustomization::GizmoManipulationStarted: Failed to resolve actor from element handle");
+		return;
+	}
 
 	// Notify that this actor is beginning to move
 	GEditor->BroadcastBeginObjectMovement(*Actor);
@@ -32,7 +37,12 @@ void FActorElementLevelEditorViewportInteractionCustomization::GizmoManipulation
 
 void FActorElementLevelEditorViewportInteractionCustomization::GizmoManipulationDeltaUpdate(const TTypedElement<ITypedElementWorldInterface>& InElementWorldHandle, const UE::Widget::EWidgetMode InWidgetMode, const EAxisList::Type InDragAxis, const FInputDeviceState& InInputState, const FTransform& InDeltaTransform, const FVector& InPivotLocation)
 {
-	AActor* Actor = ActorElementDataUtil::GetActorFromHandleChecked(InElementWorldHandle);
+	AActor* Actor = ActorElementDataUtil::GetActorFromHandle(InElementWorldHandle);
+	if (!Actor)
+	{
+		UE_LOGF(LogTemp, Warning, "FActorElementLevelEditorViewportInteractionCustomization::GizmoManipulationDeltaUpdate: Failed to resolve actor from element handle");
+		return;
+	}
 
 	GEditor->NoteActorMovement();
 
@@ -63,7 +73,12 @@ void FActorElementLevelEditorViewportInteractionCustomization::GizmoManipulation
 
 void FActorElementLevelEditorViewportInteractionCustomization::GizmoManipulationStopped(const TTypedElement<ITypedElementWorldInterface>& InElementWorldHandle, const UE::Widget::EWidgetMode InWidgetMode, const ETypedElementViewportInteractionGizmoManipulationType InManipulationType)
 {
-	AActor* Actor = ActorElementDataUtil::GetActorFromHandleChecked(InElementWorldHandle);
+	AActor* Actor = ActorElementDataUtil::GetActorFromHandle(InElementWorldHandle);
+	if (!Actor)
+	{
+		UE_LOGF(LogTemp, Warning, "FActorElementLevelEditorViewportInteractionCustomization::GizmoManipulationStopped: Failed to resolve actor from element handle");
+		return;
+	}
 
 	if (InManipulationType != ETypedElementViewportInteractionGizmoManipulationType::Drag)
 	{
@@ -89,7 +104,7 @@ void FActorElementLevelEditorViewportInteractionCustomization::GizmoManipulation
 
 void FActorElementLevelEditorViewportInteractionCustomization::PostGizmoManipulationStopped(TArrayView<const FTypedElementHandle> InElementHandles, const UE::Widget::EWidgetMode InWidgetMode)
 {
-	TArray<AActor*> MovedActors = ActorElementDataUtil::GetActorsFromHandlesChecked(InElementHandles);
+	TArray<AActor*> MovedActors = ActorElementDataUtil::GetActorsFromHandles(InElementHandles);
 	GEditor->BroadcastActorsMoved(MovedActors);
 }
 
diff --git a/Engine/Source/Editor/LevelEditor/Private/Elements/Component/ComponentElementLevelEditorSelectionCustomization.cpp b/Engine/Source/Editor/LevelEditor/Private/Elements/Component/ComponentElementLevelEditorSelectionCustomization.cpp
--- a/Engine/Source/Editor/LevelEditor/Private/Elements/Component/ComponentElementLevelEditorSelectionCustomization.cpp
+++ b/Engine/Source/Editor/LevelEditor/Private/Elements/Component/ComponentElementLevelEditorSelectionCustomization.cpp
@@ -92,7 +92,12 @@ FTypedElementHandle FComponentElementLevelEditorSelectionCustomization::GetSelec
 
 void FComponentElementLevelEditorSelectionCustomization::GetNormalizedElements(const TTypedElement<ITypedElementSelectionInterface>& InElementSelectionHandle, FTypedElementListConstRef InSelectionSet, const FTypedElementSelectionNormalizationOptions& InNormalizationOptions, FTypedElementListRef OutNormalizedElements)
 {
-	UActorComponent* Component = ComponentElementDataUtil::GetComponentFromHandleChecked(InElementSelectionHandle);
+	UActorComponent* Component = ComponentElementDataUtil::GetComponentFromHandle(InElementSelectionHandle);
+	if (!Component)
+	{
+		UE_LOGF(LogComponentLevelEditorSelection, Warning, "GetNormalizedElements: Failed to resolve component from element handle");
+		return;
+	}
 
 	if (USceneComponent* SceneComponent = Cast<USceneComponent>(Component))
 	{
@@ -135,7 +140,12 @@ void FComponentElementLevelEditorSelectionCustomization::GetNormalizedElements(c
 
 bool FComponentElementLevelEditorSelectionCustomization::CanSelectComponentElement(const TTypedElement<ITypedElementSelectionInterface>& InComponentSelectionHandle, const FTypedElementSelectionOptions& InSelectionOptions) const
 {
-	const UActorComponent* Component = ComponentElementDataUtil::GetComponentFromHandleChecked(InComponentSelectionHandle);
+	const UActorComponent* Component = ComponentElementDataUtil::GetComponentFromHandle(InComponentSelectionHandle);
+	if (!Component)
+	{
+		UE_LOGF(LogComponentLevelEditorSelection, Warning, "CanSelectComponentElement: Failed to resolve component from element handle");
+		return false;
+	}
 
 	// Bail if global selection is locked
 	return !GEdSelectionLock;
@@ -143,7 +153,12 @@ bool FComponentElementLevelEditorSelectionCustomization::CanSelectComponentEleme
 
 bool FComponentElementLevelEditorSelectionCustomization::CanDeselectComponentElement(const TTypedElement<ITypedElementSelectionInterface>& InComponentSelectionHandle, const FTypedElementSelectionOptions& InSelectionOptions) const
 {
-	const UActorComponent* Component = ComponentElementDataUtil::GetComponentFromHandleChecked(InComponentSelectionHandle);
+	const UActorComponent* Component = ComponentElementDataUtil::GetComponentFromHandle(InComponentSelectionHandle);
+	if (!Component)
+	{
+		UE_LOGF(LogComponentLevelEditorSelection, Warning, "CanDeselectComponentElement: Failed to resolve component from element handle");
+		return true;
+	}
 
 	// Bail if global selection is locked
 	return !GEdSelectionLock;
@@ -151,7 +166,12 @@ bool FComponentElementLevelEditorSelectionCustomization::CanDeselectComponentEle
 
 bool FComponentElementLevelEditorSelectionCustomization::SelectComponentElement(const TTypedElement<ITypedElementSelectionInterface>& InComponentSelectionHandle, FTypedElementListRef InSelectionSet, const FTypedElementSelectionOptions& InSelectionOptions)
 {
-	UActorComponent* Component = ComponentElementDataUtil::GetComponentFromHandleChecked(InComponentSelectionHandle);
+	UActorComponent* Component = ComponentElementDataUtil::GetComponentFromHandle(InComponentSelectionHandle);
+	if (!Component)
+	{
+		UE_LOGF(LogComponentLevelEditorSelection, Warning, "SelectComponentElement: Failed to resolve component from element handle");
+		return false;
+	}
 
 	if (!InComponentSelectionHandle.SelectElement(InSelectionSet, InSelectionOptions))
 	{
@@ -187,7 +207,12 @@ bool FComponentElementLevelEditorSelectionCustomization::SelectComponentElement(
 
 bool FComponentElementLevelEditorSelectionCustomization::DeselectComponentElement(const TTypedElement<ITypedElementSelectionInterface>& InComponentSelectionHandle, FTypedElementListRef InSelectionSet, const FTypedElementSelectionOptions& InSelectionOptions)
 {
-	UActorComponent* Component = ComponentElementDataUtil::GetComponentFromHandleChecked(InComponentSelectionHandle);
+	UActorComponent* Component = ComponentElementDataUtil::GetComponentFromHandle(InComponentSelectionHandle);
+	if (!Component)
+	{
+		UE_LOGF(LogComponentLevelEditorSelection, Warning, "DeselectComponentElement: Failed to resolve component from element handle");
+		return false;
+	}
 
 	if (!InComponentSelectionHandle.DeselectElement(InSelectionSet, InSelectionOptions))
 	{
diff --git a/Engine/Source/Editor/LevelEditor/Private/Elements/Component/ComponentElementLevelEditorViewportInteractionCustomization.cpp b/Engine/Source/Editor/LevelEditor/Private/Elements/Component/ComponentElementLevelEditorViewportInteractionCustomization.cpp
--- a/Engine/Source/Editor/LevelEditor/Private/Elements/Component/ComponentElementLevelEditorViewportInteractionCustomization.cpp
+++ b/Engine/Source/Editor/LevelEditor/Private/Elements/Component/ComponentElementLevelEditorViewportInteractionCustomization.cpp
@@ -15,7 +15,12 @@
 
 void FComponentElementLevelEditorViewportInteractionCustomization::GizmoManipulationStarted(const TTypedElement<ITypedElementWorldInterface>& InElementWorldHandle, const UE::Widget::EWidgetMode InWidgetMode)
 {
-	UActorComponent* Component = ComponentElementDataUtil::GetComponentFromHandleChecked(InElementWorldHandle);
+	UActorComponent* Component = ComponentElementDataUtil::GetComponentFromHandle(InElementWorldHandle);
+	if (!Component)
+	{
+		UE_LOGF(LogTemp, Warning, "FComponentElementLevelEditorViewportInteractionCustomization::GizmoManipulationStarted: Failed to resolve component from element handle");
+		return;
+	}
 
 	if (USceneComponent* SceneComponent = Cast<USceneComponent>(Component))
 	{
@@ -34,7 +39,12 @@ void FComponentElementLevelEditorViewportInteractionCustomization::GizmoManipula
 
 void FComponentElementLevelEditorViewportInteractionCustomization::GizmoManipulationDeltaUpdate(const TTypedElement<ITypedElementWorldInterface>& InElementWorldHandle, const UE::Widget::EWidgetMode InWidgetMode, const EAxisList::Type InDragAxis, const FInputDeviceState& InInputState, const FTransform& InDeltaTransform, const FVector& InPivotLocation)
 {
-	UActorComponent* Component = ComponentElementDataUtil::GetComponentFromHandleChecked(InElementWorldHandle);
+	UActorComponent* Component = ComponentElementDataUtil::GetComponentFromHandle(InElementWorldHandle);
+	if (!Component)
+	{
+		UE_LOGF(LogTemp, Warning, "FComponentElementLevelEditorViewportInteractionCustomization::GizmoManipulationDeltaUpdate: Failed to resolve component from element handle");
+		return;
+	}
 
 	if (USceneComponent* SceneComponent = Cast<USceneComponent>(Component))
 	{
@@ -69,7 +79,12 @@ void FComponentElementLevelEditorViewportInteractionCustomization::GizmoManipula
 
 void FComponentElementLevelEditorViewportInteractionCustomization::GizmoManipulationStopped(const TTypedElement<ITypedElementWorldInterface>& InElementWorldHandle, const UE::Widget::EWidgetMode InWidgetMode, const ETypedElementViewportInteractionGizmoManipulationType InManipulationType)
 {
-	UActorComponent* Component = ComponentElementDataUtil::GetComponentFromHandleChecked(InElementWorldHandle);
+	UActorComponent* Component = ComponentElementDataUtil::GetComponentFromHandle(InElementWorldHandle);
+	if (!Component)
+	{
+		UE_LOGF(LogTemp, Warning, "FComponentElementLevelEditorViewportInteractionCustomization::GizmoManipulationStopped: Failed to resolve component from element handle");
+		return;
+	}
 
 	if (InManipulationType != ETypedElementViewportInteractionGizmoManipulationType::Drag)
 	{
diff --git a/Engine/Source/Editor/UnrealEd/Private/Elements/Actor/ActorElementEditorSelectionInterface.cpp b/Engine/Source/Editor/UnrealEd/Private/Elements/Actor/ActorElementEditorSelectionInterface.cpp
--- a/Engine/Source/Editor/UnrealEd/Private/Elements/Actor/ActorElementEditorSelectionInterface.cpp
+++ b/Engine/Source/Editor/UnrealEd/Private/Elements/Actor/ActorElementEditorSelectionInterface.cpp
@@ -38,7 +38,11 @@ private:
  
 	virtual void SetElementImpl(const FTypedElementHandle& InElementHandle) override
 	{
-		AActor* Actor = ActorElementDataUtil::GetActorFromHandleChecked(InElementHandle);
+		AActor* Actor = ActorElementDataUtil::GetActorFromHandle(InElementHandle);
+		if (!Actor)
+		{
+			UE_LOGF(LogTemp, Warning, "FActorElementTransactedElement::SetElementImpl: Failed to resolve actor from element handle");
+		}
 		ActorPtr = Actor;
 	}
  
diff --git a/Engine/Source/Editor/UnrealEd/Private/Elements/Actor/ActorElementEditorViewportInteractionCustomization.cpp b/Engine/Source/Editor/UnrealEd/Private/Elements/Actor/ActorElementEditorViewportInteractionCustomization.cpp
--- a/Engine/Source/Editor/UnrealEd/Private/Elements/Actor/ActorElementEditorViewportInteractionCustomization.cpp
+++ b/Engine/Source/Editor/UnrealEd/Private/Elements/Actor/ActorElementEditorViewportInteractionCustomization.cpp
@@ -18,8 +18,13 @@ void FActorElementEditorViewportInteractionCustomization::GizmoManipulationStart
 	const TTypedElement<ITypedElementWorldInterface>& InElementWorldHandle,
 	const UE::Widget::EWidgetMode InWidgetMode)
 {
-	AActor* RootActor = ActorElementDataUtil::GetActorFromHandleChecked(InElementWorldHandle);
-	
+	AActor* RootActor = ActorElementDataUtil::GetActorFromHandle(InElementWorldHandle);
+	if (!RootActor)
+	{
+		UE_LOGF(LogTemp, Warning, "FActorElementEditorViewportInteractionCustomization::GizmoManipulationStarted: Failed to resolve actor from element handle");
+		return;
+	}
+
 	FTypedElementViewportInteractionCustomization::GizmoManipulationStarted(InElementWorldHandle, InWidgetMode);
 
 	TArray<AActor*> AffectedActors;
@@ -49,8 +54,13 @@ void FActorElementEditorViewportInteractionCustomization::GizmoManipulationStart
 
 void FActorElementEditorViewportInteractionCustomization::GizmoManipulationDeltaUpdate(const TTypedElement<ITypedElementWorldInterface>& InElementWorldHandle, const UE::Widget::EWidgetMode InWidgetMode, const EAxisList::Type InDragAxis, const FInputDeviceState& InInputState, const FTransform& InDeltaTransform, const FVector& InPivotLocation)
 {
-	AActor* RootActor = ActorElementDataUtil::GetActorFromHandleChecked(InElementWorldHandle);
-	
+	AActor* RootActor = ActorElementDataUtil::GetActorFromHandle(InElementWorldHandle);
+	if (!RootActor)
+	{
+		UE_LOGF(LogTemp, Warning, "FActorElementEditorViewportInteractionCustomization::GizmoManipulationDeltaUpdate: Failed to resolve actor from element handle");
+		return;
+	}
+
 	TArray<AActor*> AffectedActors;
 	AffectedActors.Add(RootActor);
 	
@@ -202,8 +212,13 @@ void FActorElementEditorViewportInteractionCustomization::GizmoManipulationStopp
 	FTypedElementViewportInteractionCustomization::GizmoManipulationStopped(
 		InElementWorldHandle, InWidgetMode, InManipulationType);
 
-	AActor* RootActor = ActorElementDataUtil::GetActorFromHandleChecked(InElementWorldHandle);
-	
+	AActor* RootActor = ActorElementDataUtil::GetActorFromHandle(InElementWorldHandle);
+	if (!RootActor)
+	{
+		UE_LOGF(LogTemp, Warning, "FActorElementEditorViewportInteractionCustomization::GizmoManipulationStopped: Failed to resolve actor from element handle");
+		return;
+	}
+
 	TArray<AActor*> AffectedActors;
 	AffectedActors.Add(RootActor);
 	
@@ -220,7 +235,12 @@ void FActorElementEditorViewportInteractionCustomization::GizmoManipulationStopp
 
 void FActorElementEditorViewportInteractionCustomization::MirrorElement(const TTypedElement<ITypedElementWorldInterface>& InElementWorldHandle, const FVector& InMirrorScale, const FVector& InPivotLocation)
 {
-	AActor* Actor = ActorElementDataUtil::GetActorFromHandleChecked(InElementWorldHandle);
+	AActor* Actor = ActorElementDataUtil::GetActorFromHandle(InElementWorldHandle);
+	if (!Actor)
+	{
+		UE_LOGF(LogTemp, Warning, "FActorElementEditorViewportInteractionCustomization::MirrorElement: Failed to resolve actor from element handle");
+		return;
+	}
 
 	Actor->Modify();
 	Actor->EditorApplyMirror(InMirrorScale, InPivotLocation);
diff --git a/Engine/Source/Editor/UnrealEd/Private/Elements/Component/ComponentElementEditorViewportInteractionCustomization.cpp b/Engine/Source/Editor/UnrealEd/Private/Elements/Component/ComponentElementEditorViewportInteractionCustomization.cpp
--- a/Engine/Source/Editor/UnrealEd/Private/Elements/Component/ComponentElementEditorViewportInteractionCustomization.cpp
+++ b/Engine/Source/Editor/UnrealEd/Private/Elements/Component/ComponentElementEditorViewportInteractionCustomization.cpp
@@ -14,7 +14,12 @@
 
 bool FComponentElementEditorViewportInteractionCustomization::GetGizmoPivotLocation(const TTypedElement<ITypedElementWorldInterface>& InElementWorldHandle, const UE::Widget::EWidgetMode InWidgetMode, FVector& OutPivotLocation)
 {
-	const UActorComponent* Component = ComponentElementDataUtil::GetComponentFromHandleChecked(InElementWorldHandle);
+	const UActorComponent* Component = ComponentElementDataUtil::GetComponentFromHandle(InElementWorldHandle);
+	if (!Component)
+	{
+		UE_LOGF(LogTemp, Warning, "FComponentElementEditorViewportInteractionCustomization::GetGizmoPivotLocation: Failed to resolve component from element handle");
+		return false;
+	}
 
 	if (const USceneComponent* SceneComponent = Cast<USceneComponent>(Component))
 	{
@@ -33,7 +38,12 @@ bool FComponentElementEditorViewportInteractionCustomization::GetGizmoPivotLocat
 
 void FComponentElementEditorViewportInteractionCustomization::GizmoManipulationDeltaUpdate(const TTypedElement<ITypedElementWorldInterface>& InElementWorldHandle, const UE::Widget::EWidgetMode InWidgetMode, const EAxisList::Type InDragAxis, const FInputDeviceState& InInputState, const FTransform& InDeltaTransform, const FVector& InPivotLocation)
 {
-	UActorComponent* Component = ComponentElementDataUtil::GetComponentFromHandleChecked(InElementWorldHandle);
+	UActorComponent* Component = ComponentElementDataUtil::GetComponentFromHandle(InElementWorldHandle);
+	if (!Component)
+	{
+		UE_LOGF(LogTemp, Warning, "FComponentElementEditorViewportInteractionCustomization::GizmoManipulationDeltaUpdate: Failed to resolve component from element handle");
+		return;
+	}
 
 	if (USceneComponent* SceneComponent = Cast<USceneComponent>(Component))
 	{
