Fix ABrush/AVolume triggering unnecessary World Outliner rebuilds

#Synced-CL 51660293

diff --git a/Engine/Source/Editor/LevelEditor/Private/Elements/Actor/ActorElementLevelEditorViewportInteractionCustomization.cpp b/Engine/Source/Editor/LevelEditor/Private/Elements/Actor/ActorElementLevelEditorViewportInteractionCustomization.cpp
--- a/Engine/Source/Editor/LevelEditor/Private/Elements/Actor/ActorElementLevelEditorViewportInteractionCustomization.cpp
+++ b/Engine/Source/Editor/LevelEditor/Private/Elements/Actor/ActorElementLevelEditorViewportInteractionCustomization.cpp
@@ -93,10 +93,13 @@ void FActorElementLevelEditorViewportInteractionCustomization::GizmoManipulation
 		FCoreUObjectDelegates::OnObjectPropertyChanged.Broadcast(Actor, PropertyChangedEvent);
 	}
 
-	if (Actor->IsA(ABrush::StaticClass()))
+	if (const ABrush* AsBrush = Cast<ABrush>(Actor))
 	{
-		GEditor->RebuildAlteredBSP();
-	} 
+		if (AsBrush->CanEverAffectBSP())
+		{
+			GEditor->RebuildAlteredBSP();
+		}
+	}
 
 	Actor->PostEditMove(true);
 	GEditor->BroadcastEndObjectMovement(*Actor);
