Adds a LevelEditorSubystem call to delete outliner folders. BP/Python can currently use actor methods to create folders and move actors around, but there is no way to actually remove a folder.

#Synced-CL 51932664

diff --git a/Engine/Source/Editor/LevelEditor/Private/LevelEditorSubsystem.cpp b/Engine/Source/Editor/LevelEditor/Private/LevelEditorSubsystem.cpp
--- a/Engine/Source/Editor/LevelEditor/Private/LevelEditorSubsystem.cpp
+++ b/Engine/Source/Editor/LevelEditor/Private/LevelEditorSubsystem.cpp
@@ -1,6 +1,7 @@
 // Copyright Epic Games, Inc. All Rights Reserved.
 
 #include "LevelEditorSubsystem.h"
+#include "EditorActorFolders.h"
 
 #include "Editor.h"
 #include "Editor/UnrealEdEngine.h"
@@ -744,6 +745,39 @@ bool ULevelEditorSubsystem::SetCurrentLevelByName(FName LevelName)
 	return bLevelFound;
 }
 
+bool ULevelEditorSubsystem::DeleteActorFolder(FName FolderPath)
+{
+	TGuardValue<bool> UnattendedScriptGuard(GIsRunningUnattendedScript, true);
+
+	if (!EditorScriptingHelpers::CheckIfInEditorAndPIE())
+	{
+		return false;
+	}
+
+	UUnrealEditorSubsystem* UnrealEditorSubsystem = GEditor->GetEditorSubsystem<UUnrealEditorSubsystem>();
+	if (!UnrealEditorSubsystem)
+	{
+		return false;
+	}
+
+	UWorld* World = UnrealEditorSubsystem->GetEditorWorld();
+	if (!World)
+	{
+		UE_LOG(LevelEditorSubsystem, Error, TEXT("DeleteActorFolder. Can't delete folder because there is no world."));
+		return false;
+	}
+
+	FActorFolders& ActorFolders = FActorFolders::Get();
+	FFolder Folder(FFolder::GetInvalidRootObject(), FolderPath);
+	if (!ActorFolders.ContainsFolder(*World, Folder))
+	{
+		return false;
+	}
+
+	ActorFolders.DeleteFolder(*World, Folder);
+	return true;
+}
+
 ULevel* ULevelEditorSubsystem::GetCurrentLevel()
 {
 	UUnrealEditorSubsystem* UnrealEditorSubsystem = GEditor->GetEditorSubsystem<UUnrealEditorSubsystem>();
diff --git a/Engine/Source/Editor/LevelEditor/Public/LevelEditorSubsystem.h b/Engine/Source/Editor/LevelEditor/Public/LevelEditorSubsystem.h
--- a/Engine/Source/Editor/LevelEditor/Public/LevelEditorSubsystem.h
+++ b/Engine/Source/Editor/LevelEditor/Public/LevelEditorSubsystem.h
@@ -196,6 +196,15 @@ public:
 	UFUNCTION(BlueprintCallable, Category = "Editor Scripting | Level Utility")
 	UE_API ULevel* GetCurrentLevel();
 
+	/**
+	 * Deletes a folder from the World Outliner.
+	 * Does not move or delete actors; only removes the folder entry itself.
+	 * @param	FolderPath	The path of the folder to delete (e.g. 'Lighting/Spotlights').
+	 * @return	True if the folder existed and was deleted.
+	 */
+	UFUNCTION(BlueprintCallable, Category = "Editor Scripting | Level Utility")
+	UE_API bool DeleteActorFolder(FName FolderPath);
+
 	/**
 	 * Get the selection set for the current world, you can use this to track
 	 * and create changes to the level editor's selection
