Fix bug in DeleteActorFolder.

#Synced-CL 52087183

diff --git a/Engine/Source/Editor/LevelEditor/Private/LevelEditorSubsystem.cpp b/Engine/Source/Editor/LevelEditor/Private/LevelEditorSubsystem.cpp
--- a/Engine/Source/Editor/LevelEditor/Private/LevelEditorSubsystem.cpp
+++ b/Engine/Source/Editor/LevelEditor/Private/LevelEditorSubsystem.cpp
@@ -768,7 +768,7 @@ bool ULevelEditorSubsystem::DeleteActorFolder(FName FolderPath)
 	}
 
 	FActorFolders& ActorFolders = FActorFolders::Get();
-	FFolder Folder(FFolder::GetInvalidRootObject(), FolderPath);
+	FFolder Folder(FFolder::GetWorldRootFolder(World).GetRootObject(), FolderPath);
 	if (!ActorFolders.ContainsFolder(*World, Folder))
 	{
 		return false;
