Rename ToggleFeatureLevelPreview to TogglePreviewPlatform
Rename IsFeatureLevelPreviewEnabled to IsPreviewPlatformEnabled
Rename IsFeatureLevelPreviewActive  to IsPreviewPlatformActive

#Synced-CL 52125551

#Synced-CL 52125551

diff --git a/Engine/Plugins/PCG/Source/PCG/Private/Elements/ControlFlow/PCGShaderFeatureLevelSwitch.cpp b/Engine/Plugins/PCG/Source/PCG/Private/Elements/ControlFlow/PCGShaderFeatureLevelSwitch.cpp
--- a/Engine/Plugins/PCG/Source/PCG/Private/Elements/ControlFlow/PCGShaderFeatureLevelSwitch.cpp
+++ b/Engine/Plugins/PCG/Source/PCG/Private/Elements/ControlFlow/PCGShaderFeatureLevelSwitch.cpp
@@ -51,7 +51,7 @@ namespace PCGShaderFeatureLevelSwitch
 #if WITH_EDITOR
 	static ERHIFeatureLevel::Type GetCurrentFeatureLevel()
 	{
-		if (GEditor && GEditor->IsFeatureLevelPreviewActive())
+		if (GEditor && GEditor->IsPreviewPlatformActive())
 		{
 			return GEditor->GetActiveFeatureLevelPreviewType();
 		}
diff --git a/Engine/Source/Editor/LevelEditor/Private/LevelEditor.cpp b/Engine/Source/Editor/LevelEditor/Private/LevelEditor.cpp
--- a/Engine/Source/Editor/LevelEditor/Private/LevelEditor.cpp
+++ b/Engine/Source/Editor/LevelEditor/Private/LevelEditor.cpp
@@ -2125,10 +2125,10 @@ void FLevelEditorModule::BindGlobalLevelEditorCommands()
 		FIsActionChecked::CreateStatic(&FLevelEditorActionCallbacks::IsMaterialQualityLevelChecked, (EMaterialQualityLevel::Type)EMaterialQualityLevel::Epic));
 
 	ActionList.MapAction(
-		Commands.ToggleFeatureLevelPreview,
-		FExecuteAction::CreateStatic(&FLevelEditorActionCallbacks::ToggleFeatureLevelPreview),
-		FIsActionChecked::CreateStatic(&FLevelEditorActionCallbacks::IsFeatureLevelPreviewEnabled),
-		FIsActionChecked::CreateStatic(&FLevelEditorActionCallbacks::IsFeatureLevelPreviewActive),
+		Commands.TogglePreviewPlatform,
+		FExecuteAction::CreateStatic(&FLevelEditorActionCallbacks::TogglePreviewPlatform),
+		FIsActionChecked::CreateStatic(&FLevelEditorActionCallbacks::IsPreviewPlatformEnabled),
+		FIsActionChecked::CreateStatic(&FLevelEditorActionCallbacks::IsPreviewPlatformActive),
 		FIsActionButtonVisible::CreateStatic(FLevelEditorActionCallbacks::IsPreviewModeButtonVisible));
 
 	ActionList.MapAction(
diff --git a/Engine/Source/Editor/LevelEditor/Private/LevelEditorActions.cpp b/Engine/Source/Editor/LevelEditor/Private/LevelEditorActions.cpp
--- a/Engine/Source/Editor/LevelEditor/Private/LevelEditorActions.cpp
+++ b/Engine/Source/Editor/LevelEditor/Private/LevelEditorActions.cpp
@@ -797,7 +797,7 @@ void FLevelEditorActionCallbacks::AttachToSocketSelection(const FName SocketName
 void FLevelEditorActionCallbacks::SetMaterialQualityLevel( EMaterialQualityLevel::Type NewQualityLevel )
 {
 	static IConsoleVariable* MaterialQualityLevelVar = IConsoleManager::Get().FindConsoleVariable(TEXT("r.MaterialQualityLevel"));
-	EConsoleVariableFlags MaxPriority = GEditor && GEditor->IsFeatureLevelPreviewActive() ? ECVF_SetByPreview: ECVF_SetByScalability;
+	EConsoleVariableFlags MaxPriority = GEditor && GEditor->IsPreviewPlatformActive() ? ECVF_SetByPreview: ECVF_SetByScalability;
 	MaterialQualityLevelVar->ReplaceCurrentPriorityAndTag(NewQualityLevel, MaxPriority);
 
 	GUnrealEd->OnSceneMaterialsModified();
@@ -812,28 +812,28 @@ bool FLevelEditorActionCallbacks::IsMaterialQualityLevelChecked( EMaterialQualit
 	return TestQualityLevel == MaterialQualityLevel;
 }
 
-void FLevelEditorActionCallbacks::ToggleFeatureLevelPreview()
+void FLevelEditorActionCallbacks::TogglePreviewPlatform()
 {
-	GEditor->ToggleFeatureLevelPreview();
+	GEditor->TogglePreviewPlatform();
 }
 
-bool FLevelEditorActionCallbacks::IsFeatureLevelPreviewEnabled()
+bool FLevelEditorActionCallbacks::IsPreviewPlatformEnabled()
 {
 	if (GUnrealEd->IsLightingBuildCurrentlyRunning())
 	{
 		return false;
 	}
-	return GEditor->IsFeatureLevelPreviewEnabled();
+	return GEditor->IsPreviewPlatformEnabled();
 }
 
-bool FLevelEditorActionCallbacks::IsFeatureLevelPreviewActive()
+bool FLevelEditorActionCallbacks::IsPreviewPlatformActive()
 {
-	return GEditor->IsFeatureLevelPreviewEnabled() && GEditor->IsFeatureLevelPreviewActive();
+	return GEditor->IsPreviewPlatformEnabled() && GEditor->IsPreviewPlatformActive();
 }
 
 bool FLevelEditorActionCallbacks::IsPreviewModeButtonVisible()
 {
-	return GEditor->IsFeatureLevelPreviewEnabled();
+	return GEditor->IsPreviewPlatformEnabled();
 }
 
 void FLevelEditorActionCallbacks::SetPreviewPlatform(FPreviewPlatformInfo NewPreviewPlatform)
@@ -955,9 +955,9 @@ void FLevelEditorActionCallbacks::BuildReflectionCapturesOnly_Execute()
 		// This is to make sure that we have valid data for Mobile Preview.
 		
 		// ES31->SM5 to be able to capture
-		ToggleFeatureLevelPreview();
+		TogglePreviewPlatform();
 		// SM5->ES31 BuildReflectionCaptures are triggered here on callback
-		ToggleFeatureLevelPreview();
+		TogglePreviewPlatform();
 	}
 	else
 	{
@@ -4621,7 +4621,7 @@ void FLevelEditorCommands::RegisterCommands()
 	UI_COMMAND(MaterialQualityLevel_High, "High", "Sets material quality in the scene to high.", EUserInterfaceActionType::RadioButton, FInputChord());
 	UI_COMMAND(MaterialQualityLevel_Epic, "Epic", "Sets material quality in the scene to Epic.", EUserInterfaceActionType::RadioButton, FInputChord());
 
-	UI_COMMAND(ToggleFeatureLevelPreview, "Preview Mode Toggle", "Toggles the Preview Mode on or off for the currently selected Preview target", EUserInterfaceActionType::ToggleButton, FInputChord());
+	UI_COMMAND(TogglePreviewPlatform, "Preview Mode Toggle", "Toggles the Preview Mode on or off for the currently selected Preview target", EUserInterfaceActionType::ToggleButton, FInputChord());
 
 	UI_COMMAND(AllowArcballRotation, "Enable Arcball Rotation", "Allow arcball rotation with rotate widget", EUserInterfaceActionType::ToggleButton, FInputChord());
 	UI_COMMAND(AllowScreenspaceRotation, "Enable Screenspace Rotation", "Allow screen rotation with rotate widget", EUserInterfaceActionType::ToggleButton, FInputChord());
diff --git a/Engine/Source/Editor/LevelEditor/Private/LevelEditorToolBar.cpp b/Engine/Source/Editor/LevelEditor/Private/LevelEditorToolBar.cpp
--- a/Engine/Source/Editor/LevelEditor/Private/LevelEditorToolBar.cpp
+++ b/Engine/Source/Editor/LevelEditor/Private/LevelEditorToolBar.cpp
@@ -106,7 +106,7 @@ namespace PreviewModeFunctionality
 			{
 				return FText::Format(LOCTEXT("PreviewModeViewingAs", "Viewing {0}."), RenderingAsPlatformName);
 			}
-			else if (GEditor->IsFeatureLevelPreviewActive())
+			else if (GEditor->IsPreviewPlatformActive())
 			{
 				return FText::Format(LOCTEXT("PreviewModeViewingAsSwitchTo", "Viewing {0}. Click to preview {1}."), RenderingAsPlatformName, SwitchToPlatformName);
 			}
@@ -126,7 +126,7 @@ namespace PreviewModeFunctionality
 			});
 		if (Item)
 		{
-			return FSlateIcon(FAppStyle::GetAppStyleSetName(), GEditor->IsFeatureLevelPreviewActive() ? Item->ActiveIconName : Item->InactiveIconName);
+			return FSlateIcon(FAppStyle::GetAppStyleSetName(), GEditor->IsPreviewPlatformActive() ? Item->ActiveIconName : Item->InactiveIconName);
 		}
 
 		EShaderPlatform ShaderPlatform = FDataDrivenShaderPlatformInfo::GetShaderPlatformFromName(GEditor->PreviewPlatform.PreviewShaderPlatformName);
@@ -139,11 +139,11 @@ namespace PreviewModeFunctionality
 		{
 		case ERHIFeatureLevel::ES3_1:
 		{
-			return FSlateIcon(FAppStyle::GetAppStyleSetName(), GEditor->IsFeatureLevelPreviewActive() ? "LevelEditor.PreviewMode.Enabled" : "LevelEditor.PreviewMode.Disabled");
+			return FSlateIcon(FAppStyle::GetAppStyleSetName(), GEditor->IsPreviewPlatformActive() ? "LevelEditor.PreviewMode.Enabled" : "LevelEditor.PreviewMode.Disabled");
 		}
 		default:
 		{
-			return FSlateIcon(FAppStyle::GetAppStyleSetName(), GEditor->IsFeatureLevelPreviewActive() ? "LevelEditor.PreviewMode.Enabled" : "LevelEditor.PreviewMode.Disabled");
+			return FSlateIcon(FAppStyle::GetAppStyleSetName(), GEditor->IsPreviewPlatformActive() ? "LevelEditor.PreviewMode.Enabled" : "LevelEditor.PreviewMode.Disabled");
 		}
 		}
 	}
@@ -151,7 +151,7 @@ namespace PreviewModeFunctionality
 	LEVELEDITOR_API void AddPreviewToggleButton(FToolMenuSection& Section)
 	{
 		Section.AddEntry(FToolMenuEntry::InitToolBarButton(
-			FLevelEditorCommands::Get().ToggleFeatureLevelPreview,
+			FLevelEditorCommands::Get().TogglePreviewPlatform,
 			TAttribute<FText>::Create(&GetPreviewModeText),
 			TAttribute<FText>::Create(&GetPreviewModeTooltip),
 			TAttribute<FSlateIcon>::Create(&GetPreviewModeIcon)
diff --git a/Engine/Source/Editor/LevelEditor/Private/SLevelViewport.cpp b/Engine/Source/Editor/LevelEditor/Private/SLevelViewport.cpp
--- a/Engine/Source/Editor/LevelEditor/Private/SLevelViewport.cpp
+++ b/Engine/Source/Editor/LevelEditor/Private/SLevelViewport.cpp
@@ -1667,7 +1667,7 @@ void SLevelViewport::BindOptionCommands( FUICommandList& OutCommandList )
 		FExecuteAction::CreateLambda([]() { GEditor->ToggleAllowConstrainedAspectRatioInPreview(); }),
 		FCanExecuteAction(),
 		FIsActionChecked::CreateLambda([]() { return GEditor->GetAllowConstrainedAspectRatioInPreview(); }),
-		FIsActionButtonVisible::CreateLambda([]() { return GEditor->IsFeatureLevelPreviewActive() && (GEditor->PreviewPlatform.GetConstrainedAspectRatio() != 0.0f); })
+		FIsActionButtonVisible::CreateLambda([]() { return GEditor->IsPreviewPlatformActive() && (GEditor->PreviewPlatform.GetConstrainedAspectRatio() != 0.0f); })
 	);
 
 	OutCommandList.MapAction(
@@ -1675,7 +1675,7 @@ void SLevelViewport::BindOptionCommands( FUICommandList& OutCommandList )
 		FExecuteAction::CreateLambda([]() { GEditor->ToggleFlipConstrainedAspectRatioInPreview(); }),
 		FCanExecuteAction::CreateLambda([]() { return GEditor->GetAllowConstrainedAspectRatioInPreview(); }),
 		FIsActionChecked::CreateLambda([]() { return GEditor->GetFlipConstrainedAspectRatioInPreview(); }),
-		FIsActionButtonVisible::CreateLambda([]() { return GEditor->IsFeatureLevelPreviewActive() && (GEditor->PreviewPlatform.GetConstrainedAspectRatio() != 0.0f); })
+		FIsActionButtonVisible::CreateLambda([]() { return GEditor->IsPreviewPlatformActive() && (GEditor->PreviewPlatform.GetConstrainedAspectRatio() != 0.0f); })
 	);
 
 	IAssetRegistry & AssetRegistry = FModuleManager::LoadModuleChecked<FAssetRegistryModule>("AssetRegistry").Get();
diff --git a/Engine/Source/Editor/LevelEditor/Public/LevelEditorActions.h b/Engine/Source/Editor/LevelEditor/Public/LevelEditorActions.h
--- a/Engine/Source/Editor/LevelEditor/Public/LevelEditorActions.h
+++ b/Engine/Source/Editor/LevelEditor/Public/LevelEditorActions.h
@@ -634,7 +634,7 @@ public:
 	TSharedPtr< FUICommandInfo > MaterialQualityLevel_High;
 	TSharedPtr< FUICommandInfo > MaterialQualityLevel_Epic;
 
-	TSharedPtr< FUICommandInfo > ToggleFeatureLevelPreview;
+	TSharedPtr< FUICommandInfo > TogglePreviewPlatform;
 
 	TArray<TSharedPtr<FUICommandInfo>> PreviewPlatformOverrides;
 
@@ -902,10 +902,17 @@ public:
 	static UE_API void AttachToSocketSelection(FName SocketName, AActor* ParentActorPtr);
 	static UE_API void SetMaterialQualityLevel( EMaterialQualityLevel::Type NewQualityLevel );
 	static UE_API bool IsMaterialQualityLevelChecked( EMaterialQualityLevel::Type TestQualityLevel );
-	static UE_API void ToggleFeatureLevelPreview();
-	static UE_API bool IsFeatureLevelPreviewEnabled();
-	static UE_API bool IsFeatureLevelPreviewActive();
+	static UE_API void TogglePreviewPlatform();
+	static UE_API bool IsPreviewPlatformEnabled();
+	static UE_API bool IsPreviewPlatformActive();
 	static UE_API bool IsPreviewModeButtonVisible();
+
+	UE_DEPRECATED(5.8, "Use TogglePreviewPlatform instead")
+	static void ToggleFeatureLevelPreview() { TogglePreviewPlatform(); }
+	UE_DEPRECATED(5.8, "Use IsPreviewPlatformEnabled instead")
+	static bool IsFeatureLevelPreviewEnabled() { return IsPreviewPlatformEnabled(); }
+	UE_DEPRECATED(5.8, "Use IsPreviewPlatformActive instead")
+	static bool IsFeatureLevelPreviewActive() { return IsPreviewPlatformActive(); }
 	static UE_API void SetPreviewPlatform(FPreviewPlatformInfo NewPreviewPlatform);
 	static UE_API bool CanExecutePreviewPlatform(FPreviewPlatformInfo NewPreviewPlatform);
 	static UE_API bool IsPreviewPlatformChecked(FPreviewPlatformInfo NewPreviewPlatform);
diff --git a/Engine/Source/Editor/UnrealEd/Classes/Editor/EditorEngine.h b/Engine/Source/Editor/UnrealEd/Classes/Editor/EditorEngine.h
--- a/Engine/Source/Editor/UnrealEd/Classes/Editor/EditorEngine.h
+++ b/Engine/Source/Editor/UnrealEd/Classes/Editor/EditorEngine.h
@@ -3220,7 +3220,7 @@ private:
 	/** Gets the init values for worlds opened via Map_Load in the editor */
 	UNREALED_API UWorld::InitializationValues GetEditorWorldInitializationValues() const;
 
-	/** Logic that needs to run on change of effective shader platform (can happen via both SetPreviewPlatform and ToggleFeatureLevelPreview). */
+	/** Logic that needs to run on change of effective shader platform (can happen via both SetPreviewPlatform and TogglePreviewPlatform). */
 	UNREALED_API void OnEffectivePreviewShaderPlatformChange();
 
 public:
@@ -3240,14 +3240,21 @@ public:
 	/** Call this function to change the feature level and to override the material quality platform of the editor and PIE worlds */
 	UNREALED_API void SetPreviewPlatform(const FPreviewPlatformInfo& NewPreviewPlatform, bool bSaveSettings, bool bCompileShaders = true);
 
-	/** Toggle the feature level preview */
-	UNREALED_API void ToggleFeatureLevelPreview();
+	/** Toggle the preview platform */
+	UNREALED_API void TogglePreviewPlatform();
 
-	/** Return whether the feature level preview is enabled for use via the user accessing Settings->Preview Rendering Level. */
-	UNREALED_API bool IsFeatureLevelPreviewEnabled() const;
+	/** Return whether the preview platform is enabled for use via the user accessing Settings->Preview Rendering Level. */
+	UNREALED_API bool IsPreviewPlatformEnabled() const;
 
-	/** Return whether the feature level preview enabled via Settings->Preview Rendering Level is currently active/displaying. */
-	UNREALED_API bool IsFeatureLevelPreviewActive() const;
+	/** Return whether the preview platform enabled via Settings->Preview Rendering Level is currently active/displaying. */
+	UNREALED_API bool IsPreviewPlatformActive() const;
+
+	UE_DEPRECATED(5.8, "Use TogglePreviewPlatform instead")
+	void ToggleFeatureLevelPreview() { TogglePreviewPlatform(); }
+	UE_DEPRECATED(5.8, "Use IsPreviewPlatformEnabled instead")
+	bool IsFeatureLevelPreviewEnabled() const { return IsPreviewPlatformEnabled(); }
+	UE_DEPRECATED(5.8, "Use IsPreviewPlatformActive instead")
+	bool IsFeatureLevelPreviewActive() const { return IsPreviewPlatformActive(); }
 
 	/** Return the current active shader platform in the editor*/
 	UNREALED_API EShaderPlatform GetActiveShaderPlatform() const;
diff --git a/Engine/Source/Editor/UnrealEd/Private/EditorEngine.cpp b/Engine/Source/Editor/UnrealEd/Private/EditorEngine.cpp
--- a/Engine/Source/Editor/UnrealEd/Private/EditorEngine.cpp
+++ b/Engine/Source/Editor/UnrealEd/Private/EditorEngine.cpp
@@ -6820,7 +6820,7 @@ void UEditorEngine::SetPreviewPlatform(const FPreviewPlatformInfo& NewPreviewPla
 	UE_LOG(LogEditor, Display, TEXT("Done with Preview Platform"));
 }
 
-void UEditorEngine::ToggleFeatureLevelPreview()
+void UEditorEngine::TogglePreviewPlatform()
 {
 	FPreviewPlatformInfo NewInfo = PreviewPlatform;
 	NewInfo.bPreviewFeatureLevelActive ^= 1;
@@ -6835,12 +6835,12 @@ void UEditorEngine::ToggleFeatureLevelPreview()
 	SetPreviewPlatform(NewInfo, /*bSaveSettings=*/true, /*bCompileShaders=*/bNeedsCompile);
 }
 
-bool UEditorEngine::IsFeatureLevelPreviewEnabled() const
+bool UEditorEngine::IsPreviewPlatformEnabled() const
 {
 	return PreviewPlatform.PreviewShaderFormatName != NAME_None || PreviewPlatform.PreviewShaderPlatformName != NAME_None;
 }
 
-bool UEditorEngine::IsFeatureLevelPreviewActive() const
+bool UEditorEngine::IsPreviewPlatformActive() const
 {
  	return PreviewPlatform.bPreviewFeatureLevelActive;
 }
diff --git a/Engine/Source/Editor/UnrealEd/Private/EditorViewportClient.cpp b/Engine/Source/Editor/UnrealEd/Private/EditorViewportClient.cpp
--- a/Engine/Source/Editor/UnrealEd/Private/EditorViewportClient.cpp
+++ b/Engine/Source/Editor/UnrealEd/Private/EditorViewportClient.cpp
@@ -1192,7 +1192,7 @@ FSceneView* FEditorViewportClient::CalcSceneView(FSceneViewFamily* ViewFamily, c
 	TimeForForceRedraw = 0.0;
 
 	bool bConstrainAspectRatio = bUseControllingActorViewInfo && ControllingActorViewInfo.bConstrainAspectRatio;
-	if (GEditor != nullptr && GEditor->IsFeatureLevelPreviewActive() && GEditor->GetAllowConstrainedAspectRatioInPreview())
+	if (GEditor != nullptr && GEditor->IsPreviewPlatformActive() && GEditor->GetAllowConstrainedAspectRatioInPreview())
 	{
 		float PlatformAspectRatio = GEditor->PreviewPlatform.GetConstrainedAspectRatio();
 		if (GEditor->GetFlipConstrainedAspectRatioInPreview() && PlatformAspectRatio != 0.f)
diff --git a/Engine/Source/Editor/UnrealEd/Private/SEditorViewport.cpp b/Engine/Source/Editor/UnrealEd/Private/SEditorViewport.cpp
--- a/Engine/Source/Editor/UnrealEd/Private/SEditorViewport.cpp
+++ b/Engine/Source/Editor/UnrealEd/Private/SEditorViewport.cpp
@@ -1511,7 +1511,7 @@ EVisibility SEditorViewport::GetCurrentShaderPlatformPreviewTextVisibility() con
 {
 	if (Client->GetWorld() && !GLevelEditorModeTools().IsViewportUIHidden())
 	{
-		return (GEditor && GEditor->IsFeatureLevelPreviewActive()) ? EVisibility::SelfHitTestInvisible : EVisibility::Collapsed;
+		return (GEditor && GEditor->IsPreviewPlatformActive()) ? EVisibility::SelfHitTestInvisible : EVisibility::Collapsed;
 	}
 	else
 	{
diff --git a/Engine/Source/Editor/UnrealEd/Private/SScalabilitySettings.cpp b/Engine/Source/Editor/UnrealEd/Private/SScalabilitySettings.cpp
--- a/Engine/Source/Editor/UnrealEd/Private/SScalabilitySettings.cpp
+++ b/Engine/Source/Editor/UnrealEd/Private/SScalabilitySettings.cpp
@@ -183,7 +183,7 @@ TSharedRef<SWidget> SScalabilitySettings::MakeDefaultButtonWidget()
 		.IsChecked(this, &SScalabilitySettings::GetDefaultCheckState)
 		.ToolTipText_Lambda([]
 		{
-			if (GEditor && GEditor->IsFeatureLevelPreviewActive())
+			if (GEditor && GEditor->IsPreviewPlatformActive())
 			{
 				return LOCTEXT("DefaultButtonWithProfileTooltip", "Resets the Quality Levels to what the Device Profile sets");
 			}
diff --git a/Engine/Source/Editor/UnrealEd/Private/Subsystems/EditorSubsystemBlueprintLibrary.cpp b/Engine/Source/Editor/UnrealEd/Private/Subsystems/EditorSubsystemBlueprintLibrary.cpp
--- a/Engine/Source/Editor/UnrealEd/Private/Subsystems/EditorSubsystemBlueprintLibrary.cpp
+++ b/Engine/Source/Editor/UnrealEd/Private/Subsystems/EditorSubsystemBlueprintLibrary.cpp
@@ -45,7 +45,7 @@ TArray<FName> UEditorSubsystemBlueprintLibrary::GetPreviewPlatformOptions() cons
 
 void UEditorSubsystemBlueprintLibrary::TogglePreviewPlatform()
 {
-	GEditor->ToggleFeatureLevelPreview();
+	GEditor->TogglePreviewPlatform();
 }
 
 void UEditorSubsystemBlueprintLibrary::SetPreviewPlatform(FName PreviewShaderPlatformName)
diff --git a/Engine/Source/Runtime/Engine/Private/LocalPlayer.cpp b/Engine/Source/Runtime/Engine/Private/LocalPlayer.cpp
--- a/Engine/Source/Runtime/Engine/Private/LocalPlayer.cpp
+++ b/Engine/Source/Runtime/Engine/Private/LocalPlayer.cpp
@@ -1245,7 +1245,7 @@ bool ULocalPlayer::GetProjectionData(FViewport* Viewport, FSceneViewProjectionDa
 	ULevelEditorPlaySettings* PlaySettingsConfig = GetMutableDefault<ULevelEditorPlaySettings>();
 	if (PlaySettingsConfig != nullptr && PlaySettingsConfig->DeviceToEmulate.IsEmpty())
 	{
-		if (GEditor != nullptr && GEditor->IsFeatureLevelPreviewActive() && GEditor->GetAllowConstrainedAspectRatioInPreview())
+		if (GEditor != nullptr && GEditor->IsPreviewPlatformActive() && GEditor->GetAllowConstrainedAspectRatioInPreview())
 		{
 			float PlatformAspectRatio = GEditor->PreviewPlatform.GetConstrainedAspectRatio();
 			if (GEditor->GetFlipConstrainedAspectRatioInPreview() && PlatformAspectRatio != 0.f)
diff --git a/Engine/Source/Runtime/Engine/Private/RendererSettings.cpp b/Engine/Source/Runtime/Engine/Private/RendererSettings.cpp
--- a/Engine/Source/Runtime/Engine/Private/RendererSettings.cpp
+++ b/Engine/Source/Runtime/Engine/Private/RendererSettings.cpp
@@ -388,9 +388,9 @@ void URendererSettings::PostEditChangeProperty(FPropertyChangedEvent& PropertyCh
 					// This is to make sure that we have valid data for Mobile Preview.
 
 					// ES31->SM5 to be able to capture
-					GEditor->ToggleFeatureLevelPreview();
+					GEditor->TogglePreviewPlatform();
 					// SM5->ES31 BuildReflectionCaptures are triggered here on callback
-					GEditor->ToggleFeatureLevelPreview();
+					GEditor->TogglePreviewPlatform();
 				}
 				else
 				{
diff --git a/Engine/Source/Runtime/Engine/Private/Scalability.cpp b/Engine/Source/Runtime/Engine/Private/Scalability.cpp
--- a/Engine/Source/Runtime/Engine/Private/Scalability.cpp
+++ b/Engine/Source/Runtime/Engine/Private/Scalability.cpp
@@ -891,7 +891,7 @@ template<typename T> void SetQualityLevelCVar(TAutoConsoleVariable<T>& CVarToEdi
 		CVarToEdit.AsVariable()->SetWithCurrentPriority(DesiredValue);
 	}
 #if WITH_EDITOR
-	else if (GEditor && GEditor->IsFeatureLevelPreviewActive())
+	else if (GEditor && GEditor->IsPreviewPlatformActive())
 	{
 		EConsoleVariableFlags MaxPriority = ECVF_SetByPreview;
 		EConsoleVariableFlags MinPriority = ECVF_SetByScalability;
