Fix ASan Linux Stack-use-after-scope error in the TEDS Outliner Widget Constructor

#rnx
#Synced-CL 52098727

diff --git a/Engine/Source/Editor/LevelEditor/Private/TedsOutliner/SceneOutlinerWidgetFactory.cpp b/Engine/Source/Editor/LevelEditor/Private/TedsOutliner/SceneOutlinerWidgetFactory.cpp
--- a/Engine/Source/Editor/LevelEditor/Private/TedsOutliner/SceneOutlinerWidgetFactory.cpp
+++ b/Engine/Source/Editor/LevelEditor/Private/TedsOutliner/SceneOutlinerWidgetFactory.cpp
@@ -51,7 +51,8 @@ TSharedPtr<SWidget> FSceneOutlinerWidgetConstructor::CreateWidget(
 		// Fallback to the main SceneOutliner Tab Id
 		FName TabIdentifier = LevelEditorTabIds::LevelEditorSceneOutliner;
 
-		if(const FString* const* IdMetaData = Arguments.FindGeneric("OutlinerIdentifier").TryGetExact<const FString*>())
+		const FMetaDataEntryView OutlinerIdMeta = Arguments.FindGeneric("OutlinerIdentifier");
+		if(const FString* const* IdMetaData = OutlinerIdMeta.TryGetExact<const FString*>())
 		{
 			TabIdentifier = FName(**IdMetaData);
 		}
