Added dynamic invalidation to the level editor toolbars

#Synced-CL 52353515

diff --git a/Engine/Source/Editor/LevelEditor/Private/LevelEditorToolBar.cpp b/Engine/Source/Editor/LevelEditor/Private/LevelEditorToolBar.cpp
--- a/Engine/Source/Editor/LevelEditor/Private/LevelEditorToolBar.cpp
+++ b/Engine/Source/Editor/LevelEditor/Private/LevelEditorToolBar.cpp
@@ -6,6 +6,7 @@
 #include "Modules/ModuleManager.h"
 #include "Framework/Application/SlateApplication.h"
 #include "Widgets/Input/SComboButton.h"
+#include "Widgets/SInvalidationPanel.h"
 #include "Widgets/Layout/SBorder.h"
 #include "Widgets/Text/STextBlock.h"
 #include "Widgets/Layout/SBox.h"
@@ -1606,80 +1607,89 @@ TSharedRef< SWidget > FLevelEditorToolBar::MakeLevelEditorToolBar( const TShared
 
 	// Create the tool bar!
 	return
-		SNew(SOverlay)
-		+SOverlay::Slot()
+#if UE_SLATE_WITH_DYNAMIC_INVALIDATION
+		SNew(SInvalidationPanel)
+		.UseDynamicInvalidation(true)
 		[
-			SNew(SImage)
-			.Image(&FAppStyle::Get().GetWidgetStyle<FToolBarStyle>("SlimToolBar").BackgroundBrush)
-		]
-		+ SOverlay::Slot()
-		[
-			SNew(SHorizontalBox)
-			+ SHorizontalBox::Slot()
-			.AutoWidth()
+#endif
+			SNew(SOverlay)
+			+SOverlay::Slot()
 			[
-				SNew(SBorder)
-				.Padding(0.f)
-				.BorderImage(FAppStyle::Get().GetBrush("NoBorder"))
-				.IsEnabled(FSlateApplication::Get().GetNormalExecutionAttribute())
-				[
-					UToolMenus::Get()->GenerateWidget("LevelEditor.LevelEditorToolBar.ModesToolBar", MenuContext)
-				]
+				SNew(SImage)
+				.Image(&FAppStyle::Get().GetWidgetStyle<FToolBarStyle>("SlimToolBar").BackgroundBrush)
 			]
-			+ SHorizontalBox::Slot()
-			.AutoWidth()
+			+ SOverlay::Slot()
 			[
-				SNew(SBorder)
-				.Padding(0.f)
-				.BorderImage(FAppStyle::Get().GetBrush("NoBorder"))
-				.IsEnabled(FSlateApplication::Get().GetNormalExecutionAttribute())
+				SNew(SHorizontalBox)
+				+ SHorizontalBox::Slot()
+				.AutoWidth()
 				[
-					UToolMenus::Get()->GenerateWidget("LevelEditor.LevelEditorToolBar.AssetsToolBar", MenuContext)
+					SNew(SBorder)
+					.Padding(0.f)
+					.BorderImage(FAppStyle::Get().GetBrush("NoBorder"))
+					.IsEnabled(FSlateApplication::Get().GetNormalExecutionAttribute())
+					[
+						UToolMenus::Get()->GenerateWidget("LevelEditor.LevelEditorToolBar.ModesToolBar", MenuContext)
+					]
 				]
-				
-			]
-			+ SHorizontalBox::Slot()
-			.HAlign(HAlign_Left)
-			.AutoWidth()
-			[
-				SNew(SBorder)
-				.Padding(0.f)
-				.BorderImage(FAppStyle::Get().GetBrush("NoBorder"))
+				+ SHorizontalBox::Slot()
+				.AutoWidth()
 				[
-					UToolMenus::Get()->GenerateWidget("LevelEditor.LevelEditorToolBar.PlayToolBar", MenuContext) // Always enabled
+					SNew(SBorder)
+					.Padding(0.f)
+					.BorderImage(FAppStyle::Get().GetBrush("NoBorder"))
+					.IsEnabled(FSlateApplication::Get().GetNormalExecutionAttribute())
+					[
+						UToolMenus::Get()->GenerateWidget("LevelEditor.LevelEditorToolBar.AssetsToolBar", MenuContext)
+					]
+				
 				]
+				+ SHorizontalBox::Slot()
+				.HAlign(HAlign_Left)
+				.AutoWidth()
+				[
+					SNew(SBorder)
+					.Padding(0.f)
+					.BorderImage(FAppStyle::Get().GetBrush("NoBorder"))
+					[
+						UToolMenus::Get()->GenerateWidget("LevelEditor.LevelEditorToolBar.PlayToolBar", MenuContext) // Always enabled
+					]
 				
-			]
-			+ SHorizontalBox::Slot()
-			.HAlign(HAlign_Fill)
-			.FillWidth(1.0f)
-			[
-				SNew(SBorder)
-				.Padding(0.f)
-				.BorderImage(FAppStyle::Get().GetBrush("NoBorder"))
-				.IsEnabled(FSlateApplication::Get().GetNormalExecutionAttribute())
+				]
+				+ SHorizontalBox::Slot()
+				.HAlign(HAlign_Fill)
+				.FillWidth(1.0f)
 				[
-					UToolMenus::Get()->GenerateWidget("LevelEditor.LevelEditorToolBar.User", MenuContext)
+					SNew(SBorder)
+					.Padding(0.f)
+					.BorderImage(FAppStyle::Get().GetBrush("NoBorder"))
+					.IsEnabled(FSlateApplication::Get().GetNormalExecutionAttribute())
+					[
+						UToolMenus::Get()->GenerateWidget("LevelEditor.LevelEditorToolBar.User", MenuContext)
+					]
 				]
-			]
-			+ SHorizontalBox::Slot()
-			.HAlign(HAlign_Right)
-			.AutoWidth()
-			.Padding(0.0f, 0.0f, 7.0f, 0.0f)
-			[
-				SNew(SBorder)
-				.Visibility_Lambda([]
-				{
-					return UE::Editor::Private::CVarLevelEditorToolbarSettingsValue != 0 ? EVisibility::Visible : EVisibility::Collapsed;
-				})
-				.Padding(0.f)
-				.BorderImage(FAppStyle::Get().GetBrush("NoBorder"))
-				.IsEnabled(FSlateApplication::Get().GetNormalExecutionAttribute())
+				+ SHorizontalBox::Slot()
+				.HAlign(HAlign_Right)
+				.AutoWidth()
+				.Padding(0.0f, 0.0f, 7.0f, 0.0f)
 				[
-					UToolMenus::Get()->GenerateWidget("LevelEditor.LevelEditorToolBar.SettingsToolBar", MenuContext)
+					SNew(SBorder)
+					.Visibility_Lambda([]
+					{
+						return UE::Editor::Private::CVarLevelEditorToolbarSettingsValue != 0 ? EVisibility::Visible : EVisibility::Collapsed;
+					})
+					.Padding(0.f)
+					.BorderImage(FAppStyle::Get().GetBrush("NoBorder"))
+					.IsEnabled(FSlateApplication::Get().GetNormalExecutionAttribute())
+					[
+						UToolMenus::Get()->GenerateWidget("LevelEditor.LevelEditorToolBar.SettingsToolBar", MenuContext)
+					]
 				]
 			]
-		];
+#if UE_SLATE_WITH_DYNAMIC_INVALIDATION
+		]
+#endif
+		;
 }
 
 static void MakeScalabilityMenu( UToolMenu* InMenu )
diff --git a/Engine/Source/Editor/UnrealEd/Private/SEditorViewport.cpp b/Engine/Source/Editor/UnrealEd/Private/SEditorViewport.cpp
--- a/Engine/Source/Editor/UnrealEd/Private/SEditorViewport.cpp
+++ b/Engine/Source/Editor/UnrealEd/Private/SEditorViewport.cpp
@@ -39,6 +39,7 @@
 #include "SEditorViewportGridPanel.h"
 #include "ToolMenus.h"
 #include "ViewportInteractions/ViewportInteractionsBehaviorSource.h"
+#include "Widgets/SInvalidationPanel.h"
 #include "Widgets/Input/SButton.h"
 #include "RenderUtils.h"
 
@@ -225,24 +226,32 @@ void SEditorViewport::Construct( const FArguments& InArgs )
 			VerticalBox->AddSlot()
 			.AutoHeight()
 			[
-				SNew(SBox)
-				.Visibility_Lambda(
-					[bHasOldViewportToolbar, bHasLegacyUpgradedToolbar = bLegacyToolbarIsAutomaticallyUpgradable]() -> EVisibility
-					{
-						// Always show the new viewport toolbar if there is no old one.
-						if (!bHasOldViewportToolbar)
+#if UE_SLATE_WITH_DYNAMIC_INVALIDATION
+				SNew(SInvalidationPanel)
+				.UseDynamicInvalidation(true)
+				[
+#endif
+					SNew(SBox)
+					.Visibility_Lambda(
+						[bHasOldViewportToolbar, bHasLegacyUpgradedToolbar = bLegacyToolbarIsAutomaticallyUpgradable]() -> EVisibility
 						{
-							return EVisibility::Visible;
+							// Always show the new viewport toolbar if there is no old one.
+							if (!bHasOldViewportToolbar)
+							{
+								return EVisibility::Visible;
+							}
+							// In the case of the old toolbar being upgraded, don't show it twice
+							PRAGMA_DISABLE_DEPRECATION_WARNINGS
+							return (UE::UnrealEd::ShowNewViewportToolbars() && (!bHasLegacyUpgradedToolbar || !UE::UnrealEd::ShowOldViewportToolbars())) ? EVisibility::Visible: EVisibility::Collapsed;
+							PRAGMA_ENABLE_DEPRECATION_WARNINGS
 						}
-						// In the case of the old toolbar being upgraded, don't show it twice
-						PRAGMA_DISABLE_DEPRECATION_WARNINGS
-						return (UE::UnrealEd::ShowNewViewportToolbars() && (!bHasLegacyUpgradedToolbar || !UE::UnrealEd::ShowOldViewportToolbars())) ? EVisibility::Visible: EVisibility::Collapsed;
-						PRAGMA_ENABLE_DEPRECATION_WARNINGS
-					}
-				)
-				[
-					NewViewportToolbar.ToSharedRef()
+					)
+					[
+						NewViewportToolbar.ToSharedRef()
+					]
+#if UE_SLATE_WITH_DYNAMIC_INVALIDATION
 				]
+#endif
 			];
 			// clang-format on
 		}
diff --git a/Engine/Source/Editor/UnrealEd/Public/Subsystems/PanelExtensionSubsystem.h b/Engine/Source/Editor/UnrealEd/Public/Subsystems/PanelExtensionSubsystem.h
--- a/Engine/Source/Editor/UnrealEd/Public/Subsystems/PanelExtensionSubsystem.h
+++ b/Engine/Source/Editor/UnrealEd/Public/Subsystems/PanelExtensionSubsystem.h
@@ -43,6 +43,12 @@ public:
 
 PRAGMA_ENABLE_DEPRECATION_WARNINGS
 
+template <>
+struct TWidgetTypeTraits<class SExtensionPanel>
+{
+	static constexpr bool SupportsInvalidation() { return true; }
+};
+
 class SExtensionPanel : public SCompoundWidget
 {
 public:
diff --git a/Engine/Source/Runtime/Slate/Private/Framework/MultiBox/SToolBarSeparatorBlock.h b/Engine/Source/Runtime/Slate/Private/Framework/MultiBox/SToolBarSeparatorBlock.h
--- a/Engine/Source/Runtime/Slate/Private/Framework/MultiBox/SToolBarSeparatorBlock.h
+++ b/Engine/Source/Runtime/Slate/Private/Framework/MultiBox/SToolBarSeparatorBlock.h
@@ -39,6 +39,12 @@ private:
 
 
 
+template <>
+struct TWidgetTypeTraits<class SToolBarSeparatorBlock>
+{
+	static constexpr bool SupportsInvalidation() { return true; }
+};
+
 /**
  * Toolbar separator MultiBlock widget
  */
diff --git a/Engine/Source/Runtime/Slate/Private/Framework/MultiBox/SToolBarStackButtonBlock.cpp b/Engine/Source/Runtime/Slate/Private/Framework/MultiBox/SToolBarStackButtonBlock.cpp
--- a/Engine/Source/Runtime/Slate/Private/Framework/MultiBox/SToolBarStackButtonBlock.cpp
+++ b/Engine/Source/Runtime/Slate/Private/Framework/MultiBox/SToolBarStackButtonBlock.cpp
@@ -54,6 +54,12 @@ protected:
 	FTSTicker::FDelegateHandle HoldPressHandle;
 };
 
+template <>
+struct TWidgetTypeTraits<SCheckBoxStack>
+{
+	static constexpr bool SupportsInvalidation() { return true; }
+};
+
 FReply SCheckBoxStack::OnMouseButtonDown( const FGeometry& MyGeometry, const FPointerEvent& MouseEvent )
 {
 	if ( MouseEvent.GetEffectingButton() == EKeys::LeftMouseButton )
diff --git a/Engine/Source/Runtime/Slate/Private/Framework/MultiBox/SWidgetBlock.h b/Engine/Source/Runtime/Slate/Private/Framework/MultiBox/SWidgetBlock.h
--- a/Engine/Source/Runtime/Slate/Private/Framework/MultiBox/SWidgetBlock.h
+++ b/Engine/Source/Runtime/Slate/Private/Framework/MultiBox/SWidgetBlock.h
@@ -66,6 +66,12 @@ private:
 
 
 
+template <>
+struct TWidgetTypeTraits<class SWidgetBlock>
+{
+	static constexpr bool SupportsInvalidation() { return true; }
+};
+
 /**
  * Arbitrary Widget MultiBlock widget
  */
diff --git a/Engine/Source/Runtime/Slate/Public/Framework/MultiBox/SToolBarComboButtonBlock.h b/Engine/Source/Runtime/Slate/Public/Framework/MultiBox/SToolBarComboButtonBlock.h
--- a/Engine/Source/Runtime/Slate/Public/Framework/MultiBox/SToolBarComboButtonBlock.h
+++ b/Engine/Source/Runtime/Slate/Public/Framework/MultiBox/SToolBarComboButtonBlock.h
@@ -118,6 +118,12 @@ private:
 
 
 
+template <>
+struct TWidgetTypeTraits<class SToolBarComboButtonBlock>
+{
+	static constexpr bool SupportsInvalidation() { return true; }
+};
+
 /**
  * Tool bar button MultiBlock widget
  */
diff --git a/Engine/Source/Runtime/Slate/Public/Framework/MultiBox/SToolBarStackButtonBlock.h b/Engine/Source/Runtime/Slate/Public/Framework/MultiBox/SToolBarStackButtonBlock.h
--- a/Engine/Source/Runtime/Slate/Public/Framework/MultiBox/SToolBarStackButtonBlock.h
+++ b/Engine/Source/Runtime/Slate/Public/Framework/MultiBox/SToolBarStackButtonBlock.h
@@ -86,6 +86,12 @@ protected:
 
 
 
+template <>
+struct TWidgetTypeTraits<class SToolBarStackButtonBlock>
+{
+	static constexpr bool SupportsInvalidation() { return true; }
+};
+
 /**
  * Tool bar button MultiBlock widget
  */
