Remove dynamic invalidation from the viewport toolbar while we upragde SCheckBox to properly support it

#Synced-CL 52353515

diff --git a/Engine/Source/Editor/LevelEditor/Private/LevelEditorToolBar.cpp b/Engine/Source/Editor/LevelEditor/Private/LevelEditorToolBar.cpp
--- a/Engine/Source/Editor/LevelEditor/Private/LevelEditorToolBar.cpp
+++ b/Engine/Source/Editor/LevelEditor/Private/LevelEditorToolBar.cpp
@@ -6,7 +6,6 @@
 #include "Modules/ModuleManager.h"
 #include "Framework/Application/SlateApplication.h"
 #include "Widgets/Input/SComboButton.h"
-#include "Widgets/SInvalidationPanel.h"
 #include "Widgets/Layout/SBorder.h"
 #include "Widgets/Text/STextBlock.h"
 #include "Widgets/Layout/SBox.h"
@@ -1607,89 +1606,80 @@ TSharedRef< SWidget > FLevelEditorToolBar::MakeLevelEditorToolBar( const TShared
 
 	// Create the tool bar!
 	return
-#if UE_SLATE_WITH_DYNAMIC_INVALIDATION
-		SNew(SInvalidationPanel)
-		.UseDynamicInvalidation(true)
+		SNew(SOverlay)
+		+SOverlay::Slot()
 		[
-#endif
-			SNew(SOverlay)
-			+SOverlay::Slot()
+			SNew(SImage)
+			.Image(&FAppStyle::Get().GetWidgetStyle<FToolBarStyle>("SlimToolBar").BackgroundBrush)
+		]
+		+ SOverlay::Slot()
+		[
+			SNew(SHorizontalBox)
+			+ SHorizontalBox::Slot()
+			.AutoWidth()
 			[
-				SNew(SImage)
-				.Image(&FAppStyle::Get().GetWidgetStyle<FToolBarStyle>("SlimToolBar").BackgroundBrush)
+				SNew(SBorder)
+				.Padding(0.f)
+				.BorderImage(FAppStyle::Get().GetBrush("NoBorder"))
+				.IsEnabled(FSlateApplication::Get().GetNormalExecutionAttribute())
+				[
+					UToolMenus::Get()->GenerateWidget("LevelEditor.LevelEditorToolBar.ModesToolBar", MenuContext)
+				]
 			]
-			+ SOverlay::Slot()
+			+ SHorizontalBox::Slot()
+			.AutoWidth()
 			[
-				SNew(SHorizontalBox)
-				+ SHorizontalBox::Slot()
-				.AutoWidth()
+				SNew(SBorder)
+				.Padding(0.f)
+				.BorderImage(FAppStyle::Get().GetBrush("NoBorder"))
+				.IsEnabled(FSlateApplication::Get().GetNormalExecutionAttribute())
 				[
-					SNew(SBorder)
-					.Padding(0.f)
-					.BorderImage(FAppStyle::Get().GetBrush("NoBorder"))
-					.IsEnabled(FSlateApplication::Get().GetNormalExecutionAttribute())
-					[
-						UToolMenus::Get()->GenerateWidget("LevelEditor.LevelEditorToolBar.ModesToolBar", MenuContext)
-					]
+					UToolMenus::Get()->GenerateWidget("LevelEditor.LevelEditorToolBar.AssetsToolBar", MenuContext)
 				]
-				+ SHorizontalBox::Slot()
-				.AutoWidth()
-				[
-					SNew(SBorder)
-					.Padding(0.f)
-					.BorderImage(FAppStyle::Get().GetBrush("NoBorder"))
-					.IsEnabled(FSlateApplication::Get().GetNormalExecutionAttribute())
-					[
-						UToolMenus::Get()->GenerateWidget("LevelEditor.LevelEditorToolBar.AssetsToolBar", MenuContext)
-					]
 				
-				]
-				+ SHorizontalBox::Slot()
-				.HAlign(HAlign_Left)
-				.AutoWidth()
+			]
+			+ SHorizontalBox::Slot()
+			.HAlign(HAlign_Left)
+			.AutoWidth()
+			[
+				SNew(SBorder)
+				.Padding(0.f)
+				.BorderImage(FAppStyle::Get().GetBrush("NoBorder"))
 				[
-					SNew(SBorder)
-					.Padding(0.f)
-					.BorderImage(FAppStyle::Get().GetBrush("NoBorder"))
-					[
-						UToolMenus::Get()->GenerateWidget("LevelEditor.LevelEditorToolBar.PlayToolBar", MenuContext) // Always enabled
-					]
-				
+					UToolMenus::Get()->GenerateWidget("LevelEditor.LevelEditorToolBar.PlayToolBar", MenuContext) // Always enabled
 				]
-				+ SHorizontalBox::Slot()
-				.HAlign(HAlign_Fill)
-				.FillWidth(1.0f)
+				
+			]
+			+ SHorizontalBox::Slot()
+			.HAlign(HAlign_Fill)
+			.FillWidth(1.0f)
+			[
+				SNew(SBorder)
+				.Padding(0.f)
+				.BorderImage(FAppStyle::Get().GetBrush("NoBorder"))
+				.IsEnabled(FSlateApplication::Get().GetNormalExecutionAttribute())
 				[
-					SNew(SBorder)
-					.Padding(0.f)
-					.BorderImage(FAppStyle::Get().GetBrush("NoBorder"))
-					.IsEnabled(FSlateApplication::Get().GetNormalExecutionAttribute())
-					[
-						UToolMenus::Get()->GenerateWidget("LevelEditor.LevelEditorToolBar.User", MenuContext)
-					]
+					UToolMenus::Get()->GenerateWidget("LevelEditor.LevelEditorToolBar.User", MenuContext)
 				]
-				+ SHorizontalBox::Slot()
-				.HAlign(HAlign_Right)
-				.AutoWidth()
-				.Padding(0.0f, 0.0f, 7.0f, 0.0f)
+			]
+			+ SHorizontalBox::Slot()
+			.HAlign(HAlign_Right)
+			.AutoWidth()
+			.Padding(0.0f, 0.0f, 7.0f, 0.0f)
+			[
+				SNew(SBorder)
+				.Visibility_Lambda([]
+				{
+					return UE::Editor::Private::CVarLevelEditorToolbarSettingsValue != 0 ? EVisibility::Visible : EVisibility::Collapsed;
+				})
+				.Padding(0.f)
+				.BorderImage(FAppStyle::Get().GetBrush("NoBorder"))
+				.IsEnabled(FSlateApplication::Get().GetNormalExecutionAttribute())
 				[
-					SNew(SBorder)
-					.Visibility_Lambda([]
-					{
-						return UE::Editor::Private::CVarLevelEditorToolbarSettingsValue != 0 ? EVisibility::Visible : EVisibility::Collapsed;
-					})
-					.Padding(0.f)
-					.BorderImage(FAppStyle::Get().GetBrush("NoBorder"))
-					.IsEnabled(FSlateApplication::Get().GetNormalExecutionAttribute())
-					[
-						UToolMenus::Get()->GenerateWidget("LevelEditor.LevelEditorToolBar.SettingsToolBar", MenuContext)
-					]
+					UToolMenus::Get()->GenerateWidget("LevelEditor.LevelEditorToolBar.SettingsToolBar", MenuContext)
 				]
 			]
-#if UE_SLATE_WITH_DYNAMIC_INVALIDATION
-		]
-#endif
-		;
+		];
 }
 
 static void MakeScalabilityMenu( UToolMenu* InMenu )
diff --git a/Engine/Source/Editor/UnrealEd/Private/SEditorViewport.cpp b/Engine/Source/Editor/UnrealEd/Private/SEditorViewport.cpp
--- a/Engine/Source/Editor/UnrealEd/Private/SEditorViewport.cpp
+++ b/Engine/Source/Editor/UnrealEd/Private/SEditorViewport.cpp
@@ -39,7 +39,6 @@
 #include "SEditorViewportGridPanel.h"
 #include "ToolMenus.h"
 #include "ViewportInteractions/ViewportInteractionsBehaviorSource.h"
-#include "Widgets/SInvalidationPanel.h"
 #include "Widgets/Input/SButton.h"
 #include "RenderUtils.h"
 
@@ -226,32 +225,24 @@ void SEditorViewport::Construct( const FArguments& InArgs )
 			VerticalBox->AddSlot()
 			.AutoHeight()
 			[
-#if UE_SLATE_WITH_DYNAMIC_INVALIDATION
-				SNew(SInvalidationPanel)
-				.UseDynamicInvalidation(true)
-				[
-#endif
-					SNew(SBox)
-					.Visibility_Lambda(
-						[bHasOldViewportToolbar, bHasLegacyUpgradedToolbar = bLegacyToolbarIsAutomaticallyUpgradable]() -> EVisibility
+				SNew(SBox)
+				.Visibility_Lambda(
+					[bHasOldViewportToolbar, bHasLegacyUpgradedToolbar = bLegacyToolbarIsAutomaticallyUpgradable]() -> EVisibility
+					{
+						// Always show the new viewport toolbar if there is no old one.
+						if (!bHasOldViewportToolbar)
 						{
-							// Always show the new viewport toolbar if there is no old one.
-							if (!bHasOldViewportToolbar)
-							{
-								return EVisibility::Visible;
-							}
-							// In the case of the old toolbar being upgraded, don't show it twice
-							PRAGMA_DISABLE_DEPRECATION_WARNINGS
-							return (UE::UnrealEd::ShowNewViewportToolbars() && (!bHasLegacyUpgradedToolbar || !UE::UnrealEd::ShowOldViewportToolbars())) ? EVisibility::Visible: EVisibility::Collapsed;
-							PRAGMA_ENABLE_DEPRECATION_WARNINGS
+							return EVisibility::Visible;
 						}
-					)
-					[
-						NewViewportToolbar.ToSharedRef()
-					]
-#if UE_SLATE_WITH_DYNAMIC_INVALIDATION
+						// In the case of the old toolbar being upgraded, don't show it twice
+						PRAGMA_DISABLE_DEPRECATION_WARNINGS
+						return (UE::UnrealEd::ShowNewViewportToolbars() && (!bHasLegacyUpgradedToolbar || !UE::UnrealEd::ShowOldViewportToolbars())) ? EVisibility::Visible: EVisibility::Collapsed;
+						PRAGMA_ENABLE_DEPRECATION_WARNINGS
+					}
+				)
+				[
+					NewViewportToolbar.ToSharedRef()
 				]
-#endif
 			];
 			// clang-format on
 		}
