[Level Viewport ITF] Call "CheckHoveredHitProxy()" from the level editor hover behavior, restoring the preselection highlight feature.

#rnx
#Synced-CL 52324855

diff --git a/Engine/Source/Editor/LevelEditor/Private/ViewportInteractions/LevelViewportClickSelection.cpp b/Engine/Source/Editor/LevelEditor/Private/ViewportInteractions/LevelViewportClickSelection.cpp
--- a/Engine/Source/Editor/LevelEditor/Private/ViewportInteractions/LevelViewportClickSelection.cpp
+++ b/Engine/Source/Editor/LevelEditor/Private/ViewportInteractions/LevelViewportClickSelection.cpp
@@ -148,6 +148,10 @@ bool ULevelViewportClickSelection::OnUpdateHover(const FInputDeviceRay& DevicePo
 	if (HHitProxy* HitProxy = GetHitProxy(DevicePos))
 	{
 		SetMouseCursorOverride(HitProxy->GetMouseCursor());
+		if (FEditorViewportClient* Client = GetEditorViewportClient())
+		{
+			Client->CheckHoveredHitProxy(HitProxy);
+		}
 		return true;
 	}
 
@@ -157,6 +161,10 @@ bool ULevelViewportClickSelection::OnUpdateHover(const FInputDeviceRay& DevicePo
 void ULevelViewportClickSelection::OnEndHover()
 {
 	ClearMouseCursorOverride();
+	if (FEditorViewportClient* Client = GetEditorViewportClient())
+	{
+		Client->CheckHoveredHitProxy(nullptr);
+	}
 }
 
 void ULevelViewportClickSelection::ProcessClick_Internal(const FViewportClick& InViewportClick, HHitProxy* InHitProxy, FSceneView& View)
